- DMX Editor
- GDTF Importer
- Blueprints nodes for DMX
- RemoteOffset is now computed from the offset between UniverseRemoteStart and UniverseLocalStart. And UniverseRemoteStart is set directly by the user.
- Store DMX Input Monitor console Protocol and Universe ID on .ini files.
- Fix Get DMX Active Mode Function Values only returning byte sized values.
- Add whitelist platforms to DMX plugins' modules.
- Remove DMXEvent Handler and move event to subsystem
- Fix: Universe channels were all reset when selecting a different protocol from the Fader properties panel.
  It happened because a new universe would always replace existing ones when collecting universes from a UDMXEntityUniverseManaged.
- Rename SendDMX and SendDMXRaw map parameters to clarify their intent.
UI changes
- Universe "Channel" property is now labeled "Address";
- On the Input Monitor, "Final Universe ID" is now "Universe ID" (before, it was a result from other values);
- Moved the code that sets new Universes' directionality to Output into the PostEditChangeChainProperty function of UDMXEntityUniverseManaged;
- Clamp Universe ID value for faders in the Output Console;
- Don't display a slider for Fader's Universe ID value anymore. Depending on the protocol selected, the slider would only go halfway.

diff --git a/Engine/Plugins/Experimental/DigitalMultiplex/DMXEngine/DMXEngine.uplugin b/Engine/Plugins/Experimental/DigitalMultiplex/DMXEngine/DMXEngine.uplugin
new file mode 100644
--- /dev/null
+++ b/Engine/Plugins/Experimental/DigitalMultiplex/DMXEngine/DMXEngine.uplugin
@@ -0,0 +1,59 @@
+{
+  "FileVersion": 1,
+  "Version": 1,
+  "VersionName": "1.0",
+  "FriendlyName": "DMX Engine",
+  "Description": "Functionality and assets for communication with DigitalMultiplexer (DMX) enabled devices",
+  "Category": "Media",
+  "CreatedBy": "Epic Games, Inc.",
+  "CreatedByURL": "http://epicgames.com",
+  "DocsURL": "",
+  "MarketplaceURL": "",
+  "SupportURL": "",
+  "CanContainContent": false,
+  "IsBetaVersion": true,
+  "Modules": [
+    {
+      "Name": "DMXRuntime",
+      "Type": "Runtime",
+      "LoadingPhase": "Default",
+      "WhitelistPlatforms": 
+      [
+        "Win64",
+        "Win32",
+        "Mac",
+        "Linux"
+      ]
+    },
+    {
+      "Name": "DMXEditor",
+      "Type": "Editor",
+      "LoadingPhase": "Default",
+      "WhitelistPlatforms": 
+      [
+        "Win64",
+        "Win32",
+        "Mac",
+        "Linux"
+      ]
+    },
+	{
+      "Name": "DMXBlueprintGraph",
+      "Type": "UncookedOnly",
+      "LoadingPhase": "Default",
+      "WhitelistPlatforms": 
+      [
+        "Win64",
+        "Win32",
+        "Mac",
+        "Linux"
+      ]
+    }
+  ],
+  "Plugins": [
+    {
+      "Name": "DMXProtocol",
+      "Enabled": true
+    }
+  ]
+}
diff --git a/Engine/Plugins/Experimental/DigitalMultiplex/DMXEngine/Source/DMXBlueprintGraph/Private/Customizations/K2Node_GetDMXActiveModeFunctionValuesCustomization.cpp b/Engine/Plugins/Experimental/DigitalMultiplex/DMXEngine/Source/DMXBlueprintGraph/Private/Customizations/K2Node_GetDMXActiveModeFunctionValuesCustomization.cpp
new file mode 100644
--- /dev/null
+++ b/Engine/Plugins/Experimental/DigitalMultiplex/DMXEngine/Source/DMXBlueprintGraph/Private/Customizations/K2Node_GetDMXActiveModeFunctionValuesCustomization.cpp
@@ -0,0 +1,113 @@
+// Copyright Epic Games, Inc. All Rights Reserved.
+
+#include "Customizations/K2Node_GetDMXActiveModeFunctionValuesCustomization.h"
+#include "K2Node_GetDMXActiveModeFunctionValues.h"
+#include "K2Node_GetDMXActiveModeFunctionValues.h"
+
+#include "Widgets/DeclarativeSyntaxSupport.h"
+#include "Widgets/Text/STextBlock.h"
+#include "Widgets/Layout/SBox.h"
+#include "Widgets/Layout/SUniformGridPanel.h"
+#include "Widgets/Input/SButton.h"
+#include "DetailLayoutBuilder.h"
+#include "DetailWidgetRow.h"
+#include "DetailCategoryBuilder.h"
+#include "IDetailsView.h"
+
+#define LOCTEXT_NAMESPACE "K2Node_GetDMXActiveModeFunctionValuesCustomization"
+
+TSharedRef<IDetailCustomization> K2Node_GetDMXActiveModeFunctionValuesCustomization::MakeInstance()
+{
+	return MakeShared<K2Node_GetDMXActiveModeFunctionValuesCustomization>();
+}
+
+void K2Node_GetDMXActiveModeFunctionValuesCustomization::CustomizeDetails(IDetailLayoutBuilder& InDetailLayout)
+{
+	this->DetailLayout = &InDetailLayout;
+	UK2Node_GetDMXActiveModeFunctionValues* Node = GetK2Node_GetDMXActiveModeFunctionValues();
+
+	static const FName FixtureSettingsCategoryName = TEXT("Fixture Settings");
+	InDetailLayout.EditCategory(FixtureSettingsCategoryName, FText::GetEmpty(), ECategoryPriority::Important);
+	IDetailCategoryBuilder& FunctionActionsCategory = InDetailLayout.EditCategory("DMXFunctionActions", LOCTEXT("FunctionActionCategoryName", "Function Actions"), ECategoryPriority::Important);
+
+	FunctionActionsCategory.AddCustomRow(FText::GetEmpty())
+		.WholeRowContent()
+		.HAlign(HAlign_Left)
+		[
+			SNew( SBox )
+			.MaxDesiredWidth(300.f)
+			[
+				SNew(SUniformGridPanel)
+				.SlotPadding(2.0f)
+				+ SUniformGridPanel::Slot(0, 0)
+				[
+					SNew(SButton)
+					.OnClicked(this, &K2Node_GetDMXActiveModeFunctionValuesCustomization::ExposeFunctionsClicked)
+					.ToolTipText(LOCTEXT("ExposeFunctionsButtonTooltip", "Expose Functions to Node Pins"))
+					.IsEnabled_Lambda([Node]() -> bool { return Node && !Node->IsExposed(); })
+					.HAlign(HAlign_Center)
+					[
+						SNew(STextBlock)
+						.Text(LOCTEXT("ExposeFunctionsButton", "Expose Functions"))
+					]
+				]
+				+ SUniformGridPanel::Slot(1, 0)
+				[
+					SNew(SButton)
+					.OnClicked(this, &K2Node_GetDMXActiveModeFunctionValuesCustomization::ResetFunctionsClicked)
+					.ToolTipText(LOCTEXT("ResetFunctionsButtonTooltip", "Resets Functions from Node Pins."))
+					.IsEnabled_Lambda([Node]() -> bool { return Node && Node->IsExposed(); })
+					.HAlign(HAlign_Center)
+					[
+						SNew(STextBlock)
+						.Text(LOCTEXT("ResetEmitterButton", "Reset Functions"))
+					]
+				]
+			]
+		];
+}
+
+FReply K2Node_GetDMXActiveModeFunctionValuesCustomization::ExposeFunctionsClicked()
+{
+	const TArray< TWeakObjectPtr<UObject> >& SelectedObjects = DetailLayout->GetSelectedObjects();
+
+	for (int32 Idx = 0; Idx < SelectedObjects.Num(); ++Idx)
+	{
+		if (SelectedObjects[Idx].IsValid())
+		{
+			if (UK2Node_GetDMXActiveModeFunctionValues* K2Node_GetDMXActiveModeFunctionValues = Cast<UK2Node_GetDMXActiveModeFunctionValues>(SelectedObjects[Idx].Get()))
+			{
+				K2Node_GetDMXActiveModeFunctionValues->ExposeFunctions();
+			}
+		}
+	}
+
+	return FReply::Handled();
+}
+
+FReply K2Node_GetDMXActiveModeFunctionValuesCustomization::ResetFunctionsClicked()
+{
+	if (UK2Node_GetDMXActiveModeFunctionValues* K2Node_GetDMXActiveModeFunctionValues = GetK2Node_GetDMXActiveModeFunctionValues())
+	{
+		K2Node_GetDMXActiveModeFunctionValues->ResetFunctions();
+	}
+
+	return FReply::Handled();
+}
+
+UK2Node_GetDMXActiveModeFunctionValues* K2Node_GetDMXActiveModeFunctionValuesCustomization::GetK2Node_GetDMXActiveModeFunctionValues() const
+{
+	const TArray< TWeakObjectPtr<UObject> >& SelectedObjects = DetailLayout->GetSelectedObjects();
+
+	for (int32 Idx = 0; Idx < SelectedObjects.Num(); ++Idx)
+	{
+		if (SelectedObjects[Idx].IsValid())
+		{
+			return Cast<UK2Node_GetDMXActiveModeFunctionValues>(SelectedObjects[Idx].Get());
+		}
+	}
+
+	return nullptr;
+}
+
+#undef LOCTEXT_NAMESPACE
diff --git a/Engine/Plugins/Experimental/DigitalMultiplex/DMXEngine/Source/DMXBlueprintGraph/Private/Customizations/K2Node_GetDMXActiveModeFunctionValuesCustomization.h b/Engine/Plugins/Experimental/DigitalMultiplex/DMXEngine/Source/DMXBlueprintGraph/Private/Customizations/K2Node_GetDMXActiveModeFunctionValuesCustomization.h
new file mode 100644
--- /dev/null
+++ b/Engine/Plugins/Experimental/DigitalMultiplex/DMXEngine/Source/DMXBlueprintGraph/Private/Customizations/K2Node_GetDMXActiveModeFunctionValuesCustomization.h
@@ -0,0 +1,36 @@
+// Copyright Epic Games, Inc. All Rights Reserved.
+
+#pragma once
+
+#include "CoreMinimal.h"
+#include "Input/Reply.h"
+#include "IDetailCustomization.h"
+
+class IDetailLayoutBuilder;
+class UK2Node_GetDMXActiveModeFunctionValues;
+
+class K2Node_GetDMXActiveModeFunctionValuesCustomization 
+	: public IDetailCustomization
+{
+public:
+	/** Makes a new instance of this detail layout class for a specific detail view requesting it */
+	static TSharedRef<IDetailCustomization> MakeInstance();
+private:
+	//~ Begin IDetailCustomization Interface
+	virtual void CustomizeDetails(IDetailLayoutBuilder& InDetailLayout) override;
+	//~ Begin IDetailCustomization Interface
+
+	/** Handles the Expose button's on click event */
+	FReply ExposeFunctionsClicked();
+
+	/** Handles the Reset button's on click event */
+	FReply ResetFunctionsClicked();
+
+	/** Get the pointer to GetDMXActiveModeFunctionValues node */
+	UK2Node_GetDMXActiveModeFunctionValues* GetK2Node_GetDMXActiveModeFunctionValues() const;
+
+private:
+
+	/** Cached off reference to the layout builder */
+	IDetailLayoutBuilder* DetailLayout;
+};
diff --git a/Engine/Plugins/Experimental/DigitalMultiplex/DMXEngine/Source/DMXBlueprintGraph/Private/DMXGraphPanelPinFactory.cpp b/Engine/Plugins/Experimental/DigitalMultiplex/DMXEngine/Source/DMXBlueprintGraph/Private/DMXGraphPanelPinFactory.cpp
new file mode 100644
--- /dev/null
+++ b/Engine/Plugins/Experimental/DigitalMultiplex/DMXEngine/Source/DMXBlueprintGraph/Private/DMXGraphPanelPinFactory.cpp
@@ -0,0 +1,83 @@
+// Copyright Epic Games, Inc. All Rights Reserved.
+
+#include "DMXGraphPanelPinFactory.h"
+#include "K2Node.h"
+#include "K2Node_GetAllFixturesOfType.h"
+#include "EdGraphSchema_K2.h"
+#include "Widgets/DeclarativeSyntaxSupport.h"
+#include "SGraphPin.h"
+#include "SGraphPinNameList.h"
+#include "Library/DMXLibrary.h"
+#include "Library/DMXEntityReference.h"
+#include "Library/DMXEntity.h"
+#include "Library/DMXEntityFixtureType.h"
+#include "DMXProtocolTypes.h"
+#include "Widgets/SDynamicNameListGraphPin.h"
+#include "Widgets/SDMXEntityReferenceGraphPin.h"
+#include "Widgets/SNullGraphPin.h"
+
+TSharedPtr<class SGraphPin> FDMXGraphPanelPinFactory::CreatePin(class UEdGraphPin* InPin) const
+{
+	if (InPin->PinType.PinCategory == UEdGraphSchema_K2::PC_Name)
+	{
+		UObject* Outer = InPin->GetOuter();
+
+		if (UK2Node_DMXBase* K2Node_DMXBase = Cast<UK2Node_DMXBase>(Outer))
+		{
+			const UEdGraphPin* DMXLibraryPin = K2Node_DMXBase->GetClassPin();
+			if (DMXLibraryPin != nullptr)
+			{
+				if (DMXLibraryPin->LinkedTo.Num() == 0)
+				{
+					if (UDMXLibrary* DMXLibrary = Cast<UDMXLibrary>(DMXLibraryPin->DefaultObject))
+					{
+						if (UK2Node_GetAllFixturesOfType * K2Node_GetAllFixturesOfType = Cast<UK2Node_GetAllFixturesOfType>(Outer))
+						{
+							TArray<TSharedPtr<FName>> FixtureTypeList;
+							DMXLibrary->ForEachEntityOfType<UDMXEntityFixtureType>([&FixtureTypeList](UDMXEntityFixtureType* Fixture)
+								{
+									TSharedPtr<FName> NewFixtureType = MakeShared<FName>(FName(*Fixture->GetDisplayName()));
+									FixtureTypeList.Add(NewFixtureType);
+								});
+
+							return SNew(SGraphPinNameList, InPin, FixtureTypeList);
+						}
+					}
+				}
+			}
+		}
+	}
+	else if (InPin->PinType.PinCategory == UEdGraphSchema_K2::PC_Struct)
+	{
+		if (UScriptStruct* PinStructType = Cast<UScriptStruct>(InPin->PinType.PinSubCategoryObject.Get()))
+		{
+			if (PinStructType->IsChildOf(FDMXProtocolName::StaticStruct()))
+			{
+				return SNew(SDynamicNameListGraphPin<FDMXProtocolName>, InPin)
+					.OptionsSource(MakeAttributeLambda(&FDMXProtocolName::GetPossibleValues));
+			}
+			else if (PinStructType->IsChildOf(FDMXFixtureCategory::StaticStruct()))
+			{
+				return SNew(SDynamicNameListGraphPin<FDMXFixtureCategory>, InPin)
+					.OptionsSource(MakeAttributeLambda(&FDMXFixtureCategory::GetPossibleValues))
+					.UpdateOptionsDelegate(&FDMXFixtureCategory::OnPossibleValuesUpdated);
+			}
+			else if (PinStructType->IsChildOf(FDMXEntityReference::StaticStruct()))
+			{
+				return SNew(SDMXEntityReferenceGraphPin, InPin);
+			}
+		}
+	}
+	else if (InPin->PinType.PinCategory == UEdGraphSchema_K2::PC_Object)
+	{
+		if (UClass* PinObjectType = Cast<UClass>(InPin->PinType.PinSubCategoryObject.Get()))
+		{
+			if (PinObjectType->IsChildOf(UDMXEntity::StaticClass()))
+			{
+				return SNew(SNullGraphPin, InPin);
+			}
+		}
+	}
+
+	return FGraphPanelPinFactory::CreatePin(InPin);
+}
\ No newline at end of file
diff --git a/Engine/Plugins/Experimental/DigitalMultiplex/DMXEngine/Source/DMXBlueprintGraph/Private/K2Node_DMXBase.cpp b/Engine/Plugins/Experimental/DigitalMultiplex/DMXEngine/Source/DMXBlueprintGraph/Private/K2Node_DMXBase.cpp
new file mode 100644
--- /dev/null
+++ b/Engine/Plugins/Experimental/DigitalMultiplex/DMXEngine/Source/DMXBlueprintGraph/Private/K2Node_DMXBase.cpp
@@ -0,0 +1,140 @@
+// Copyright Epic Games, Inc. All Rights Reserved.
+
+#include "K2Node_DMXBase.h"
+#include "EdGraphSchema_K2.h"
+#include "Library/DMXLibrary.h"
+#include "DMXProtocolConstants.h"
+#include "BlueprintActionDatabaseRegistrar.h"
+#include "FindInBlueprintManager.h"
+#include "BlueprintNodeSpawner.h"
+#include "Kismet2/BlueprintEditorUtils.h"
+
+FName FK2Node_DMXBaseHelper::ClassPinName("DMXLibrary");
+
+UK2Node_DMXBase::~UK2Node_DMXBase()
+{
+	if (CachedLibrary != nullptr)
+	{
+		CachedLibrary->GetOnEntitiesUpdated().Remove(LibraryEntitiesUpdatedHandle);
+		CachedLibrary = nullptr;
+	}
+}
+
+void UK2Node_DMXBase::AllocateDefaultPins()
+{
+	UEdGraphPin* ClassPin = CreatePin(EGPD_Input, UEdGraphSchema_K2::PC_Object, GetClassPinBaseClass(), FK2Node_DMXBaseHelper::ClassPinName);
+
+	Super::AllocateDefaultPins();
+}
+
+bool UK2Node_DMXBase::IsCompatibleWithGraph(const UEdGraph* TargetGraph) const
+{
+	UBlueprint* Blueprint = FBlueprintEditorUtils::FindBlueprintForGraph(TargetGraph);
+	return Super::IsCompatibleWithGraph(TargetGraph) && (!Blueprint || (FBlueprintEditorUtils::FindUserConstructionScript(Blueprint) != TargetGraph && Blueprint->GeneratedClass->GetDefaultObject()->ImplementsGetWorld()));
+}
+
+void UK2Node_DMXBase::AddSearchMetaDataInfo(TArray<struct FSearchTagDataPair>& OutTaggedMetaData) const
+{
+	Super::AddSearchMetaDataInfo(OutTaggedMetaData);
+	OutTaggedMetaData.Add(FSearchTagDataPair(FFindInBlueprintSearchTags::FiB_NativeName, CachedNodeTitle.GetCachedText()));
+}
+
+void UK2Node_DMXBase::PinConnectionListChanged(UEdGraphPin* Pin)
+{
+	Super::PinConnectionListChanged(Pin);
+
+	if (Pin != nullptr)
+	{
+		if (Pin->PinName == FK2Node_DMXBaseHelper::ClassPinName)
+		{
+			if (Pin->DefaultObject != nullptr)
+			{
+				OnLibraryAssetChanged(Cast<UDMXLibrary>(Pin->DefaultObject));
+			}
+		}
+	}
+}
+
+void UK2Node_DMXBase::PinDefaultValueChanged(UEdGraphPin* Pin)
+{
+	if (Pin != nullptr)
+	{
+		if (Pin->PinName == FK2Node_DMXBaseHelper::ClassPinName)
+		{
+			if (Pin->DefaultObject != nullptr)
+			{
+				OnLibraryAssetChanged(Cast<UDMXLibrary>(Pin->DefaultObject));
+			}
+		}
+	}
+}
+
+void UK2Node_DMXBase::GetMenuActions(FBlueprintActionDatabaseRegistrar& ActionRegistrar) const
+{
+	UClass* ActionKey = GetClass();
+
+	if (ActionRegistrar.IsOpenForRegistration(ActionKey))
+	{
+		UBlueprintNodeSpawner* NodeSpawner = UBlueprintNodeSpawner::Create(GetClass());
+		check(NodeSpawner != nullptr);
+
+		ActionRegistrar.AddBlueprintAction(ActionKey, NodeSpawner);
+	}
+}
+
+FText UK2Node_DMXBase::GetMenuCategory() const
+{
+	return FText::FromString(DMX_K2_CATEGORY_NAME);
+}
+
+UEdGraphPin* UK2Node_DMXBase::GetClassPin(const TArray<UEdGraphPin*>* InPinsToSearch /*= NULL*/) const
+{
+	const TArray<UEdGraphPin*>* PinsToSearch = InPinsToSearch ? InPinsToSearch : &Pins;
+
+	UEdGraphPin* Pin = nullptr;
+	for (UEdGraphPin* TestPin : *PinsToSearch)
+	{
+		if (TestPin && TestPin->PinName == FK2Node_DMXBaseHelper::ClassPinName)
+		{
+			Pin = TestPin;
+			break;
+		}
+	}
+	check(Pin == nullptr || Pin->Direction == EGPD_Input);
+	return Pin;
+}
+
+UClass* UK2Node_DMXBase::GetClassPinBaseClass() const
+{
+	return UDMXLibrary::StaticClass();
+}
+
+void UK2Node_DMXBase::OnLibraryAssetChanged(UDMXLibrary* Library)
+{
+	// Unbind from previous library
+	if (CachedLibrary != nullptr)
+	{
+		CachedLibrary->GetOnEntitiesUpdated().Remove(LibraryEntitiesUpdatedHandle);
+		CachedLibrary = nullptr;
+	}
+
+	if (Library != nullptr)
+	{
+		LibraryEntitiesUpdatedHandle = Library->GetOnEntitiesUpdated().AddUObject(this, &UK2Node_DMXBase::OnLibraryEntitiesUpdated);
+		CachedLibrary = Library;
+	}
+
+	OnLibraryEntitiesUpdated(Library);
+}
+
+void UK2Node_DMXBase::OnLibraryEntitiesUpdated(UDMXLibrary* Library)
+{
+	UEdGraphPin* ClassPin = FindPin(FK2Node_DMXBaseHelper::ClassPinName);
+	if (ClassPin != nullptr && Library != nullptr && ClassPin->DefaultObject == Library)
+	{
+		if (UBlueprint* BP = GetBlueprint())
+		{
+			FBlueprintEditorUtils::MarkBlueprintAsModified(BP);
+		}
+	}
+}
diff --git a/Engine/Plugins/Experimental/DigitalMultiplex/DMXEngine/Source/DMXBlueprintGraph/Private/K2Node_GetAllFixturesOfType.cpp b/Engine/Plugins/Experimental/DigitalMultiplex/DMXEngine/Source/DMXBlueprintGraph/Private/K2Node_GetAllFixturesOfType.cpp
new file mode 100644
--- /dev/null
+++ b/Engine/Plugins/Experimental/DigitalMultiplex/DMXEngine/Source/DMXBlueprintGraph/Private/K2Node_GetAllFixturesOfType.cpp
@@ -0,0 +1,86 @@
+// Copyright Epic Games, Inc. All Rights Reserved.
+
+#include "K2Node_GetAllFixturesOfType.h"
+#include "DMXSubsystem.h"
+#include "Library/DMXEntityFixturePatch.h"
+#include "Library/DMXLibrary.h"
+#include "DMXSubsystem.h"
+#include "EdGraphSchema_K2.h"
+#include "K2Node_CallFunction.h"
+#include "KismetCompiler.h"
+#include "EditorCategoryUtils.h"
+#include "Engine/Engine.h"
+
+#define LOCTEXT_NAMESPACE "K2Node_GetAllFixturesOfType"
+
+FName FK2Node_GetAllFixturesOfType::FixtureTypePinName(TEXT("FixtureType"));
+FName FK2Node_GetAllFixturesOfType::OutResultPinName(TEXT("Fixtures"));
+
+UK2Node_GetAllFixturesOfType::UK2Node_GetAllFixturesOfType()
+{
+	NodeTooltip = LOCTEXT("NodeTooltip", "Returns all fixtures of selected type");
+}
+
+void UK2Node_GetAllFixturesOfType::AllocateDefaultPins()
+{
+	CreatePin(EGPD_Input, UEdGraphSchema_K2::PC_Exec, UEdGraphSchema_K2::PN_Execute);
+	UEdGraphPin* ThenPin = CreatePin(EGPD_Output, UEdGraphSchema_K2::PC_Exec, UEdGraphSchema_K2::PN_Then);
+
+	Super::AllocateDefaultPins();
+
+	UEdGraphPin* FixtureTypePin = CreatePin(EGPD_Input, UEdGraphSchema_K2::PC_Name, FK2Node_GetAllFixturesOfType::FixtureTypePinName);
+
+	FCreatePinParams PinParams = FCreatePinParams();
+	PinParams.ContainerType = EPinContainerType::Array;
+	UEdGraphPin* OutResultPin = CreatePin(EGPD_Output, UEdGraphSchema_K2::PC_Object, UDMXEntityFixturePatch::StaticClass() ,FK2Node_GetAllFixturesOfType::OutResultPinName, PinParams);
+}
+
+FText UK2Node_GetAllFixturesOfType::GetNodeTitle(ENodeTitleType::Type TitleType) const
+{
+	return LOCTEXT("BaseTitle", "Get All Fixtures Of Type");
+}
+
+void UK2Node_GetAllFixturesOfType::ExpandNode(class FKismetCompilerContext& CompilerContext, UEdGraph* SourceGraph)
+{
+	UDMXSubsystem* Subsystem = GEngine->GetEngineSubsystem<UDMXSubsystem>();
+
+	const FName FunctionName = GET_FUNCTION_NAME_CHECKED(UDMXSubsystem, GetAllFixturesOfType);
+	static const FName ClassName(TEXT("DMXLibrary"));
+	static const FName FixtureTypeName(TEXT("FixtureType"));
+	static const FName OutResultName(TEXT("OutResult"));
+	
+	UEdGraphPin* SelfPin = CreatePin(EGPD_Input, UEdGraphSchema_K2::PC_Object, UDMXSubsystem::StaticClass(), UEdGraphSchema_K2::PN_Self);
+	SelfPin->DefaultObject = Subsystem;
+
+	UK2Node_GetAllFixturesOfType* GetNode = this;
+	UEdGraphPin* ClassPin = GetNode->GetClassPin();
+	UEdGraphPin* FixtureTypePin = GetNode->FindPin(FK2Node_GetAllFixturesOfType::FixtureTypePinName);
+	UEdGraphPin* SendOutResult = GetNode->FindPin(FK2Node_GetAllFixturesOfType::OutResultPinName);
+	UEdGraphPin* ThenPin = GetNode->FindPin(UEdGraphSchema_K2::PN_Then);
+
+	UK2Node_CallFunction* GetAllFixturesNode = CompilerContext.SpawnIntermediateNode<UK2Node_CallFunction>(GetNode, SourceGraph);
+	GetAllFixturesNode->FunctionReference.SetExternalMember(FunctionName, UDMXSubsystem::StaticClass());
+	GetAllFixturesNode->AllocateDefaultPins();
+
+	UEdGraphPin* CallClass = GetAllFixturesNode->FindPinChecked(ClassName);
+	UEdGraphPin* CallFixtureType = GetAllFixturesNode->FindPinChecked(FixtureTypeName);
+	UEdGraphPin* CallOutResult = GetAllFixturesNode->FindPinChecked(OutResultName);
+	UEdGraphPin* CallThenPin = GetAllFixturesNode->GetThenPin();
+
+	CompilerContext.MovePinLinksToIntermediate(*SelfPin, *GetAllFixturesNode->FindPin(UEdGraphSchema_K2::PN_Self));
+	CompilerContext.MovePinLinksToIntermediate(*GetExecPin(), *GetAllFixturesNode->GetExecPin());
+	CompilerContext.MovePinLinksToIntermediate(*ClassPin, *CallClass);
+	CompilerContext.MovePinLinksToIntermediate(*FixtureTypePin, *CallFixtureType);
+	CallOutResult->PinType = SendOutResult->PinType;
+	CompilerContext.MovePinLinksToIntermediate(*SendOutResult, *CallOutResult);
+	CompilerContext.MovePinLinksToIntermediate(*ThenPin, *CallThenPin);
+}
+
+void UK2Node_GetAllFixturesOfType::ReallocatePinsDuringReconstruction(TArray<UEdGraphPin*>& OldPins)
+{
+	AllocateDefaultPins();
+
+	RestoreSplitPins(OldPins);
+}
+
+#undef LOCTEXT_NAMESPACE
diff --git a/Engine/Plugins/Experimental/DigitalMultiplex/DMXEngine/Source/DMXBlueprintGraph/Private/K2Node_GetDMXActiveModeFunctionValues.cpp b/Engine/Plugins/Experimental/DigitalMultiplex/DMXEngine/Source/DMXBlueprintGraph/Private/K2Node_GetDMXActiveModeFunctionValues.cpp
new file mode 100644
--- /dev/null
+++ b/Engine/Plugins/Experimental/DigitalMultiplex/DMXEngine/Source/DMXBlueprintGraph/Private/K2Node_GetDMXActiveModeFunctionValues.cpp
@@ -0,0 +1,445 @@
+// Copyright Epic Games, Inc. All Rights Reserved.
+
+#include "K2Node_GetDMXActiveModeFunctionValues.h"
+
+#include "Library/DMXEntityFixturePatch.h"
+#include "DMXSubsystem.h"
+#include "DMXProtocolConstants.h"
+#include "DMXBlueprintGraphLog.h"
+#include "K2Node_GetDMXFixturePatch.h"
+
+#include "KismetCompiler.h"
+#include "BlueprintActionDatabaseRegistrar.h"
+#include "BlueprintNodeSpawner.h"
+#include "Kismet2/BlueprintEditorUtils.h"
+#include "ScopedTransaction.h"
+#include "EdGraphUtilities.h"
+#include "K2Node_CallFunction.h"
+#include "Engine/Engine.h"
+
+#define LOCTEXT_NAMESPACE "UK2Node_GetDMXActiveModeFunctionValues"
+
+const FName UK2Node_GetDMXActiveModeFunctionValues::InputDMXFixturePatchPinName(TEXT("InFixturePatch"));
+const FName UK2Node_GetDMXActiveModeFunctionValues::InputDMXProtocolPinName(TEXT("InProtocol"));
+
+const FName UK2Node_GetDMXActiveModeFunctionValues::OutputFunctionsMapPinName(TEXT("OutFunctionsMap"));
+const FName UK2Node_GetDMXActiveModeFunctionValues::OutputIsSuccessPinName(TEXT("OutIsSuccessPinName"));
+
+
+UK2Node_GetDMXActiveModeFunctionValues::UK2Node_GetDMXActiveModeFunctionValues()
+{
+	bIsEditable = true;
+	bIsExposed = false;
+}
+
+void UK2Node_GetDMXActiveModeFunctionValues::OnFixturePatchChanged()
+{
+	// Reset fixture path nodes if we receive a notification
+	if (Pins.Num() > 0 && bIsExposed)
+	{
+		ResetFunctions();
+	}
+}
+
+void UK2Node_GetDMXActiveModeFunctionValues::RemovePinsRecursive(UEdGraphPin* PinToRemove)
+{
+	for (int32 SubPinIndex = PinToRemove->SubPins.Num() - 1; SubPinIndex >= 0; --SubPinIndex)
+	{
+		RemovePinsRecursive(PinToRemove->SubPins[SubPinIndex]);
+	}
+
+	int32 PinRemovalIndex = INDEX_NONE;
+	if (Pins.Find(PinToRemove, PinRemovalIndex))
+	{
+		Pins.RemoveAt(PinRemovalIndex);
+		PinToRemove->MarkPendingKill();
+	}
+}
+
+void UK2Node_GetDMXActiveModeFunctionValues::RemoveOutputPin(UEdGraphPin* Pin)
+{
+	checkSlow(Pins.Contains(Pin));
+
+	FScopedTransaction Transaction(LOCTEXT("RemovePinTx", "RemovePin"));
+	Modify();
+
+	RemovePinsRecursive(Pin);
+	PinConnectionListChanged(Pin);
+
+	FBlueprintEditorUtils::MarkBlueprintAsStructurallyModified(GetBlueprint());
+}
+
+FText UK2Node_GetDMXActiveModeFunctionValues::GetNodeTitle(ENodeTitleType::Type TitleType) const
+{
+	return LOCTEXT("NodeTitle", "Get DMX Function Values");
+}
+
+void UK2Node_GetDMXActiveModeFunctionValues::AllocateDefaultPins()
+{
+	const UEdGraphSchema_K2* K2Schema = GetDefault<UEdGraphSchema_K2>();
+
+	// Add execution pins
+	CreatePin(EGPD_Input, UEdGraphSchema_K2::PC_Exec, UEdGraphSchema_K2::PN_Execute);
+	CreatePin(EGPD_Output, UEdGraphSchema_K2::PC_Exec, UEdGraphSchema_K2::PN_Then);
+
+	// Add input pin
+	UEdGraphPin* InputDMXFixturePatchPin = CreatePin(EGPD_Input, UEdGraphSchema_K2::PC_Object, UDMXEntityFixturePatch::StaticClass(), InputDMXFixturePatchPinName);
+	K2Schema->ConstructBasicPinTooltip(*InputDMXFixturePatchPin, LOCTEXT("InputDMXFixturePatch", "Input DMX Fixture Patch"), InputDMXFixturePatchPin->PinToolTip);
+	
+	// Protocol Name pin
+	UEdGraphPin* InputDMXProtocolPin = CreatePin(EGPD_Input, UEdGraphSchema_K2::PC_Struct, FDMXProtocolName::StaticStruct(), InputDMXProtocolPinName);
+	K2Schema->ConstructBasicPinTooltip(*InputDMXProtocolPin, LOCTEXT("InputDMXProtocolPin", "The DMX protocol name"), InputDMXProtocolPin->PinToolTip);
+
+	// Add output pin
+	FCreatePinParams OutputFunctionsMapPinParams;
+	OutputFunctionsMapPinParams.ContainerType = EPinContainerType::Map;
+	OutputFunctionsMapPinParams.ValueTerminalType.TerminalCategory = UEdGraphSchema_K2::PC_Int;
+	UEdGraphPin* OutputFunctionsMapPin = CreatePin(EGPD_Output, UEdGraphSchema_K2::PC_Name, TEXT(""), OutputFunctionsMapPinName, OutputFunctionsMapPinParams);
+	K2Schema->ConstructBasicPinTooltip(*OutputFunctionsMapPin, LOCTEXT("OutputFunctionsMap", "Output Functions Map."), OutputFunctionsMapPin->PinToolTip);
+
+	UEdGraphPin* OutputIsSuccessPin = CreatePin(EGPD_Output, UEdGraphSchema_K2::PC_Boolean, TEXT(""), OutputIsSuccessPinName);
+	K2Schema->ConstructBasicPinTooltip(*OutputIsSuccessPin, LOCTEXT("OutputIsSuccessPin", "Is SuccessP"), OutputIsSuccessPin->PinToolTip);
+
+	Super::AllocateDefaultPins();
+}
+
+void UK2Node_GetDMXActiveModeFunctionValues::ExpandNode(FKismetCompilerContext& CompilerContext, UEdGraph* SourceGraph)
+{
+	Super::ExpandNode(CompilerContext, SourceGraph);
+	const UEdGraphSchema_K2* Schema = CompilerContext.GetSchema();
+
+	UDMXSubsystem* Subsystem = GEngine->GetEngineSubsystem<UDMXSubsystem>();
+
+	UEdGraphPin* SelfPin = CreatePin(EGPD_Input, UEdGraphSchema_K2::PC_Object, UDMXSubsystem::StaticClass(), UEdGraphSchema_K2::PN_Self);
+	SelfPin->DefaultObject = Subsystem;
+
+	// Call DMX function map. Always first and always call it
+	// Static C++ function
+	static const FName FunctionName = GET_FUNCTION_NAME_CHECKED(UDMXSubsystem, GetFunctionsMap);
+	UFunction* GetFunctionsMapPointer = FindField<UFunction>(UDMXSubsystem::StaticClass(), FunctionName);
+	check(GetFunctionsMapPointer);
+
+	UK2Node_CallFunction* GetFunctionsMapNode = CompilerContext.SpawnIntermediateNode<UK2Node_CallFunction>(this, SourceGraph);
+	GetFunctionsMapNode->FunctionReference.SetExternalMember(FunctionName, UDMXSubsystem::StaticClass());
+	GetFunctionsMapNode->AllocateDefaultPins();
+
+	// Function pins
+	UEdGraphPin* GetFunctionsMapNodeSelfPin = GetFunctionsMapNode->FindPinChecked(UEdGraphSchema_K2::PN_Self);
+	UEdGraphPin* GetFunctionsMapNodeExecPin = GetFunctionsMapNode->GetExecPin();
+	UEdGraphPin* GetFunctionsMapNodeInFixturePatchPin = GetFunctionsMapNode->FindPinChecked(TEXT("InFixturePatch"));
+	UEdGraphPin* GetFunctionsMapNodeSelectedProtocolPin = GetFunctionsMapNode->FindPinChecked(TEXT("SelectedProtocol"));
+	UEdGraphPin* GetFunctionsMapNodeOutFunctionsMapPin = GetFunctionsMapNode->FindPinChecked(TEXT("OutFunctionsMap"));
+	UEdGraphPin* GetFunctionsMapNodeOutIsSuccessPin = GetFunctionsMapNode->FindPinChecked(UEdGraphSchema_K2::PN_ReturnValue);
+	UEdGraphPin* GetFunctionsMapNodeThenPin = GetFunctionsMapNode->GetThenPin();
+
+	// Hook up inputs
+	CompilerContext.MovePinLinksToIntermediate(*SelfPin, *GetFunctionsMapNodeSelfPin);
+	CompilerContext.MovePinLinksToIntermediate(*GetExecPin(), *GetFunctionsMapNodeExecPin);
+	CompilerContext.MovePinLinksToIntermediate(*GetInputDMXFixturePatchPin(), *GetFunctionsMapNodeInFixturePatchPin);
+	CompilerContext.MovePinLinksToIntermediate(*GetInputDMXProtocolPin(), *GetFunctionsMapNodeSelectedProtocolPin);
+
+	// Hook up outputs
+	CompilerContext.MovePinLinksToIntermediate(*GetOutputFunctionsMapPin(), *GetFunctionsMapNodeOutFunctionsMapPin);
+	CompilerContext.MovePinLinksToIntermediate(*GetOutputIsSuccessPin(), *GetFunctionsMapNodeOutIsSuccessPin);
+
+	// Call dmx active mode pin functions
+	if (UserDefinedPins.Num() > 0)
+	{
+		TArray<UEdGraphPin*> IntPairs;
+		TArray<UEdGraphPin*> NamePairs;
+
+		// Call functions for dmx function values
+		for (const TSharedPtr<FUserPinInfo>& PinInfo : UserDefinedPins)
+		{
+			UEdGraphPin* Pin = FindPinChecked(PinInfo->PinName);
+			{
+				if (Pin->Direction == EGPD_Output)
+				{
+					IntPairs.Add(Pin);
+				}
+				else if (Pin->Direction == EGPD_Input)
+				{
+					NamePairs.Add(Pin);
+				}
+			}
+		}
+
+		check(IntPairs.Num() == NamePairs.Num());
+
+		UEdGraphPin* LastThenPin = GetFunctionsMapNodeThenPin;
+		bool bResult = false;
+		for (int32 PairIndex = 0; PairIndex < NamePairs.Num(); ++PairIndex)
+		{
+			const FName GetFunctionsValueName = GET_FUNCTION_NAME_CHECKED(UDMXSubsystem, GetFunctionsValue);
+			UFunction* GetFunctionsValuePointer = FindField<UFunction>(UDMXSubsystem::StaticClass(), GetFunctionsValueName);
+			check(GetFunctionsValuePointer);
+
+			UK2Node_CallFunction* GetFunctionsValueNode = CompilerContext.SpawnIntermediateNode<UK2Node_CallFunction>(this, SourceGraph);
+			GetFunctionsValueNode->FunctionReference.SetExternalMember(GetFunctionsValueName, UDMXSubsystem::StaticClass());
+			GetFunctionsValueNode->AllocateDefaultPins();
+
+			UEdGraphPin* GetFunctionsValueSelfPin = GetFunctionsValueNode->FindPinChecked(UEdGraphSchema_K2::PN_Self);
+			UEdGraphPin* GetFunctionsValueExecPin = GetFunctionsValueNode->GetExecPin();
+			UEdGraphPin* GetFunctionsValueNodeOutInNamePin = GetFunctionsValueNode->FindPinChecked(TEXT("InName"));
+			UEdGraphPin* GetFunctionsValueNodeOutFunctionsMapPin = GetFunctionsValueNode->FindPinChecked(TEXT("InFunctionsMap"));
+			UEdGraphPin* GetFunctionsValueNodeOutValuePin = GetFunctionsValueNode->FindPinChecked(UEdGraphSchema_K2::PN_ReturnValue);
+			UEdGraphPin* GetFunctionsValueNodeThenPin = GetFunctionsValueNode->GetThenPin();
+
+			// Input
+			CompilerContext.MovePinLinksToIntermediate(*SelfPin, *GetFunctionsValueSelfPin);
+			CompilerContext.MovePinLinksToIntermediate(*NamePairs[PairIndex], *GetFunctionsValueNodeOutInNamePin);
+
+			// Output
+			Schema->TryCreateConnection(GetFunctionsValueNodeOutFunctionsMapPin, GetFunctionsMapNodeOutFunctionsMapPin);
+			CompilerContext.MovePinLinksToIntermediate(*IntPairs[PairIndex], *GetFunctionsValueNodeOutValuePin);
+
+			// Execution
+			Schema->TryCreateConnection(LastThenPin, GetFunctionsValueExecPin);
+			LastThenPin = GetFunctionsValueNodeThenPin;
+		}
+
+		CompilerContext.MovePinLinksToIntermediate(*GetThenPin(), *LastThenPin);
+	}
+	else
+	{
+		CompilerContext.MovePinLinksToIntermediate(*GetThenPin(), *GetFunctionsMapNodeThenPin);
+	}
+}
+
+void UK2Node_GetDMXActiveModeFunctionValues::GetMenuActions(FBlueprintActionDatabaseRegistrar& ActionRegistrar) const
+{
+	UClass* ActionKey = GetClass();
+
+	if (ActionRegistrar.IsOpenForRegistration(ActionKey))
+	{
+		UBlueprintNodeSpawner* NodeSpawner = UBlueprintNodeSpawner::Create(GetClass());
+		check(NodeSpawner != nullptr);
+
+		ActionRegistrar.AddBlueprintAction(ActionKey, NodeSpawner);
+	}
+}
+
+FText UK2Node_GetDMXActiveModeFunctionValues::GetMenuCategory() const
+{
+	return FText::FromString(DMX_K2_CATEGORY_NAME);
+}
+
+UEdGraphPin* UK2Node_GetDMXActiveModeFunctionValues::CreatePinFromUserDefinition(const TSharedPtr<FUserPinInfo> NewPinInfo)
+{
+	UEdGraphPin* NewPin = CreatePin(NewPinInfo->DesiredPinDirection, NewPinInfo->PinType, NewPinInfo->PinName);
+	const UEdGraphSchema_K2* K2Schema = GetDefault<UEdGraphSchema_K2>();
+	K2Schema->SetPinAutogeneratedDefaultValue(NewPin, NewPinInfo->PinDefaultValue);
+
+	if (NewPinInfo->DesiredPinDirection == EEdGraphPinDirection::EGPD_Input)
+	{
+		NewPin->bHidden = true;
+	}
+
+	return NewPin;
+}
+
+bool UK2Node_GetDMXActiveModeFunctionValues::CanCreateUserDefinedPin(const FEdGraphPinType& InPinType, EEdGraphPinDirection InDesiredDirection, FText& OutErrorMessage)
+{
+	if (!IsEditable())
+	{
+		return false;
+	}
+
+	// Make sure that if this is an exec node we are allowed one.
+	const UEdGraphSchema_K2* Schema = GetDefault<UEdGraphSchema_K2>();
+	if (InPinType.PinCategory == UEdGraphSchema_K2::PC_Exec && !CanModifyExecutionWires())
+	{
+		OutErrorMessage = LOCTEXT("MultipleExecPinError", "Cannot support more exec pins!");
+		return false;
+	}
+	else
+	{
+		TArray<TSharedPtr<UEdGraphSchema_K2::FPinTypeTreeInfo>> TypeTree;
+		Schema->GetVariableTypeTree(TypeTree, ETypeTreeFilter::RootTypesOnly);
+
+		bool bIsValid = false;
+		for (TSharedPtr<UEdGraphSchema_K2::FPinTypeTreeInfo>& TypeInfo : TypeTree)
+		{
+			FEdGraphPinType CurrentType = TypeInfo->GetPinType(false);
+			// only concerned with the list of categories
+			if (CurrentType.PinCategory == InPinType.PinCategory)
+			{
+				bIsValid = true;
+				break;
+			}
+		}
+
+		if (!bIsValid)
+		{
+			OutErrorMessage = LOCTEXT("AddInputPinError", "Cannot add pins of this type to custom event node!");
+			return false;
+		}
+	}
+
+	return true;
+}
+
+bool UK2Node_GetDMXActiveModeFunctionValues::ModifyUserDefinedPinDefaultValue(TSharedPtr<FUserPinInfo> PinInfo, const FString& NewDefaultValue)
+{
+	if (Super::ModifyUserDefinedPinDefaultValue(PinInfo, NewDefaultValue))
+	{
+		const UEdGraphSchema_K2* K2Schema = GetDefault<UEdGraphSchema_K2>();
+		K2Schema->HandleParameterDefaultValueChanged(this);
+
+		return true;
+	}
+	return false;
+}
+
+UEdGraphPin* UK2Node_GetDMXActiveModeFunctionValues::GetInputDMXFixturePatchPin() const
+{
+	UEdGraphPin* Pin = FindPinChecked(InputDMXFixturePatchPinName);
+	check(Pin->Direction == EGPD_Input);
+
+	return Pin;
+}
+
+UEdGraphPin* UK2Node_GetDMXActiveModeFunctionValues::GetInputDMXProtocolPin() const
+{
+	UEdGraphPin* Pin = FindPinChecked(InputDMXProtocolPinName);
+	check(Pin->Direction == EGPD_Input);
+
+	return Pin;
+}
+
+UEdGraphPin* UK2Node_GetDMXActiveModeFunctionValues::GetOutputFunctionsMapPin() const
+{
+	UEdGraphPin* Pin = FindPinChecked(OutputFunctionsMapPinName);
+	check(Pin->Direction == EGPD_Output);
+
+	return Pin;
+}
+
+UEdGraphPin* UK2Node_GetDMXActiveModeFunctionValues::GetOutputIsSuccessPin() const
+{
+	UEdGraphPin* Pin = FindPinChecked(OutputIsSuccessPinName);
+	check(Pin->Direction == EGPD_Output);
+
+	return Pin;
+}
+
+UEdGraphPin* UK2Node_GetDMXActiveModeFunctionValues::GetThenPin() const
+{
+	const UEdGraphSchema_K2* K2Schema = GetDefault<UEdGraphSchema_K2>();
+
+	UEdGraphPin* Pin = FindPinChecked(UEdGraphSchema_K2::PN_Then);
+	check(Pin->Direction == EGPD_Output);
+	return Pin;
+}
+
+void UK2Node_GetDMXActiveModeFunctionValues::ExposeFunctions()
+{
+	if (bIsExposed == true && UserDefinedPins.Num())
+	{
+		return;
+	}
+
+	const UEdGraphSchema_K2* K2Schema = GetDefault<UEdGraphSchema_K2>();
+	if (const FDMXFixtureMode* ActiveFixtureMode = GetActiveFixtureMode())
+	{
+		for (const FDMXFixtureFunction& Function : ActiveFixtureMode->Functions)
+		{
+			{
+				FEdGraphPinType PinType;
+				PinType.PinCategory = UEdGraphSchema_K2::PC_Int;
+				UEdGraphPin* Pin = CreateUserDefinedPin(GetPinName(Function), PinType, EGPD_Output);
+			}
+
+			{
+				FEdGraphPinType PinType;
+				PinType.PinCategory = UEdGraphSchema_K2::PC_Name;
+				UEdGraphPin* Pin = CreateUserDefinedPin(*(GetPinName(Function).ToString() + FString("_Input")), PinType, EGPD_Input);
+				K2Schema->TrySetDefaultValue(*Pin, Function.FunctionName);
+			}
+
+			const bool bIsCompiling = GetBlueprint()->bBeingCompiled;
+			if (!bIsCompiling)
+			{
+				FBlueprintEditorUtils::MarkBlueprintAsStructurallyModified(GetBlueprint());
+			}
+		}
+	}
+	else
+	{
+		UE_LOG_DMXBLUEPRINTGRAPH(Verbose, TEXT("No active mode found"));
+		return;
+	}
+
+
+	Modify();
+	bIsExposed = true;
+}
+
+void UK2Node_GetDMXActiveModeFunctionValues::ResetFunctions()
+{
+	if (bIsExposed)
+	{
+		while (UserDefinedPins.Num())
+		{
+			TSharedPtr<FUserPinInfo> Pin = UserDefinedPins[0];
+			RemoveUserDefinedPin(Pin);
+		}
+
+		// Reconstruct the entry/exit definition and recompile the blueprint to make sure the signature has changed before any fixups
+		bDisableOrphanPinSaving = true;
+
+		const UEdGraphSchema_K2* K2Schema = GetDefault<UEdGraphSchema_K2>();
+		K2Schema->HandleParameterDefaultValueChanged(this);
+	}
+
+
+	bIsExposed = false;
+}
+
+UDMXEntityFixturePatch* UK2Node_GetDMXActiveModeFunctionValues::GetFixturePatchFromPin() const
+{
+	UEdGraphPin* FixturePatchPin = GetInputDMXFixturePatchPin();
+
+	// Case with default object
+	if (FixturePatchPin->DefaultObject != nullptr && FixturePatchPin->LinkedTo.Num() == 0)
+	{
+		return Cast<UDMXEntityFixturePatch>(FixturePatchPin->DefaultObject);
+	}
+
+	// Case with linked object
+	if (FixturePatchPin->LinkedTo.Num() > 0)
+	{
+		if (UK2Node_GetDMXFixturePatch* NodeGetFixturePatch = Cast<UK2Node_GetDMXFixturePatch>(FixturePatchPin->LinkedTo[0]->GetOwningNode()))
+		{
+			return NodeGetFixturePatch->GetFixturePatchRefFromPin().GetFixturePatch();
+		}
+	}
+
+	return nullptr;
+}
+
+FName UK2Node_GetDMXActiveModeFunctionValues::GetPinName(const FDMXFixtureFunction& Function)
+{
+	FString EnumString;
+	EnumString = StaticEnum<EDMXFixtureSignalFormat>()->GetDisplayNameTextByIndex((int64)Function.DataType).ToString();
+
+	return *FString::Printf(TEXT("%s_%s"), *Function.FunctionName, *EnumString);
+}
+
+const FDMXFixtureMode* UK2Node_GetDMXActiveModeFunctionValues::GetActiveFixtureMode() const
+{
+	if (UDMXEntityFixturePatch* FixturePatch = GetFixturePatchFromPin())
+	{
+		if (UDMXEntityFixtureType* FixtureType = FixturePatch->ParentFixtureTypeTemplate)
+		{
+			if (FixtureType->Modes.Num() && FixtureType->Modes.Num() >= FixturePatch->ActiveMode)
+			{
+				return &FixtureType->Modes[FixturePatch->ActiveMode];
+			}
+		}
+	}
+
+	return nullptr;
+}
+
+#undef LOCTEXT_NAMESPACE
diff --git a/Engine/Plugins/Experimental/DigitalMultiplex/DMXEngine/Source/DMXBlueprintGraph/Private/Widgets/SDynamicNameListGraphPin.h b/Engine/Plugins/Experimental/DigitalMultiplex/DMXEngine/Source/DMXBlueprintGraph/Private/Widgets/SDynamicNameListGraphPin.h
new file mode 100644
--- /dev/null
+++ b/Engine/Plugins/Experimental/DigitalMultiplex/DMXEngine/Source/DMXBlueprintGraph/Private/Widgets/SDynamicNameListGraphPin.h
@@ -0,0 +1,66 @@
+// Copyright Epic Games, Inc. All Rights Reserved.
+
+#pragma once
+
+#include "CoreMinimal.h"
+#include "SGraphPin.h"
+#include "Widgets/SNameListPicker.h"
+#include "Widgets/DeclarativeSyntaxSupport.h"
+
+template<typename TStructType>
+class DMXBLUEPRINTGRAPH_API SDynamicNameListGraphPin
+	: public SGraphPin
+{
+public:
+	SLATE_BEGIN_ARGS(SDynamicNameListGraphPin)
+	{}
+		SLATE_ARGUMENT(FSimpleMulticastDelegate*, UpdateOptionsDelegate)
+		SLATE_ATTRIBUTE(TArray<FName>, OptionsSource)
+	SLATE_END_ARGS()
+
+	/**  Slate widget construction method */
+	void Construct(const FArguments& InArgs, UEdGraphPin* InGraphPinObj)
+	{
+		OptionsSource = InArgs._OptionsSource;
+		SGraphPin::Construct(SGraphPin::FArguments(), InGraphPinObj);
+	}
+
+	//~ Begin SGraphPin Interface
+	virtual TSharedRef<SWidget>	GetDefaultValueWidget() override
+	{
+		//Create widget
+		return SNew(SNameListPicker)
+			.HasMultipleValues(false)
+			.Value(this, &SDynamicNameListGraphPin::GetValue)
+			.OnValueChanged(this, &SDynamicNameListGraphPin::SetValue)
+			.UpdateOptionsDelegate(UpdateOptionsDelegate)
+			.OptionsSource(OptionsSource)
+			.Visibility(this, &SGraphPin::GetDefaultValueVisibility);
+	}
+	//~ End SGraphPin Interface
+
+private:
+	FName GetValue() const
+	{
+		TStructType ProtocolName;
+
+		if (!GraphPinObj->GetDefaultAsString().IsEmpty())
+		{
+			TStructType::StaticStruct()->ImportText(*GraphPinObj->GetDefaultAsString(), &ProtocolName, nullptr, EPropertyPortFlags::PPF_None, GLog, TStructType::StaticStruct()->GetName());
+		}
+
+		return ProtocolName;
+	}
+
+	void SetValue(FName NewValue)
+	{
+		FString ValueString;
+		TStructType NewProtocolName(NewValue);
+		TStructType::StaticStruct()->ExportText(ValueString, &NewProtocolName, nullptr, nullptr, EPropertyPortFlags::PPF_None, nullptr);
+		GraphPinObj->GetSchema()->TrySetDefaultValue(*GraphPinObj, ValueString);
+	}
+
+private:
+	TAttribute<TArray<FName>> OptionsSource;
+	FSimpleMulticastDelegate* UpdateOptionsDelegate;
+};
diff --git a/Engine/Plugins/Experimental/DigitalMultiplex/DMXEngine/Source/DMXBlueprintGraph/Public/K2Node_DMXBase.h b/Engine/Plugins/Experimental/DigitalMultiplex/DMXEngine/Source/DMXBlueprintGraph/Public/K2Node_DMXBase.h
new file mode 100644
--- /dev/null
+++ b/Engine/Plugins/Experimental/DigitalMultiplex/DMXEngine/Source/DMXBlueprintGraph/Public/K2Node_DMXBase.h
@@ -0,0 +1,57 @@
+// Copyright Epic Games, Inc. All Rights Reserved.
+
+#pragma once
+
+#include "CoreMinimal.h"
+#include "K2Node.h"
+#include "EdGraph/EdGraphNodeUtils.h"
+#include "K2Node_DMXBase.generated.h"
+
+struct FK2Node_DMXBaseHelper
+{
+	static FName ClassPinName;
+};
+
+/**   */
+UCLASS()
+class DMXBLUEPRINTGRAPH_API UK2Node_DMXBase : public UK2Node
+{
+	GENERATED_BODY()
+	
+public:
+	~UK2Node_DMXBase();
+
+	//~ Begin UK2Node implementation
+	virtual void AllocateDefaultPins() override;
+	virtual bool IsCompatibleWithGraph(const UEdGraph* TargetGraph) const override;
+	virtual void AddSearchMetaDataInfo(TArray<struct FSearchTagDataPair>& OutTaggedMetaData) const override;
+	virtual void PinConnectionListChanged(UEdGraphPin* Pin);
+	virtual void PinDefaultValueChanged(UEdGraphPin* Pin) override;
+	//~ End UK2Node implementation
+
+	virtual bool IsNodeSafeToIgnore() const override { return true; }
+	virtual void GetMenuActions(FBlueprintActionDatabaseRegistrar& ActionRegistrar) const override;
+	virtual FText GetMenuCategory() const override;
+
+	/** Get the blueprint input pin */
+	UEdGraphPin* GetClassPin(const TArray<UEdGraphPin*>* InPinsToSearch = NULL) const;
+
+	virtual UClass* GetClassPinBaseClass() const;
+
+	virtual void OnLibraryAssetChanged(UDMXLibrary* Library);
+
+protected:
+
+	UFUNCTION()
+	void OnLibraryEntitiesUpdated(UDMXLibrary* Library);
+
+	/** Tooltip text for this node. */
+	FText NodeTooltip;
+
+	/** Constructing FText strings can be costly, so we cache the node's title */
+	FNodeTextCache CachedNodeTitle;
+
+	/** Stores the Library to be able to unbind from it later */
+	UDMXLibrary* CachedLibrary;
+	FDelegateHandle LibraryEntitiesUpdatedHandle;
+};
diff --git a/Engine/Plugins/Experimental/DigitalMultiplex/DMXEngine/Source/DMXBlueprintGraph/Public/K2Node_GetAllFixturesOfType.h b/Engine/Plugins/Experimental/DigitalMultiplex/DMXEngine/Source/DMXBlueprintGraph/Public/K2Node_GetAllFixturesOfType.h
new file mode 100644
--- /dev/null
+++ b/Engine/Plugins/Experimental/DigitalMultiplex/DMXEngine/Source/DMXBlueprintGraph/Public/K2Node_GetAllFixturesOfType.h
@@ -0,0 +1,32 @@
+// Copyright Epic Games, Inc. All Rights Reserved.
+
+#pragma once
+
+#include "CoreMinimal.h"
+#include "K2Node_DMXBase.h"
+#include "K2Node_GetAllFixturesOfType.generated.h"
+
+struct FK2Node_GetAllFixturesOfType
+{
+	static FName FixtureTypePinName;
+	static FName OutResultPinName;
+};
+
+/**   */
+UCLASS()
+class DMXBLUEPRINTGRAPH_API UK2Node_GetAllFixturesOfType : public UK2Node_DMXBase
+{
+	GENERATED_BODY()
+	
+	UK2Node_GetAllFixturesOfType();
+
+	//~ Begin UK2Node implementation
+	virtual void AllocateDefaultPins() override;
+	virtual FText GetNodeTitle(ENodeTitleType::Type TitleType) const override;
+	virtual void ExpandNode(class FKismetCompilerContext& CompilerContext, UEdGraph* SourceGraph) override;
+	//~ End UK2Node implementation
+
+	virtual bool IsNodeSafeToIgnore() const override { return true; }
+	virtual void ReallocatePinsDuringReconstruction(TArray<UEdGraphPin*>& OldPins) override;
+
+};
diff --git a/Engine/Plugins/Experimental/DigitalMultiplex/DMXEngine/Source/DMXBlueprintGraph/Public/K2Node_GetDMXActiveModeFunctionValues.h b/Engine/Plugins/Experimental/DigitalMultiplex/DMXEngine/Source/DMXBlueprintGraph/Public/K2Node_GetDMXActiveModeFunctionValues.h
new file mode 100644
--- /dev/null
+++ b/Engine/Plugins/Experimental/DigitalMultiplex/DMXEngine/Source/DMXBlueprintGraph/Public/K2Node_GetDMXActiveModeFunctionValues.h
@@ -0,0 +1,88 @@
+// Copyright Epic Games, Inc. All Rights Reserved.
+
+#pragma once
+
+#include "CoreMinimal.h"
+
+#include "K2Node_EditablePinBase.h"
+#include "K2Node_AddPinInterface.h"
+
+#include "K2Node_GetDMXActiveModeFunctionValues.generated.h"
+
+class FBlueprintActionDatabaseRegistrar;
+class UBlueprint;
+class UEdGraph;
+class UEdGraphPin;
+class UDMXEntityFixturePatch;
+struct FDMXFixtureMode;
+struct FDMXFixtureFunction;
+
+UCLASS(BlueprintType)
+class UK2Node_GetDMXActiveModeFunctionValues
+	: public UK2Node_EditablePinBase
+{
+	GENERATED_BODY()
+
+public:
+	UK2Node_GetDMXActiveModeFunctionValues();
+
+	//~ Begin UEdGraphNode Interface
+	virtual FText GetNodeTitle(ENodeTitleType::Type TitleType) const override;
+	virtual void AllocateDefaultPins() override;
+	//~ End UEdGraphNode Interface
+
+	//~ Begin UK2Node Interface
+	virtual bool ShouldShowNodeProperties() const override { return true; }
+	virtual void ExpandNode(class FKismetCompilerContext& CompilerContext, UEdGraph* SourceGraph) override;
+	virtual void GetMenuActions(FBlueprintActionDatabaseRegistrar& ActionRegistrar) const override;
+	virtual FText GetMenuCategory() const override;
+	//~ End UK2Node Interface
+
+	//~ Begin UK2Node_EditablePinBase Interface
+	virtual bool ShouldUseConstRefParams() const override { return true; }
+	virtual UEdGraphPin* CreatePinFromUserDefinition(const TSharedPtr<FUserPinInfo> NewPinInfo) override;
+	virtual bool CanCreateUserDefinedPin(const FEdGraphPinType& InPinType, EEdGraphPinDirection InDesiredDirection, FText& OutErrorMessage) override;
+	virtual bool ModifyUserDefinedPinDefaultValue(TSharedPtr<FUserPinInfo> PinInfo, const FString& NewDefaultValue) override;
+	//~ End UK2Node_EditablePinBase Interface
+
+	/** Triggers when FixturePatch from input pin has been changed */
+	void OnFixturePatchChanged();
+
+	void RemovePinsRecursive(UEdGraphPin* Pin);
+	void RemoveOutputPin(UEdGraphPin* Pin);
+
+public:
+	bool IsExposed() const { return bIsExposed; }
+
+	UDMXEntityFixturePatch* GetFixturePatchFromPin() const;
+
+	UEdGraphPin* GetInputDMXFixturePatchPin() const;
+	UEdGraphPin* GetInputDMXProtocolPin() const;
+	UEdGraphPin* GetOutputFunctionsMapPin() const;
+	UEdGraphPin* GetOutputIsSuccessPin() const;
+	UEdGraphPin* GetThenPin() const;
+
+	/** Expose DMX function pins from fixture patch active mode */
+	void ExposeFunctions();
+
+	/** Removes DMX function node pins */
+	void ResetFunctions();
+
+private:
+	/** Get pin name based on FixtureFunction input */
+	FName GetPinName(const FDMXFixtureFunction& Function);
+
+	/**  Get pinter to active mode, never cache this pointer and keep it dynamic */
+	const FDMXFixtureMode* GetActiveFixtureMode() const;
+
+public:
+	static const FName InputDMXFixturePatchPinName;
+	static const FName InputDMXProtocolPinName;
+
+	static const FName OutputFunctionsMapPinName;
+	static const FName OutputIsSuccessPinName;
+
+public:
+	UPROPERTY()
+	bool bIsExposed;
+};
diff --git a/Engine/Plugins/Experimental/DigitalMultiplex/DMXEngine/Source/DMXEditor/Private/AssetTools/AssetTypeActions_DMXEditorLibrary.cpp b/Engine/Plugins/Experimental/DigitalMultiplex/DMXEngine/Source/DMXEditor/Private/AssetTools/AssetTypeActions_DMXEditorLibrary.cpp
new file mode 100644
--- /dev/null
+++ b/Engine/Plugins/Experimental/DigitalMultiplex/DMXEngine/Source/DMXEditor/Private/AssetTools/AssetTypeActions_DMXEditorLibrary.cpp
@@ -0,0 +1,38 @@
+// Copyright Epic Games, Inc. All Rights Reserved.
+
+#include "AssetTools/AssetTypeActions_DMXEditorLibrary.h"
+#include "Library/DMXLibrary.h"
+#include "DMXEditorModule.h"
+
+#define LOCTEXT_NAMESPACE "AssetTypeActions_DMXEditorLibrary"
+
+FText FAssetTypeActions_DMXEditorLibrary::GetName() const
+{
+	return LOCTEXT("DMXEditorLibrary", "DMX Library");
+}
+
+UClass * FAssetTypeActions_DMXEditorLibrary::GetSupportedClass() const
+{
+	return UDMXLibrary::StaticClass();
+}
+
+void FAssetTypeActions_DMXEditorLibrary::OpenAssetEditor(const TArray<UObject*>& InObjects, TSharedPtr<class IToolkitHost> EditWithinLevelEditor)
+{
+	TArray<UDMXLibrary*> LibrariesToOpen;
+	for (UObject* Obj : InObjects)
+	{
+		UDMXLibrary* Library = Cast<UDMXLibrary>(Obj);
+		if (Library)
+		{
+			LibrariesToOpen.Add(Library);
+		}
+	}
+
+	FDMXEditorModule& DMXEditorModule = FDMXEditorModule::Get();
+	for (UDMXLibrary* Library : LibrariesToOpen)
+	{
+		DMXEditorModule.CreateEditor(EToolkitMode::Standalone, EditWithinLevelEditor, Library);
+	}
+}
+
+#undef LOCTEXT_NAMESPACE
diff --git a/Engine/Plugins/Experimental/DigitalMultiplex/DMXEngine/Source/DMXEditor/Private/AssetTools/AssetTypeActions_DMXEditorLibrary.h b/Engine/Plugins/Experimental/DigitalMultiplex/DMXEngine/Source/DMXEditor/Private/AssetTools/AssetTypeActions_DMXEditorLibrary.h
new file mode 100644
--- /dev/null
+++ b/Engine/Plugins/Experimental/DigitalMultiplex/DMXEngine/Source/DMXEditor/Private/AssetTools/AssetTypeActions_DMXEditorLibrary.h
@@ -0,0 +1,20 @@
+// Copyright Epic Games, Inc. All Rights Reserved.
+
+#pragma once
+
+#include "CoreMinimal.h"
+#include "Toolkits/IToolkitHost.h"
+#include "AssetTypeActions_Base.h"
+#include "DMXEditorModule.h"
+
+class FAssetTypeActions_DMXEditorLibrary: public FAssetTypeActions_Base
+{
+public:
+	//~ Begin IAssetTypeActions implementation
+	virtual FText GetName() const override;
+	virtual UClass* GetSupportedClass() const override;
+	virtual FColor GetTypeColor() const override { return FColor(62, 140, 35); }
+	virtual void OpenAssetEditor(const TArray<UObject*>& InObjects, TSharedPtr<class IToolkitHost> EditWithinLevelEditor = TSharedPtr<IToolkitHost>()) override;
+	virtual uint32 GetCategories() override { return FDMXEditorModule::GetAssetCategory(); }
+	//~ End IAssetTypeActions implementation
+};
diff --git a/Engine/Plugins/Experimental/DigitalMultiplex/DMXEngine/Source/DMXEditor/Private/Commands/DMXEditorCommands.cpp b/Engine/Plugins/Experimental/DigitalMultiplex/DMXEngine/Source/DMXEditor/Private/Commands/DMXEditorCommands.cpp
new file mode 100644
--- /dev/null
+++ b/Engine/Plugins/Experimental/DigitalMultiplex/DMXEngine/Source/DMXEditor/Private/Commands/DMXEditorCommands.cpp
@@ -0,0 +1,35 @@
+// Copyright Epic Games, Inc. All Rights Reserved.
+
+#include "Commands/DMXEditorCommands.h"
+
+#define LOCTEXT_NAMESPACE "DMXEditorCommands"
+
+FDMXEditorCommandsImpl::FDMXEditorCommandsImpl()
+	: TCommands<FDMXEditorCommandsImpl>(TEXT("DMXEditor"), LOCTEXT("DMXEditor", "DMX Editor"), NAME_None, FEditorStyle::GetStyleSetName())
+{}
+
+void FDMXEditorCommandsImpl::RegisterCommands()
+{
+	UI_COMMAND(GoToDocumentation, "View Documentation", "View documentation about DMX editor", EUserInterfaceActionType::Button, FInputChord());
+
+	UI_COMMAND(AddNewEntityController, "New Controller", "Creates a new Controller in this library", EUserInterfaceActionType::Button, FInputChord());
+	UI_COMMAND(AddNewEntityFixtureType, "New Fixture Type", "Creates a new Fixture Type in this library", EUserInterfaceActionType::Button, FInputChord());
+	UI_COMMAND(AddNewEntityFixturePatch, "Add Fixture", "Creates a new Fixture Patch in this library", EUserInterfaceActionType::Button, FInputChord());
+}
+
+void FDMXEditorCommands::Register()
+{
+	return FDMXEditorCommandsImpl::Register();
+}
+
+const FDMXEditorCommandsImpl& FDMXEditorCommands::Get()
+{
+	return FDMXEditorCommandsImpl::Get();
+}
+
+void FDMXEditorCommands::Unregister()
+{
+	return FDMXEditorCommandsImpl::Unregister();
+}
+
+#undef LOCTEXT_NAMESPACE
diff --git a/Engine/Plugins/Experimental/DigitalMultiplex/DMXEngine/Source/DMXEditor/Private/Customizations/DMXEditorPropertyEditorCustomization.cpp b/Engine/Plugins/Experimental/DigitalMultiplex/DMXEngine/Source/DMXEditor/Private/Customizations/DMXEditorPropertyEditorCustomization.cpp
new file mode 100644
--- /dev/null
+++ b/Engine/Plugins/Experimental/DigitalMultiplex/DMXEngine/Source/DMXEditor/Private/Customizations/DMXEditorPropertyEditorCustomization.cpp
@@ -0,0 +1,863 @@
+// Copyright Epic Games, Inc. All Rights Reserved.
+
+#include "Customizations/DMXEditorPropertyEditorCustomization.h"
+#include "DMXEditor.h"
+#include "DMXEditorLog.h"
+
+#include "DMXProtocolTypes.h"
+#include "Library/DMXEntityReference.h"
+#include "Library/DMXEntity.h"
+#include "Library/DMXEntityController.h"
+#include "Library/DMXEntityFixtureType.h"
+#include "Library/DMXEntityFixturePatch.h"
+#include "Widgets/SDMXEntityDropdownMenu.h"
+#include "Game/DMXComponent.h"
+
+#include "DetailLayoutBuilder.h"
+#include "DetailCategoryBuilder.h"
+#include "IDetailPropertyRow.h"
+#include "IDetailChildrenBuilder.h"
+#include "DetailWidgetRow.h"
+#include "PropertyCustomizationHelpers.h"
+#include "IPropertyUtilities.h"
+
+#include "Widgets/Input/SButton.h"
+#include "Widgets/Input/SEditableTextBox.h"
+#include "Widgets/Images/SImage.h"
+#include "Widgets/Input/SComboBox.h"
+#include "Widgets/Input/SNumericEntryBox.h"
+#include "ScopedTransaction.h"
+
+#define LOCTEXT_NAMESPACE "DMXCustomizeDetails"
+
+static void CollectChildPropertiesRecursive(TSharedPtr<IPropertyHandle> Node, TArray<TSharedPtr<IPropertyHandle>>& OutProperties)
+{
+	uint32 NodeNumChildren = 0;
+	Node->GetNumChildren(NodeNumChildren);
+
+	for (uint32 ChildIdx = 0; ChildIdx < NodeNumChildren; ++ChildIdx)
+	{
+		TSharedPtr<IPropertyHandle> ChildHandle = Node->GetChildHandle(ChildIdx);
+		CollectChildPropertiesRecursive(ChildHandle, OutProperties);
+
+		if (ChildHandle->GetProperty())
+		{
+			OutProperties.AddUnique(ChildHandle);
+		}
+	}
+}
+
+void FDMXCustomization::CustomizeDetails(IDetailLayoutBuilder& DetailLayout)
+{
+	// keep Display Name as first property
+	DetailLayout.EditCategory("Entity Properties", FText::GetEmpty(), ECategoryPriority::Important);
+}
+
+void FDMXControllersDetails::CustomizeDetails(IDetailLayoutBuilder& DetailLayout)
+{
+	FDMXCustomization::CustomizeDetails(DetailLayout);
+
+	// Hide Universes for Controllers because the user should set them by range.
+	DetailLayout.HideProperty(GET_MEMBER_NAME_CHECKED(UDMXEntityUniverseManaged, Universes), UDMXEntityUniverseManaged::StaticClass());
+}
+
+FDMXFixtureTypeFunctionsDetails::FDMXFixtureTypeFunctionsDetails(TWeakPtr<FDMXEditor> InDMXEditorPtr)
+	: DMXEditorPtr(InDMXEditorPtr)
+{
+	check(DMXEditorPtr.IsValid());
+
+	const TArray<UDMXEntity*>&& SelectedEntities = DMXEditorPtr.Pin()->GetSelectedEntitiesFromTypeTab(UDMXEntityFixtureType::StaticClass());
+	for (UDMXEntity* Entity : SelectedEntities)
+	{
+		if (UDMXEntityFixtureType* FixtureType = Cast<UDMXEntityFixtureType>(Entity))
+		{
+			SelectedFixtures.Add(FixtureType);
+		}
+	}
+}
+
+void FDMXFixtureTypeFunctionsDetails::CustomizeHeader(TSharedRef<IPropertyHandle> InStructPropertyHandle, class FDetailWidgetRow& InHeaderRow, IPropertyTypeCustomizationUtils& InStructCustomizationUtils)
+{
+	InHeaderRow
+		.NameContent()
+		[
+			InStructPropertyHandle->CreatePropertyNameWidget()
+		]
+		.ValueContent()
+		[
+			InStructPropertyHandle->CreatePropertyValueWidget(false)
+		];
+}
+
+void FDMXFixtureTypeFunctionsDetails::CustomizeChildren(TSharedRef<IPropertyHandle> InStructPropertyHandle, IDetailChildrenBuilder& InStructBuilder, IPropertyTypeCustomizationUtils& InStructCustomizationUtils)
+{
+	// Retrieve structure's child properties
+	uint32 NumChildren;
+	InStructPropertyHandle->GetNumChildren(NumChildren);
+	TMap<FName, TSharedPtr< IPropertyHandle > > PropertyHandles;
+	for (uint32 ChildIndex = 0; ChildIndex < NumChildren; ++ChildIndex)
+	{
+		TSharedRef<IPropertyHandle> ChildHandle = InStructPropertyHandle->GetChildHandle(ChildIndex).ToSharedRef();
+		const FName PropertyName = ChildHandle->GetProperty()->GetFName();
+		PropertyHandles.Add(PropertyName, ChildHandle);
+	}
+
+	FText NewPropertyLabel;
+	FText ToolTip;
+	GetCustomNameFieldSettings(NewPropertyLabel, NamePropertyName, ToolTip, ExistingNameError);
+
+	// Check validity of the name property handle
+	check(PropertyHandles[NamePropertyName]->IsValidHandle());
+
+	NamePropertyHandle = PropertyHandles[NamePropertyName];
+
+	for (const auto& Itt : PropertyHandles)
+	{
+		// Never show default name property, we will show custom editable text instead
+		if (Itt.Key != NamePropertyName)
+		{
+			AddProperty(InStructBuilder, Itt.Key, Itt.Value.ToSharedRef());
+		}
+		else if (Itt.Key == NamePropertyName)
+		{
+			// Add custom name widget
+			InStructBuilder
+				.AddCustomRow(LOCTEXT("FunctionNameWidget", "FunctionNameWidget"))
+				.NameContent()
+				[
+					SNew(STextBlock)
+					.Text(NewPropertyLabel)
+					.Font(IDetailLayoutBuilder::GetDetailFont())
+				]
+				.ValueContent()
+				.MaxDesiredWidth(250.0f)
+				[
+					// The lambdas below are a workaround to avoid ambiguous TSharedFromThis template resolution
+					// that could be both IPropertyTypeCustomization and IDetailCustomization.
+					SAssignNew(NameEditableTextBox, SEditableTextBox)
+					.Text(this, &FDMXFixtureTypeFunctionsDetails::OnGetFunctionName)
+					.ToolTipText(ToolTip)
+					.OnTextChanged(this, &FDMXFixtureTypeFunctionsDetails::OnFunctionNameChanged)
+					.OnTextCommitted(this, &FDMXFixtureTypeFunctionsDetails::OnFunctionNameCommitted)
+					.Font(IDetailLayoutBuilder::GetDetailFont())
+				];
+		}
+	}
+}
+
+void FDMXFixtureTypeFunctionsDetails::AddProperty(IDetailChildrenBuilder& InStructBuilder, const FName& PropertyName, TSharedRef<IPropertyHandle> PropertyHandle)
+{
+	// Default implementation simply adds the properties
+	InStructBuilder.AddProperty(PropertyHandle);
+}
+
+void FDMXFixtureTypeFunctionsDetails::OnFunctionNameChanged(const FText& InNewText)
+{
+	check(NameEditableTextBox.IsValid() && NamePropertyHandle.IsValid() && NamePropertyHandle->IsValidHandle());
+
+	if (FText::TrimPrecedingAndTrailing(InNewText).IsEmpty())
+	{
+		NameEditableTextBox->SetError(LOCTEXT("FunctionNameError_Empty", "The name can't be blank!"));
+		return;
+	}
+
+	FString CurrentName;
+	NamePropertyHandle->GetValue(CurrentName);
+
+	const FString& NewName = InNewText.ToString();
+	if (CurrentName.Equals(NewName))
+	{
+		NameEditableTextBox->SetError(FText::GetEmpty());
+		return;
+	}
+
+	const TArray<FString>&& ExistingNames = GetExistingNames();
+	if (ExistingNames.Contains(NewName))
+	{
+		NameEditableTextBox->SetError(ExistingNameError);
+	}
+	else
+	{
+		NameEditableTextBox->SetError(FText::GetEmpty());
+	}
+}
+
+void FDMXFixtureTypeFunctionsDetails::OnFunctionNameCommitted(const FText& InNewText, ETextCommit::Type InTextCommit)
+{
+	if (InTextCommit != ETextCommit::OnCleared && !NameEditableTextBox->HasError())
+	{
+		const FString& NewName = InNewText.ToString();
+		SetFunctionName(NewName);
+	}
+	NameEditableTextBox->SetError(FText::GetEmpty());
+}
+
+FText FDMXFixtureTypeFunctionsDetails::OnGetFunctionName() const
+{
+	FString Name;
+	NamePropertyHandle->GetValue(Name);
+	return FText::FromString(Name);
+}
+
+void FDMXFixtureTypeFunctionsDetails::SetFunctionName(const FString& NewName)
+{
+	if (NamePropertyHandle.IsValid() && NamePropertyHandle->IsValidHandle())
+	{
+		NamePropertyHandle->SetValue(NewName);
+	}
+}
+
+void FDMXFixtureModeDetails::CustomizeChildren(TSharedRef<IPropertyHandle> InStructPropertyHandle, IDetailChildrenBuilder& InStructBuilder, IPropertyTypeCustomizationUtils& InStructCustomizationUtils)
+{
+	FDMXFixtureTypeFunctionsDetails::CustomizeChildren(InStructPropertyHandle, InStructBuilder, InStructCustomizationUtils);
+
+	AutoChannelSpanHandle = InStructPropertyHandle->GetChildHandle(GET_MEMBER_NAME_CHECKED(FDMXFixtureMode, bAutoChannelSpan));
+	check(AutoChannelSpanHandle);
+}
+
+void FDMXFixtureModeDetails::GetCustomNameFieldSettings(FText& OutNewPropertyLabel, FName& OutNamePropertyName, FText& OutToolTip, FText& OutExistingNameError)
+{
+	OutNewPropertyLabel = LOCTEXT("FixtureModeNameLabel", "Mode Name");
+	OutNamePropertyName = GET_MEMBER_NAME_CHECKED(FDMXFixtureMode, ModeName);
+	OutToolTip = LOCTEXT("FixtureModeNameToolTip", "The name of this mode");
+	OutExistingNameError = LOCTEXT("FixtureModeName_Existent", "This name is already used by another Mode in this fixture!");
+}
+
+TArray<FString> FDMXFixtureModeDetails::GetExistingNames() const
+{
+	TArray<FString> ExistingNames;
+	for (UDMXEntityFixtureType* Fixture : SelectedFixtures)
+	{
+		if (Fixture == nullptr)
+		{
+			continue; 
+		}
+
+		for (const FDMXFixtureMode& Mode : Fixture->Modes)
+		{
+			ExistingNames.Add(Mode.ModeName);
+		}
+	}
+
+	return ExistingNames;
+}
+
+void FDMXFixtureFunctionDetails::CustomizeChildren(TSharedRef<IPropertyHandle> InStructPropertyHandle, IDetailChildrenBuilder& InStructBuilder, IPropertyTypeCustomizationUtils& InStructCustomizationUtils)
+{
+	FDMXFixtureTypeFunctionsDetails::CustomizeChildren(InStructPropertyHandle, InStructBuilder, InStructCustomizationUtils);
+
+	StructPropertyHandle = InStructPropertyHandle;
+
+	DataTypeHandle = InStructPropertyHandle->GetChildHandle(GET_MEMBER_NAME_CHECKED(FDMXFixtureFunction, DataType));
+	check(DataTypeHandle);
+
+	DefaultValueHandle = InStructPropertyHandle->GetChildHandle(GET_MEMBER_NAME_CHECKED(FDMXFixtureFunction, DefaultValue));
+	check(DefaultValueHandle);
+
+	// Create fields for the individual channels (each byte) that will be displayed depending on the selected DataType
+	AddChannelInputFields(InStructBuilder);
+}
+
+void FDMXFixtureFunctionDetails::AddProperty(IDetailChildrenBuilder& InStructBuilder, const FName& PropertyName, TSharedRef<IPropertyHandle> PropertyHandle)
+{
+	if (PropertyName == GET_MEMBER_NAME_CHECKED(FDMXFixtureFunction, SubFunctions))
+	{
+		// Conditionally control Sub Functions property visibility
+		InStructBuilder
+			.AddProperty(PropertyHandle)
+			.Visibility(MakeAttributeSP(this, &FDMXFixtureFunctionDetails::GetSubFunctionsVisibility));
+
+		return;
+	}
+
+	// Simply add the property
+	FDMXFixtureTypeFunctionsDetails::AddProperty(InStructBuilder, PropertyName, PropertyHandle);
+}
+
+EVisibility FDMXFixtureFunctionDetails::GetSubFunctionsVisibility() const
+{
+	void* DataTypePtr = nullptr;
+	if (DataTypeHandle->GetValueData(DataTypePtr) == FPropertyAccess::Success)
+	{
+		EDMXFixtureSignalFormat* DataType = reinterpret_cast<EDMXFixtureSignalFormat*>(DataTypePtr);
+		if (*DataType == EDMXFixtureSignalFormat::E8BitSubFunctions)
+		{
+			return EVisibility::Visible;
+		}
+	}
+	return EVisibility::Collapsed;
+}
+
+void FDMXFixtureFunctionDetails::AddChannelInputFields(IDetailChildrenBuilder& InStructBuilder)
+{
+	const FMargin Padding = FMargin(2.0f, 0.0f, 0.0f, 0.0f);
+
+	InStructBuilder
+		.AddCustomRow(LOCTEXT("ChannelsWidget", "Channels Widget"))
+		.NameContent()
+		[
+			SNew(STextBlock)
+			.Text(LOCTEXT("ChannelsValues", "Channels Values"))
+			.Font(IDetailLayoutBuilder::GetDetailFont())
+			.ToolTipText(LOCTEXT("ChannelsToolTip", "Individual channels values. Useful for things like colors"))
+		]
+		.ValueContent()
+		.MinDesiredWidth(340.0f)
+		.MaxDesiredWidth(340.0f)
+		[
+			SNew(SHorizontalBox)
+			+ SHorizontalBox::Slot()
+			.FillWidth(1.0f)
+			.VAlign(VAlign_Fill)
+			[
+				CreateChannelField(1, SNumericEntryBox<uint8>::RedLabelBackgroundColor)
+			]
+
+			+ SHorizontalBox::Slot()
+			.FillWidth(1.0f)
+			.VAlign(VAlign_Fill)
+			.Padding(Padding)
+			[
+				CreateChannelField(2, SNumericEntryBox<uint8>::GreenLabelBackgroundColor)
+			]
+			
+			+ SHorizontalBox::Slot()
+			.FillWidth(1.0f)
+			.VAlign(VAlign_Fill)
+			.Padding(Padding)
+			[
+				CreateChannelField(3, SNumericEntryBox<uint8>::BlueLabelBackgroundColor)
+			]
+
+			+ SHorizontalBox::Slot()
+			.FillWidth(1.0f)
+			.VAlign(VAlign_Fill)
+			.Padding(Padding)
+			[
+				CreateChannelField(4)
+			]
+		]
+		.Visibility(MakeAttributeSP(this, &FDMXFixtureFunctionDetails::GetChannelInputVisibility, (uint8)1));
+}
+
+TSharedRef<SWidget> FDMXFixtureFunctionDetails::CreateChannelField(uint8 Channel, const FLinearColor& LabelColor /*= FLinearColor(0.0f, 0.0f, 0.0f, 0.5f)*/)
+{
+	return SNew(SNumericEntryBox<uint8>)
+		.MinValue(0)
+		.MaxValue(255)
+		.MaxSliderValue(255)
+		.Value(this, &FDMXFixtureFunctionDetails::GetChannelValue, Channel)
+		.OnValueChanged(this, &FDMXFixtureFunctionDetails::HandleChannelValueChanged, Channel)
+		.OnValueCommitted(this, &FDMXFixtureFunctionDetails::HandleChannelValueCommitted)
+		.Font(IDetailLayoutBuilder::GetDetailFont())
+		.AllowSpin(true)
+		.Visibility(MakeAttributeSP(this, &FDMXFixtureFunctionDetails::GetChannelInputVisibility, Channel))
+		.LabelPadding(0)
+		.Label()
+		[
+			SNumericEntryBox<uint8>::BuildLabel(
+				FText::AsNumber(Channel),
+				FLinearColor::White,
+				LabelColor)
+		];
+}
+
+TOptional<uint8> FDMXFixtureFunctionDetails::GetChannelValue(uint8 Channel) const
+{
+	TOptional<uint8> RetVal;
+	int64 Value;
+	if (DefaultValueHandle->GetValue(Value) == FPropertyAccess::Success)
+	{
+		const uint8 BytesOffset = (Channel - 1) * 8;
+		RetVal = Value >> BytesOffset & 0xff;
+	}
+
+	return RetVal;
+}
+
+EVisibility FDMXFixtureFunctionDetails::GetChannelInputVisibility(uint8 Channel) const
+{
+	void* DataTypePtr = nullptr;
+	if (DataTypeHandle->GetValueData(DataTypePtr) == FPropertyAccess::Success)
+	{
+		EDMXFixtureSignalFormat* DataType = (EDMXFixtureSignalFormat*)DataTypePtr;
+		if ( (*DataType >= EDMXFixtureSignalFormat::E32Bit && Channel <= 4)
+			|| (*DataType >= EDMXFixtureSignalFormat::E24Bit && Channel <= 3)
+			|| (*DataType >= EDMXFixtureSignalFormat::E16Bit && Channel <= 2) )
+		{
+			return EVisibility::Visible;
+		}
+
+		return EVisibility::Collapsed;
+	}
+
+	return EVisibility::Collapsed;
+}
+
+void FDMXFixtureFunctionDetails::HandleChannelValueChanged(uint8 NewValue, uint8 Channel)
+{
+	int64 DefaultValue;
+	if (DefaultValueHandle->GetValue(DefaultValue) == FPropertyAccess::Success)
+	{
+		uint8* ValueBytes = reinterpret_cast<uint8*>(&DefaultValue);
+		ValueBytes[Channel - 1] = NewValue;
+		DefaultValueHandle->SetValue(DefaultValue, EPropertyValueSetFlags::InteractiveChange);
+	}
+}
+
+void FDMXFixtureFunctionDetails::HandleChannelValueCommitted(uint8 NewValue, ETextCommit::Type CommitType)
+{
+	int64 DefaultValue;
+	if (DefaultValueHandle->GetValue(DefaultValue) == FPropertyAccess::Success)
+	{
+		DefaultValueHandle->SetValue(DefaultValue);
+	}
+}
+
+void FDMXFixtureFunctionDetails::GetCustomNameFieldSettings(FText& OutNewPropertyLabel, FName& OutNamePropertyName, FText& OutToolTip, FText& OutExistingNameError)
+{
+	OutNewPropertyLabel = LOCTEXT("FixtureFunctionNameLabel", "Function Name");
+	OutNamePropertyName = GET_MEMBER_NAME_CHECKED(FDMXFixtureFunction, FunctionName);
+	OutToolTip = LOCTEXT("FixtureFunctionNameToolTip", "The name of this function");
+	OutExistingNameError = LOCTEXT("FixtureFunctionName_Existent", "This name is already used by another function in this mode!");
+}
+
+TArray<FString> FDMXFixtureFunctionDetails::GetExistingNames() const
+{
+	check(NamePropertyHandle.IsValid() && NamePropertyHandle->IsValidHandle());
+
+	TArray<FString> ExistingNames;
+
+	// Enumerate all properties indexes from the function in a map. This way, if we change the Functions
+	// properties order in the future, we don't need to change this code because we'll find it by name.
+	TSharedRef<IPropertyHandle> FunctionStruct = NamePropertyHandle->GetParentHandle().ToSharedRef();
+	TMap<FName, uint32> FunctionPropertiesMap;
+	uint32 NumFunctionProperties;
+	FunctionStruct->GetNumChildren(NumFunctionProperties);
+	for (uint32 PropertyIndex = 0; PropertyIndex < NumFunctionProperties; ++PropertyIndex)
+	{
+		TSharedRef<IPropertyHandle> PropertyHandle = FunctionStruct->GetChildHandle(PropertyIndex).ToSharedRef();
+		const FName PropertyName = PropertyHandle->GetProperty()->GetFName();
+		FunctionPropertiesMap.Add(PropertyName, PropertyIndex);
+	}
+	const uint32 NamePropertyIndex = FunctionPropertiesMap[NamePropertyName];
+
+	// Get the parent functions array property to be able to read the other functions names inside the current Mode
+	TSharedRef<IPropertyHandle> ParentFunctionsArray = FunctionStruct->GetParentHandle().ToSharedRef();
+
+	uint32 NumFunctions;
+	ParentFunctionsArray->GetNumChildren(NumFunctions);
+	for (uint32 FunctionIndex = 0; FunctionIndex < NumFunctions; ++FunctionIndex)
+	{
+		// The current function in the functions array
+		TSharedRef<IPropertyHandle> Function = ParentFunctionsArray->GetChildHandle(FunctionIndex).ToSharedRef();
+
+		// Add the function name to the list of existing names
+		TSharedRef<IPropertyHandle> NameHandle = Function->GetChildHandle(NamePropertyIndex).ToSharedRef();
+		FString FunctionName;
+		NameHandle->GetValue(FunctionName);
+		ExistingNames.Add(FunctionName);
+	}
+
+	return ExistingNames;
+}
+
+void FDMXFixtureSubFunctionDetails::CustomizeChildren(TSharedRef<IPropertyHandle> InStructPropertyHandle, IDetailChildrenBuilder& InStructBuilder, IPropertyTypeCustomizationUtils& InStructCustomizationUtils)
+{
+	// Skip parent implementation and use base's
+	FDMXFixtureTypeFunctionsDetails::CustomizeChildren(InStructPropertyHandle, InStructBuilder, InStructCustomizationUtils);
+}
+
+void FDMXFixtureSubFunctionDetails::GetCustomNameFieldSettings(FText& OutNewPropertyLabel, FName& OutNamePropertyName, FText& OutToolTip, FText& OutExistingNameError)
+{
+	OutNewPropertyLabel = LOCTEXT("FixtureSubFunctionNameLabel", "Sub Function Name");
+	OutNamePropertyName = GET_MEMBER_NAME_CHECKED(FDMXFixtureSubFunction, FunctionName);
+	OutToolTip = LOCTEXT("FixtureSubFunctionNameToolTip", "The name of this sub function");
+	OutExistingNameError = LOCTEXT("FixtureSubFunctionName_Existent", "This name is already used by another sub function in this function!");
+}
+
+void FDMXFixturePatchesDetails::CustomizeDetails(IDetailLayoutBuilder& DetailLayout)
+{
+	FDMXCustomization::CustomizeDetails(DetailLayout);
+
+	// Make a Fixture Types dropdown for the Fixture Type template property
+	ParentFixtureTypeHandle = DetailLayout.GetProperty(GET_MEMBER_NAME_CHECKED(UDMXEntityFixturePatch, ParentFixtureTypeTemplate));
+	check(ParentFixtureTypeHandle->IsValidHandle());
+	DetailLayout.EditDefaultProperty(ParentFixtureTypeHandle)->CustomWidget(false)
+		.NameContent()
+		[
+			ParentFixtureTypeHandle->CreatePropertyNameWidget()
+		]
+		.ValueContent()
+		.MinDesiredWidth(200.0f)
+		.MaxDesiredWidth(400.0f)
+		[
+			SNew(SDMXEntityPickerButton<UDMXEntityFixtureType>)
+			.DMXEditor(DMXEditorPtr)
+			.CurrentEntity(this, &FDMXFixturePatchesDetails::GetParentFixtureTemplate)
+			.OnEntitySelected(this, &FDMXFixturePatchesDetails::OnParentTemplateSelected)
+			.HasMultipleValues(this, &FDMXFixturePatchesDetails::GetParentFixtureTypeIsMultipleValues)
+		];
+
+	// Make a modes dropdown to select the active Fixture Type Mode, if a valid Fixture Type is selected
+	ActiveModeHandle = DetailLayout.GetProperty(GET_MEMBER_NAME_CHECKED(UDMXEntityFixturePatch, ActiveMode));
+	check(ActiveModeHandle->IsValidHandle());
+
+	TSharedPtr<uint32> DefaultSelectedActiveMode = nullptr;
+	GenerateActiveModeOptions();
+	if (ActiveModeOptions.Num() > 0)
+	{
+		// Test if we have a single valid type of Fixture selected as Template
+		UObject* Object = nullptr;
+		if (ParentFixtureTypeHandle->GetValue(Object) == FPropertyAccess::Success && Object != nullptr)
+		{
+			int32 ActiveModeValue = 0;
+			if (ActiveModeHandle->GetValue(ActiveModeValue) == FPropertyAccess::Success)
+			{
+				DefaultSelectedActiveMode = ActiveModeOptions[ActiveModeValue];
+			}
+		}
+	}
+
+	DetailLayout.EditDefaultProperty(ActiveModeHandle)->CustomWidget(false)
+		.NameContent()
+		[
+			ActiveModeHandle->CreatePropertyNameWidget()
+		]
+		.ValueContent()
+		.MaxDesiredWidth(160.0f)
+		[
+			SNew(SComboBox< TSharedPtr<uint32> >)
+			.IsEnabled(this, &FDMXFixturePatchesDetails::GetActiveModeEditable)
+			.OptionsSource(&ActiveModeOptions)
+			.OnGenerateWidget(this, &FDMXFixturePatchesDetails::GenerateActiveModeOptionWidget)
+			.OnSelectionChanged(this, &FDMXFixturePatchesDetails::OnActiveModeChanged)
+			.InitiallySelectedItem(DefaultSelectedActiveMode)
+			[
+				SNew(STextBlock)
+				.MinDesiredWidth(50.0f)
+				.Text(this, &FDMXFixturePatchesDetails::GetCurrentActiveModeLabel)
+				.Font(DetailLayout.GetDetailFont())
+			]
+		];
+}
+
+void FDMXFixturePatchesDetails::GenerateActiveModeOptions()
+{
+	UObject* Object = nullptr;
+	if (ParentFixtureTypeHandle->GetValue(Object) == FPropertyAccess::Success)
+	{
+		if (UDMXEntityFixtureType* Fixture = Cast<UDMXEntityFixtureType>(Object))
+		{
+			const uint32 NumModes = Fixture->Modes.Num();
+			for (uint32 ModeIndex = 0; ModeIndex < NumModes; ++ModeIndex)
+			{
+				ActiveModeOptions.Add(MakeShared<uint32>(ModeIndex));
+			}
+		}
+	}
+}
+
+TWeakObjectPtr<UDMXEntityFixtureType> FDMXFixturePatchesDetails::GetParentFixtureTemplate() const
+{
+	UObject* Object;
+	if (ParentFixtureTypeHandle->GetValue(Object) == FPropertyAccess::Success)
+	{
+		return Cast<UDMXEntityFixtureType>(Object);
+	}
+	return nullptr;
+}
+
+void FDMXFixturePatchesDetails::OnParentTemplateSelected(UDMXEntity* NewTemplate) const
+{
+	ParentFixtureTypeHandle->SetValue(Cast<UDMXEntityFixtureType>(NewTemplate));
+}
+
+bool FDMXFixturePatchesDetails::GetParentFixtureTypeIsMultipleValues() const
+{
+	UObject* Object;
+	return ParentFixtureTypeHandle->GetValue(Object) == FPropertyAccess::MultipleValues;
+}
+
+bool FDMXFixturePatchesDetails::GetActiveModeEditable() const
+{
+	UObject* Object = nullptr;
+	if (ParentFixtureTypeHandle->GetValue(Object) == FPropertyAccess::Success && Object != nullptr)
+	{
+		if (UDMXEntityFixtureType* Fixture = Cast<UDMXEntityFixtureType>(Object))
+		{
+			return Fixture->Modes.Num() > 0;
+		}
+		return false;
+	}
+	return false;
+}
+
+TSharedRef<SWidget> FDMXFixturePatchesDetails::GenerateActiveModeOptionWidget(const TSharedPtr<uint32> InMode) const
+{
+	UObject* Object = nullptr;
+	if (ParentFixtureTypeHandle->GetValue(Object) == FPropertyAccess::Success && Object != nullptr)
+	{
+		if (UDMXEntityFixtureType* Patch = Cast<UDMXEntityFixtureType>(Object))
+		{
+			return SNew(STextBlock)
+				.Text(FText::FromString(Patch->Modes[*InMode].ModeName));
+		}
+	}
+
+	return SNullWidget::NullWidget;
+}
+
+void FDMXFixturePatchesDetails::OnActiveModeChanged(const TSharedPtr<uint32> InSelectedMode, ESelectInfo::Type SelectInfo)
+{
+	ActiveModeHandle->SetValue(*InSelectedMode);
+}
+
+FText FDMXFixturePatchesDetails::GetCurrentActiveModeLabel() const
+{
+	static const FText MultipleValuesLabel = LOCTEXT("MultipleValues_Label", "Multiple Values");
+	static const FText NullTypeLabel = LOCTEXT("NullFixtureType_Label", "No Fixture Type selected");
+	static const FText MultipleTypesLabel = LOCTEXT("MultipleFixtureTypes_Label", "Multiple Types Selected");
+	static const FText NoModesLabel = LOCTEXT("NoModes_Label", "No modes in Fixture Type");
+
+	UObject* Object = nullptr;
+	const FPropertyAccess::Result FixtureTemplateAccessResult = ParentFixtureTypeHandle->GetValue(Object);
+	UDMXEntityFixtureType* FixtureTemplate = Cast<UDMXEntityFixtureType>(Object);
+
+	// Is only one type of Fixture Type selected?
+	if (FixtureTemplateAccessResult == FPropertyAccess::Success)
+	{
+		// Is this type valid?
+		if (FixtureTemplate != nullptr)
+		{
+			// We can try to get the mode, although it could be a different one for each of the templates
+			int32 ModeValue = 0;
+			if (ActiveModeHandle->GetValue(ModeValue) == FPropertyAccess::Success)
+			{
+				const TArray<FDMXFixtureMode>& CurrentModes = FixtureTemplate->Modes;
+				if (CurrentModes.Num() > 0)
+				{
+					return FText::FromString(
+						CurrentModes[ModeValue].ModeName
+					);
+				}
+				else
+				{
+					return NoModesLabel;
+				}
+			}
+			return MultipleValuesLabel;
+		}
+		return NullTypeLabel;
+	}
+	return MultipleTypesLabel;
+}
+
+const FName FDMXEntityReferenceCustomization::NAME_DMXLibrary = GET_MEMBER_NAME_CHECKED(FDMXEntityReference, DMXLibrary);
+
+void FDMXEntityReferenceCustomization::CustomizeHeader(TSharedRef<IPropertyHandle> InPropertyHandle, FDetailWidgetRow& InHeaderRow, IPropertyTypeCustomizationUtils& CustomizationUtils)
+{
+	StructHandle = InPropertyHandle;
+
+	// If we should display any children, leave the standard header. Otherwise, create custom picker
+	TSharedRef<SWidget> ValueContent = !GetDisplayLibrary()
+		? CreateEntityPickerWidget(InPropertyHandle)
+		: InPropertyHandle->CreatePropertyValueWidget(false);
+
+	InHeaderRow
+		.NameContent()
+		[
+			InPropertyHandle->CreatePropertyNameWidget()
+		]
+		.ValueContent()
+		[
+			ValueContent
+		];
+}
+
+void FDMXEntityReferenceCustomization::CustomizeChildren(TSharedRef<IPropertyHandle> InPropertyHandle, IDetailChildrenBuilder& InChildBuilder, IPropertyTypeCustomizationUtils& CustomizationUtils)
+{
+	if (!GetDisplayLibrary())
+	{
+		// Don't add any child properties
+		return;
+	}
+
+	// Retrieve structure's child properties
+	uint32 NumChildren;
+	InPropertyHandle->GetNumChildren(NumChildren);
+	TMap<FName, TSharedPtr< IPropertyHandle > > PropertyHandles;
+	for (uint32 ChildIndex = 0; ChildIndex < NumChildren; ++ChildIndex)
+	{
+		TSharedRef<IPropertyHandle> ChildHandle = InPropertyHandle->GetChildHandle(ChildIndex).ToSharedRef();
+		const FName PropertyName = ChildHandle->GetProperty()->GetFName();
+		PropertyHandles.Add(PropertyName, ChildHandle);
+	}
+
+	// Add the properties
+	for (const auto& Itt : PropertyHandles)
+	{
+		if ((Itt.Key == NAME_DMXLibrary && GetDisplayLibrary())
+			|| (Itt.Key != NAME_DMXLibrary && Itt.Key != "EntityId")) // EntityId is private, so GET_MEMBER... won't work
+		{
+			InChildBuilder.AddProperty(Itt.Value.ToSharedRef());
+		}
+	}
+
+	// Add the picker
+	InChildBuilder.AddCustomRow(LOCTEXT("EntityReferencePickerSearchText", "Entity"))
+		.NameContent()
+		[
+			SNew(STextBlock)
+			.Font(CustomizationUtils.GetRegularFont())
+			.Text(this, &FDMXEntityReferenceCustomization::GetPickerPropertyLabel)
+		]
+		.ValueContent()
+		.MinDesiredWidth(200.0f)
+		.MaxDesiredWidth(400.0f)
+		[
+			CreateEntityPickerWidget(InPropertyHandle)
+		]
+		;
+}
+
+bool FDMXEntityReferenceCustomization::GetDisplayLibrary() const
+{
+	TArray<const void*> RawDataArr;
+	StructHandle->AccessRawData(RawDataArr);
+
+	for (const void* RawData : RawDataArr)
+	{
+		const FDMXEntityReference* EntityRefPtr = reinterpret_cast<const FDMXEntityReference*>(RawData);
+		if (!EntityRefPtr->bDisplayLibraryPicker)
+		{
+			return false;
+		}
+	}
+
+	return true;
+}
+
+TSharedRef<SWidget> FDMXEntityReferenceCustomization::CreateEntityPickerWidget(TSharedPtr<IPropertyHandle> InPropertyHandle) const
+{
+	return SNew(SDMXEntityPickerButton<UDMXEntity>)
+		.CurrentEntity(this, &FDMXEntityReferenceCustomization::GetCurrentEntity)
+		.HasMultipleValues(this, &FDMXEntityReferenceCustomization::GetEntityIsMultipleValues)
+		.OnEntitySelected(this, &FDMXEntityReferenceCustomization::OnEntitySelected)
+		.EntityTypeFilter(this, &FDMXEntityReferenceCustomization::GetEntityType)
+		.DMXLibrary(this, &FDMXEntityReferenceCustomization::GetDMXLibrary)
+		.IsEnabled(this, &FDMXEntityReferenceCustomization::GetPickerEnabled)
+		;
+}
+
+FText FDMXEntityReferenceCustomization::GetPickerPropertyLabel() const
+{
+	if (TSubclassOf<UDMXEntity> EntityType = GetEntityType())
+	{
+		return FDMXEditorUtils::GetEntityTypeNameText(EntityType, false);
+	}
+
+	return LOCTEXT("GenericTypeEntityLabel", "Entity");
+}
+
+bool FDMXEntityReferenceCustomization::GetPickerEnabled() const
+{
+	return GetEntityType() != nullptr;
+}
+
+TWeakObjectPtr<UDMXEntity> FDMXEntityReferenceCustomization::GetCurrentEntity() const
+{
+	if (GetEntityIsMultipleValues()) return nullptr;
+
+	void* StructPtr = nullptr;
+	if (StructHandle->GetValueData(StructPtr) == FPropertyAccess::Success && StructPtr != nullptr)
+	{
+		const FDMXEntityReference* EntityRef = reinterpret_cast<FDMXEntityReference*>(StructPtr);
+		return EntityRef->GetEntity();
+	}
+
+	return nullptr;
+}
+
+bool FDMXEntityReferenceCustomization::GetEntityIsMultipleValues() const
+{
+	TArray<void*> RawData;
+	StructHandle->AccessRawData(RawData);
+	if (RawData[0] == nullptr)
+	{
+		return true; 
+	}
+
+	bool bFirstEntitySet = false;
+	UDMXEntity* FirstEntityPtr = nullptr;
+
+	for (const void* StructPtr : RawData)
+	{
+		const FDMXEntityReference* EntityRefPtr = reinterpret_cast<const FDMXEntityReference*>(StructPtr);
+
+		if (bFirstEntitySet)
+		{
+			if (EntityRefPtr->GetEntity() != FirstEntityPtr)
+			{
+				return true;
+			}
+		}
+		else
+		{
+			FirstEntityPtr = EntityRefPtr->GetEntity();
+			bFirstEntitySet = true;
+		}
+	}
+
+	return false;
+}
+
+void FDMXEntityReferenceCustomization::OnEntitySelected(UDMXEntity* NewEntity) const
+{
+	FDMXEntityReference NewStructValues;
+	NewStructValues.SetEntity(NewEntity);
+
+	// Export new values to text format that can be imported later into the actual struct properties
+	FString TextValue;
+	FDMXEntityReference::StaticStruct()->ExportText(TextValue, &NewStructValues, nullptr, nullptr, EPropertyPortFlags::PPF_None, nullptr);
+
+	// Set values on edited property handle from exported text
+	ensure(StructHandle->SetValueFromFormattedString(TextValue, EPropertyValueSetFlags::DefaultFlags) == FPropertyAccess::Success);
+}
+
+TSubclassOf<UDMXEntity> FDMXEntityReferenceCustomization::GetEntityType() const
+{
+	TArray<void*> RawData;
+	StructHandle->AccessRawData(RawData);
+	if (RawData[0] == nullptr)
+	{
+		return nullptr; 
+	}
+
+	const TSubclassOf<UDMXEntity> FirstEntityType = reinterpret_cast<FDMXEntityReference*>(RawData[0])->GetEntityType();
+
+	for (const void* StructPtr : RawData)
+	{
+		const FDMXEntityReference* EntityRefPtr = reinterpret_cast<const FDMXEntityReference*>(StructPtr);
+		if (EntityRefPtr->GetEntityType() != FirstEntityType)
+		{
+			// Different types are selected
+			return nullptr;
+		}
+	}
+
+	return FirstEntityType;
+}
+
+TWeakObjectPtr<UDMXLibrary> FDMXEntityReferenceCustomization::GetDMXLibrary() const
+{
+	TSharedPtr<IPropertyHandle> LibraryHandle = StructHandle->GetChildHandle(NAME_DMXLibrary);
+	UObject* Object = nullptr;
+	if (LibraryHandle->GetValue(Object) == FPropertyAccess::Success)
+	{
+		return Cast<UDMXLibrary>(Object);
+	}
+	return nullptr;
+}
+
+#undef LOCTEXT_NAMESPACE
diff --git a/Engine/Plugins/Experimental/DigitalMultiplex/DMXEngine/Source/DMXEditor/Private/Customizations/DMXEditorPropertyEditorCustomization.h b/Engine/Plugins/Experimental/DigitalMultiplex/DMXEngine/Source/DMXEditor/Private/Customizations/DMXEditorPropertyEditorCustomization.h
new file mode 100644
--- /dev/null
+++ b/Engine/Plugins/Experimental/DigitalMultiplex/DMXEngine/Source/DMXEditor/Private/Customizations/DMXEditorPropertyEditorCustomization.h
@@ -0,0 +1,370 @@
+// Copyright Epic Games, Inc. All Rights Reserved.
+
+#pragma once
+
+#include "CoreMinimal.h"
+#include "IDetailCustomization.h"
+#include "IPropertyTypeCustomization.h"
+#include "DMXProtocolTypes.h"
+#include "Widgets/SNameListPicker.h"
+#include "DetailWidgetRow.h"
+#include "IPropertyUtilities.h"
+#include "Templates/SubclassOf.h"
+
+class FDMXEditor;
+class UDMXLibrary;
+class UDMXEntity;
+class UDMXEntityFixtureType;
+class UDMXEntityFixturePatch;
+class UDMXComponent;
+class UDMXEntity;
+
+class SEditableTextBox;
+struct EVisibility;
+
+class FDMXCustomization
+	: public IDetailCustomization
+{
+public:
+	/** Constructor */
+	FDMXCustomization(TWeakPtr<FDMXEditor> InDMXEditorPtr)
+		: DMXEditorPtr(InDMXEditorPtr)
+	{}
+
+	/** IDetailCustomization interface */
+	virtual void CustomizeDetails(IDetailLayoutBuilder & DetailLayout) override;
+
+protected:
+	/** Weak reference to the DMX editor */
+	TWeakPtr<FDMXEditor> DMXEditorPtr;
+};
+
+/** Details customization for Controllers */
+class FDMXControllersDetails
+	: public FDMXCustomization
+{
+public:
+	/** Constructor */
+	FDMXControllersDetails(TWeakPtr<FDMXEditor> InDMXEditorPtr)
+		: FDMXCustomization(InDMXEditorPtr)
+	{}
+
+	/** IDetailCustomization interface */
+	virtual void CustomizeDetails(IDetailLayoutBuilder& DetailLayout) override;
+};
+
+/** Base class for Fixture Types' Modes, Functions and Sub Functions customizations */
+class FDMXFixtureTypeFunctionsDetails
+	: public IPropertyTypeCustomization
+{
+public:
+	/** Constructor */
+	FDMXFixtureTypeFunctionsDetails(TWeakPtr<FDMXEditor> InDMXEditorPtr);
+
+	virtual void CustomizeHeader(TSharedRef<IPropertyHandle> InStructPropertyHandle, class FDetailWidgetRow& InHeaderRow, IPropertyTypeCustomizationUtils& InStructCustomizationUtils) override;
+	virtual void CustomizeChildren(TSharedRef<IPropertyHandle> InStructPropertyHandle, IDetailChildrenBuilder& InStructBuilder, IPropertyTypeCustomizationUtils& InStructCustomizationUtils) override;
+
+protected:
+	/** Called to get new Name field settings. */
+	virtual void GetCustomNameFieldSettings(FText& OutNewPropertyLabel, FName& OutNamePropertyName, FText& OutToolTip, FText& OutExistingNameError) = 0;
+
+	/** Find the existing names for the function type being edited within the Fixture Type. */
+	virtual TArray<FString> GetExistingNames() const = 0;
+	
+	/** Allows customization of how properties are added */
+	virtual void AddProperty(IDetailChildrenBuilder& InStructBuilder, const FName& PropertyName, TSharedRef<IPropertyHandle> PropertyHandle);
+
+	void OnFunctionNameChanged(const FText& InNewText);
+	void OnFunctionNameCommitted(const FText& InNewText, ETextCommit::Type InTextCommit);
+	FText OnGetFunctionName() const;
+	/** Changes the function name on the fixture properties */
+	void SetFunctionName(const FString& NewName);
+
+protected:
+	TArray<UDMXEntityFixtureType*> SelectedFixtures;
+	TSharedPtr<IPropertyHandle> NamePropertyHandle;
+	FName NamePropertyName;
+
+	TWeakPtr<FDMXEditor> DMXEditorPtr;
+
+private:
+	TSharedPtr<SEditableTextBox> NameEditableTextBox;
+	FText ExistingNameError;
+};
+
+/** Details customization for Fixture Modes */
+class FDMXFixtureModeDetails
+	: public FDMXFixtureTypeFunctionsDetails
+{
+public:
+	/** Constructor */
+	FDMXFixtureModeDetails(TWeakPtr<FDMXEditor> InDMXEditorPtr)
+		: FDMXFixtureTypeFunctionsDetails(InDMXEditorPtr)
+	{}
+
+	void CustomizeChildren(TSharedRef<IPropertyHandle> InStructPropertyHandle, IDetailChildrenBuilder& InStructBuilder, IPropertyTypeCustomizationUtils& InStructCustomizationUtils) override;
+
+protected:
+	//~ FDMXFixtureTypeFunctionsDetails interface
+	virtual void GetCustomNameFieldSettings(FText& OutNewPropertyLabel, FName& OutNamePropertyName, FText& OutToolTip, FText& OutExistingNameError) override;
+	virtual TArray<FString> GetExistingNames() const override;
+
+private:
+	TSharedPtr<IPropertyHandle> AutoChannelSpanHandle;
+};
+
+/** Details customization for Fixture Mode Functions */
+class FDMXFixtureFunctionDetails
+	: public FDMXFixtureTypeFunctionsDetails
+{
+public:
+	/** Constructor */
+	FDMXFixtureFunctionDetails(TWeakPtr<FDMXEditor> InDMXEditorPtr)
+		: FDMXFixtureTypeFunctionsDetails(InDMXEditorPtr)
+	{}
+
+	virtual void CustomizeChildren(TSharedRef<IPropertyHandle> InStructPropertyHandle, IDetailChildrenBuilder& InStructBuilder, IPropertyTypeCustomizationUtils& InStructCustomizationUtils) override;
+
+protected:
+	//~ FDMXFixtureTypeFunctionsDetails interface
+	virtual void AddProperty(IDetailChildrenBuilder& InStructBuilder, const FName& PropertyName, TSharedRef<IPropertyHandle> PropertyHandle) override;
+	virtual void GetCustomNameFieldSettings(FText& OutNewPropertyLabel, FName& OutNamePropertyName, FText& OutToolTip, FText& OutExistingNameError) override;
+	virtual TArray<FString> GetExistingNames() const override;
+
+private:
+	EVisibility GetSubFunctionsVisibility() const;
+	
+	void AddChannelInputFields(IDetailChildrenBuilder& InStructBuilder);
+	TSharedRef<SWidget> CreateChannelField(uint8 Channel, const FLinearColor& LabelColor = FLinearColor(0.0f, 0.0f, 0.0f, 0.5f));
+	TOptional<uint8> GetChannelValue(uint8 Channel) const;
+	EVisibility GetChannelInputVisibility(uint8 Channel) const;
+
+	void HandleChannelValueChanged(uint8 NewValue, uint8 Channel);
+	void HandleChannelValueCommitted(uint8 NewValue, ETextCommit::Type CommitType);
+
+private:
+	TSharedPtr<IPropertyHandle> StructPropertyHandle;
+	TSharedPtr<IPropertyHandle> DataTypeHandle;
+	TSharedPtr<IPropertyHandle> DefaultValueHandle;
+};
+
+/**
+ * Details customization for Fixture Function's Sub Functions.
+ * The GetExistingNames implementation is very similar to FDMXFixtureFunctionDetails's.
+ * So we only need to override GetCustomNameFieldSettings and the rest just works. :)
+ */
+class FDMXFixtureSubFunctionDetails
+	: public FDMXFixtureFunctionDetails
+{
+public:
+	/** Constructor */
+	FDMXFixtureSubFunctionDetails(TWeakPtr<FDMXEditor> InDMXEditorPtr)
+		: FDMXFixtureFunctionDetails(InDMXEditorPtr)
+	{}
+
+	virtual void CustomizeChildren(TSharedRef<IPropertyHandle> InStructPropertyHandle, IDetailChildrenBuilder& InStructBuilder, IPropertyTypeCustomizationUtils& InStructCustomizationUtils) override;
+
+protected:
+	//~ FDMXFixtureTypeFunctionsDetails interface
+	virtual void GetCustomNameFieldSettings(FText& OutNewPropertyLabel, FName& OutNamePropertyName, FText& OutToolTip, FText& OutExistingNameError) override;
+};
+
+/** Details customization for FixturePatches */
+class FDMXFixturePatchesDetails
+	: public FDMXCustomization
+{
+public:
+	/** Constructor */
+	FDMXFixturePatchesDetails(TWeakPtr<FDMXEditor> InDMXEditorPtr)
+		: FDMXCustomization(InDMXEditorPtr)
+	{}
+
+	/** IDetailCustomization interface */
+	virtual void CustomizeDetails(IDetailLayoutBuilder& DetailLayout) override;
+
+private:
+	/** Fill the ActiveModeOptions array with the modes for the selected patches */
+	void GenerateActiveModeOptions();
+	
+	TWeakObjectPtr<UDMXEntityFixtureType> GetParentFixtureTemplate() const;
+	void OnParentTemplateSelected(UDMXEntity* NewTemplate) const;
+	bool GetParentFixtureTypeIsMultipleValues() const;
+
+	bool GetActiveModeEditable() const;
+	TSharedRef<SWidget> GenerateActiveModeOptionWidget(const TSharedPtr<uint32> InMode) const;
+	void OnActiveModeChanged(const TSharedPtr<uint32> InSelectedMode, ESelectInfo::Type SelectInfo);
+	FText GetCurrentActiveModeLabel() const;
+
+private:
+	TArray< TSharedPtr<uint32> > ActiveModeOptions;
+
+	TSharedPtr<IPropertyHandle> ParentFixtureTypeHandle;
+	TSharedPtr<IPropertyHandle> ActiveModeHandle;
+};
+
+/**  Customization for any property that should be displayed as a dropdown of options from a FName array */
+template<typename TStructType>
+class FNameListCustomization
+	: public IPropertyTypeCustomization
+{
+private:
+	using TNameListType = FNameListCustomization<TStructType>;
+
+public:
+	/** Construction requires a delegate that returns the source list of possible names */
+	FNameListCustomization(TAttribute<TArray<FName>> InOptionsList, FSimpleMulticastDelegate* InUpdateOptionsDelegate)
+		: ListDelegate(InOptionsList)
+		, UpdateOptionsDelegate(InUpdateOptionsDelegate)
+		, StructPropertyHandle(nullptr)
+	{}
+
+	//~ IPropertyTypeCustomization interface begin
+	virtual void CustomizeHeader(TSharedRef<IPropertyHandle> InPropertyHandle, FDetailWidgetRow& InHeaderRow, IPropertyTypeCustomizationUtils& CustomizationUtils) override
+	{
+		StructPropertyHandle = InPropertyHandle;
+		TSharedPtr<IPropertyUtilities> PropertyUtils = CustomizationUtils.GetPropertyUtilities();
+
+		check(CastFieldChecked<FStructProperty>(StructPropertyHandle->GetProperty())->Struct == TStructType::StaticStruct());
+
+		InHeaderRow
+			.NameContent()
+			[
+				StructPropertyHandle->CreatePropertyNameWidget()
+			]
+			.ValueContent()
+			[
+				SNew(SNameListPicker)
+				.Font(CustomizationUtils.GetRegularFont())
+				.HasMultipleValues(this, &TNameListType::HasMultipleValues)
+				.OptionsSource(ListDelegate)
+				.UpdateOptionsDelegate(UpdateOptionsDelegate)
+				.Value(this, &TNameListType::GetValue)
+				.OnValueChanged(this, &TNameListType::SetValue)
+			]
+		.IsEnabled(MakeAttributeLambda([=] { return !InPropertyHandle->IsEditConst() && PropertyUtils->IsPropertyEditingEnabled(); }));
+	}
+
+	virtual void CustomizeChildren(TSharedRef<IPropertyHandle> InPropertyHandle, IDetailChildrenBuilder& InChildBuilder, IPropertyTypeCustomizationUtils& CustomizationUtils) override {};
+	//~ IPropertyTypeCustomization interface end
+
+private:
+	FName GetValue() const
+	{
+		TArray<const void*> RawData;
+		StructPropertyHandle->AccessRawData(RawData);
+
+		for (const void* RawPtr : RawData)
+		{
+			if (RawPtr != nullptr)
+			{
+				return reinterpret_cast<const TStructType*>(RawPtr)->Name;
+			}
+		}
+
+		return FName();
+	}
+
+	void SetValue(FName NewValue)
+	{
+		FStructProperty* StructProperty = CastFieldChecked<FStructProperty>(StructPropertyHandle->GetProperty());
+
+		TArray<void*> RawData;
+		StructPropertyHandle->AccessRawData(RawData);
+		TStructType* PreviousValue = reinterpret_cast<TStructType*>(RawData[0]);
+		TStructType NewProtocolName(NewValue);
+
+		// Export new value to text format that can be imported later
+		FString TextValue;
+		StructProperty->Struct->ExportText(TextValue, &NewProtocolName, PreviousValue, nullptr, EPropertyPortFlags::PPF_None, nullptr);
+
+		// Set values on edited property handle from exported text
+		ensure(StructPropertyHandle->SetValueFromFormattedString(TextValue, EPropertyValueSetFlags::DefaultFlags) == FPropertyAccess::Result::Success);
+	}
+
+	bool HasMultipleValues() const
+	{
+		TArray<const void*> RawData;
+		StructPropertyHandle->AccessRawData(RawData);
+
+		TOptional<TStructType> CompareAgainst;
+		for (const void* RawPtr : RawData)
+		{
+			if (RawPtr == nullptr)
+			{
+				if (CompareAgainst.IsSet())
+				{
+					return false;
+				}
+			}
+			else
+			{
+				const TStructType* ThisValue = reinterpret_cast<const TStructType*>(RawPtr);
+
+				if (!CompareAgainst.IsSet())
+				{
+					CompareAgainst = *ThisValue;
+				}
+				else if (!(*ThisValue == CompareAgainst.GetValue()))
+				{
+					return true;
+				}
+			}
+		}
+
+		return false;
+	}
+
+private:
+	TAttribute< TArray<FName> > ListDelegate;
+	FSimpleMulticastDelegate* UpdateOptionsDelegate;
+	TSharedPtr<IPropertyHandle> StructPropertyHandle;
+};
+
+class FDMXEntityReferenceCustomization
+	: public IPropertyTypeCustomization
+{
+public:
+	//~ IPropertyTypeCustomization interface begin
+	virtual void CustomizeHeader(TSharedRef<IPropertyHandle> InPropertyHandle, FDetailWidgetRow& InHeaderRow, IPropertyTypeCustomizationUtils& CustomizationUtils) override;
+	virtual void CustomizeChildren(TSharedRef<IPropertyHandle> InPropertyHandle, IDetailChildrenBuilder& InChildBuilder, IPropertyTypeCustomizationUtils& CustomizationUtils) override;
+	//~ IPropertyTypeCustomization interface end
+
+private:
+	bool GetDisplayLibrary() const;
+
+	TSharedRef<SWidget> CreateEntityPickerWidget(TSharedPtr<IPropertyHandle> InPropertyHandle) const;
+	FText GetPickerPropertyLabel() const;
+	bool GetPickerEnabled() const;
+
+	TWeakObjectPtr<UDMXEntity> GetCurrentEntity() const;
+	bool GetEntityIsMultipleValues() const;
+	void OnEntitySelected(UDMXEntity* NewEntity) const;
+	TSubclassOf<UDMXEntity> GetEntityType() const;
+	TWeakObjectPtr<UDMXLibrary> GetDMXLibrary() const;
+
+private:
+	static const FName NAME_DMXLibrary;
+
+	TSharedPtr<IPropertyHandle> StructHandle;
+};
+
+struct FDMXCustomizationFactory
+{
+	template<typename TDetailCustomizationType, typename TReturnType>
+	static TSharedRef<TReturnType> MakeInstance(TWeakPtr<FDMXEditor> InEditor)
+	{
+		return MakeShared<TDetailCustomizationType>(InEditor);
+	}
+
+	template<typename TDetailCustomizationType, typename TReturnType>
+	static TSharedRef<TReturnType> MakeInstance()
+	{
+		return MakeShared<TDetailCustomizationType>();
+	}
+
+	template<typename TDetailCustomizationType>
+	static TSharedRef<IPropertyTypeCustomization> MakeInstance(TAttribute<TArray<FName>> InListOptions, FSimpleMulticastDelegate* UpdateListDelegate)
+	{
+		return MakeShared<TDetailCustomizationType>(InListOptions, UpdateListDelegate);
+	}
+};
