- DMX Editor

diff --git a/Engine/Plugins/Experimental/DigitalMultiplex/DMXEngine/Source/DMXEditor/Private/DMXEditor.cpp b/Engine/Plugins/Experimental/DigitalMultiplex/DMXEngine/Source/DMXEditor/Private/DMXEditor.cpp
new file mode 100644
--- /dev/null
+++ b/Engine/Plugins/Experimental/DigitalMultiplex/DMXEngine/Source/DMXEditor/Private/DMXEditor.cpp
@@ -0,0 +1,343 @@
+// Copyright Epic Games, Inc. All Rights Reserved.
+
+#include "DMXEditor.h"
+#include "DMXEditorLog.h"
+#include "DMXEditorModule.h"
+#include "DMXEditorTabs.h"
+#include "DMXEditorUtils.h"
+#include "Library/DMXLibrary.h"
+#include "Library/DMXEntityFader.h"
+#include "Library/DMXEntityController.h"
+#include "Library/DMXEntityFixtureType.h"
+#include "Library/DMXEntityFixturePatch.h"
+#include "Modes/DMXEditorApplicationMode.h"
+#include "Toolbars/DMXEditorToolbar.h"
+#include "Commands/DMXEditorCommands.h"
+
+#include "Widgets/Layout/SScrollBox.h"
+#include "Widgets/Docking/SDockTab.h"
+
+#include "Widgets/SDMXInputConsole.h"
+#include "Widgets/SDMXEntityEditor.h"
+#include "Widgets/SDMXEntityInspector.h"
+#include "Widgets/SDMXOutputConsole.h"
+
+#include "Modules/ModuleManager.h"
+#include "ScopedTransaction.h"
+
+#define LOCTEXT_NAMESPACE "FDMXEditor"
+
+const FName FDMXEditor::ToolkitFName(TEXT("DMXEditor"));
+
+FDMXEditor::FDMXEditor()
+{}
+
+FDMXEditor::~FDMXEditor()
+{
+	InputConsoleWidget.Reset();
+	OutputConsoleWidget.Reset();
+	ControllersWidget.Reset();
+	FixtureTypesWidget.Reset();
+	FixturePatchWidget.Reset();
+}
+
+FName FDMXEditor::GetToolkitFName() const
+{
+	return ToolkitFName;
+}
+
+FText FDMXEditor::GetBaseToolkitName() const
+{
+	return LOCTEXT("DMXEditor", "DMX Editor");
+}
+
+FString FDMXEditor::GetWorldCentricTabPrefix() const
+{
+	return LOCTEXT("WorldCentricTabPrefix_LevelScript", "Script ").ToString();
+}
+
+FLinearColor FDMXEditor::GetWorldCentricTabColorScale() const
+{
+	return FLinearColor(0.0f, 0.0f, 0.3f, 0.5f);
+}
+
+void FDMXEditor::InitEditor(const EToolkitMode::Type Mode, const TSharedPtr<class IToolkitHost>& InitToolkitHost, UDMXLibrary* DMXLibrary)
+{
+	if (!Toolbar.IsValid())
+	{
+		Toolbar = MakeShared<FDMXEditorToolbar>(SharedThis(this));
+	}
+
+	// Initialize the asset editor and spawn nothing (dummy layout)
+	const TSharedRef<FTabManager::FLayout> DummyLayout = FTabManager::NewLayout("NullLayout")->AddArea(FTabManager::NewPrimaryArea());
+	const bool bCreateDefaultStandaloneMenu = true;
+	const bool bCreateDefaultToolbar = true;
+	FAssetEditorToolkit::InitAssetEditor(Mode, InitToolkitHost, FDMXEditorModule::DMXEditorAppIdentifier, DummyLayout, bCreateDefaultStandaloneMenu, bCreateDefaultToolbar, (UObject*)DMXLibrary);
+
+	CommonInitialization(DMXLibrary);
+
+	InitalizeExtenders();
+
+	RegenerateMenusAndToolbars();
+
+	const bool bShouldOpenInDefaultsMode = true;
+	bool bNewlyCreated = true;
+	RegisterApplicationModes(DMXLibrary, bShouldOpenInDefaultsMode, bNewlyCreated);
+}
+
+void FDMXEditor::CommonInitialization(UDMXLibrary* DMXLibrary)
+{
+	CreateDefaultCommands();
+	CreateDefaultTabContents(DMXLibrary);
+}
+
+void FDMXEditor::InitalizeExtenders()
+{
+	FDMXEditorModule* DMXEditorModule = &FDMXEditorModule::Get();
+	TSharedPtr<FExtender> MenuExtender = DMXEditorModule->GetMenuExtensibilityManager()->GetAllExtenders(GetToolkitCommands(), GetEditingObjects());
+	AddMenuExtender(MenuExtender);
+
+	TSharedPtr<FExtender> ToolbarExtender = DMXEditorModule->GetToolBarExtensibilityManager()->GetAllExtenders(GetToolkitCommands(), GetEditingObjects());
+	AddToolbarExtender(ToolbarExtender);
+}
+
+void FDMXEditor::RegisterApplicationModes(UDMXLibrary* DMXLibrary, bool bShouldOpenInDefaultsMode, bool bNewlyCreated)
+{
+	// Only one for now
+	FWorkflowCentricApplication::AddApplicationMode(
+		FDMXEditorApplicationMode::DefaultsMode,
+		MakeShared<FDMXEditorDefaultApplicationMode>(SharedThis(this)));
+	FWorkflowCentricApplication::SetCurrentMode(FDMXEditorApplicationMode::DefaultsMode);
+}
+
+UDMXLibrary * FDMXEditor::GetDMXLibrary() const
+{
+	return GetEditableDMXLibrary();
+}
+
+
+void FDMXEditor::RegisterToolbarTab(const TSharedRef<class FTabManager>& InTabManager)
+{
+	FAssetEditorToolkit::RegisterTabSpawners(InTabManager);
+}
+
+void FDMXEditor::CreateDefaultTabContents(UDMXLibrary* DMXLibrary)
+{
+	InputConsoleWidget = CreateInputConsoleWidget();
+	OutputConsoleWidget = CreateOutputConsoleWidget();
+	ControllersWidget = CreateControllersWidget();
+	FixtureTypesWidget = CreateFixtureTypesWidget();
+	FixturePatchWidget = CreateFixturePatchWidget();
+}
+
+void FDMXEditor::CreateDefaultCommands()
+{
+	FDMXEditorCommands::Register();
+
+	FDMXEditorModule& DMXEditorModule = FModuleManager::LoadModuleChecked<FDMXEditorModule>("DMXEditor");
+	ToolkitCommands->Append(DMXEditorModule.GetsSharedDMXEditorCommands());
+
+	// Entity creation
+	ToolkitCommands->MapAction(
+		FDMXEditorCommands::Get().AddNewEntityController,
+		FExecuteAction::CreateLambda([this]() { OnAddNewEntity(UDMXEntityController::StaticClass()); }),
+		FCanExecuteAction::CreateLambda([this]()->bool { return CanAddNewEntity(UDMXEntityController::StaticClass()); })
+	);
+	ToolkitCommands->MapAction(
+		FDMXEditorCommands::Get().AddNewEntityFixtureType,
+		FExecuteAction::CreateLambda([this]() { OnAddNewEntity(UDMXEntityFixtureType::StaticClass()); }),
+		FCanExecuteAction::CreateLambda([this]()->bool { return CanAddNewEntity(UDMXEntityFixtureType::StaticClass()); })
+	);
+	ToolkitCommands->MapAction(
+		FDMXEditorCommands::Get().AddNewEntityFixturePatch,
+		FExecuteAction::CreateLambda([this]() { OnAddNewEntity(UDMXEntityFixturePatch::StaticClass()); }),
+		FCanExecuteAction::CreateLambda([this]()->bool { return CanAddNewEntity(UDMXEntityFixturePatch::StaticClass()); })
+	);
+}
+
+void FDMXEditor::OnAddNewEntity(TSubclassOf<UDMXEntity> InEntityClass)
+{
+	if (!InvokeEditorTabFromEntityType(InEntityClass))
+	{
+		return; 
+	}
+
+	// Create the new entity with a unique name
+	FString BaseName;
+	OnGetBaseNameForNewEntity.Broadcast(InEntityClass, BaseName);
+	FString EntityName = FDMXEditorUtils::FindUniqueEntityName(GetDMXLibrary(), InEntityClass, BaseName);
+
+	UDMXEntity* NewEntity = nullptr;
+	if (FDMXEditorUtils::AddEntity(GetDMXLibrary(), EntityName, InEntityClass, &NewEntity))
+	{
+		OnSetupNewEntity.Broadcast(NewEntity);
+		RenameNewlyAddedEntity(NewEntity, InEntityClass);
+	}
+	else
+	{
+		UE_LOG_DMXEDITOR(Error, TEXT("Add Entity error!"));
+	}
+}
+
+bool FDMXEditor::InvokeEditorTabFromEntityType(TSubclassOf<UDMXEntity> InEntityClass)
+{
+	// Make sure we're in the right tab for the current type
+	FName TargetTabId = NAME_None;
+	if (InEntityClass->IsChildOf(UDMXEntityController::StaticClass()))
+	{
+		TargetTabId = FDMXEditorTabs::DMXControllersId;
+	}
+	else if (InEntityClass->IsChildOf(UDMXEntityFixtureType::StaticClass()))
+	{
+		TargetTabId = FDMXEditorTabs::DMXFixtureTypesEditorTabId;
+	}
+	else if (InEntityClass->IsChildOf(UDMXEntityFixturePatch::StaticClass()))
+	{
+		TargetTabId = FDMXEditorTabs::DMXFixturePatchEditorTabId;
+	}
+	else
+	{
+		UE_LOG_DMXEDITOR(Error, TEXT("%S: Unimplemented Entity type. Can't set currect Tab."), __FUNCTION__);
+	}
+
+	if (!TargetTabId.IsNone())
+	{
+		FName CurrentTab = FGlobalTabmanager::Get()->GetActiveTab()->GetLayoutIdentifier().TabType;
+		if (!CurrentTab.IsEqual(TargetTabId))
+		{
+			TabManager->InvokeTab(MoveTemp(TargetTabId));
+		}
+		
+		return true;
+	}
+
+	return false;
+}
+
+bool FDMXEditor::CanAddNewEntity(TSubclassOf<UDMXEntity> InEntityClass) const
+{
+	return true;
+}
+
+bool FDMXEditor::NewEntity_IsVisibleForType(TSubclassOf<UDMXEntity> InEntityClass) const
+{
+	return true;
+}
+
+void FDMXEditor::RenameNewlyAddedEntity(UDMXEntity* InEntity, TSubclassOf<UDMXEntity> InEntityClass)
+{
+	TSharedPtr<SDMXEntityEditor> EntityEditor = GetEditorWidgetForEntityType(InEntityClass);
+	if (!EntityEditor.IsValid())
+	{
+		return;
+	}
+
+	// if this check ever fails, something is really wrong!
+	// How can an Entity be created without the button in the editor?!
+	check(EntityEditor.IsValid());
+	
+	EntityEditor->RequestRenameOnNewEntity(InEntity, ESelectInfo::OnMouseClick);
+}
+
+TSharedPtr<SDMXEntityEditor> FDMXEditor::GetEditorWidgetForEntityType(TSubclassOf<UDMXEntity> InEntityClass) const
+{
+	TSharedPtr<SDMXEntityEditor> EntityEditor = nullptr;
+
+	if (InEntityClass->IsChildOf(UDMXEntityController::StaticClass()))
+	{
+		EntityEditor = ControllersWidget;
+	}
+	else if (InEntityClass->IsChildOf(UDMXEntityFixtureType::StaticClass()))
+	{
+		EntityEditor = FixtureTypesWidget;
+	}
+	else if (InEntityClass->IsChildOf(UDMXEntityFixturePatch::StaticClass()))
+	{
+		EntityEditor = FixturePatchWidget;
+	}
+	else
+	{
+		UE_LOG_DMXEDITOR(Error, TEXT("%S not implemented for %s"), __FUNCTION__, *InEntityClass->GetFName().ToString());
+	}
+
+	return EntityEditor;
+}
+
+void FDMXEditor::SelectEntityInItsTypeTab(UDMXEntity* InEntity, ESelectInfo::Type InSelectionType /*= ESelectInfo::Type::Direct*/)
+{
+	check(InEntity != nullptr);
+
+	if (!InvokeEditorTabFromEntityType(InEntity->GetClass()))
+	{
+		return;
+	}
+
+	if (TSharedPtr<SDMXEntityEditor> EntityEditor = GetEditorWidgetForEntityType(InEntity->GetClass()))
+	{
+		EntityEditor->SelectEntity(InEntity, InSelectionType);
+	}
+}
+
+void FDMXEditor::SelectEntitiesInTypeTab(const TArray<UDMXEntity*>& InEntities, ESelectInfo::Type InSelectionType /*= ESelectInfo::Type::Direct*/)
+{
+	if (InEntities.Num() == 0 || InEntities[0] == nullptr)
+	{
+		return; 
+	}
+
+	if (!InvokeEditorTabFromEntityType(InEntities[0]->GetClass()))
+	{
+		return;
+	}
+
+	if (TSharedPtr<SDMXEntityEditor> EntityEditor = GetEditorWidgetForEntityType(InEntities[0]->GetClass()))
+	{
+		EntityEditor->SelectEntities(InEntities, InSelectionType);
+	}
+}
+
+TArray<UDMXEntity*> FDMXEditor::GetSelectedEntitiesFromTypeTab(TSubclassOf<UDMXEntity> InEntityClass) const
+{
+	if (TSharedPtr<SDMXEntityEditor> EntityEditor = GetEditorWidgetForEntityType(InEntityClass))
+	{
+		return EntityEditor->GetSelectedEntities();
+	}
+
+	return TArray<UDMXEntity*>();
+}
+
+UDMXLibrary * FDMXEditor::GetEditableDMXLibrary() const
+{
+	return Cast<UDMXLibrary>(GetEditingObject());
+}
+
+TSharedRef<SDMXInputConsole> FDMXEditor::CreateInputConsoleWidget()
+{
+	return SNew(SDMXInputConsole);
+}
+
+TSharedRef<SWidget> FDMXEditor::CreateOutputConsoleWidget()
+{
+	return SNew(SDMXOutputConsole)
+		.DMXEditor(SharedThis(this));
+}
+
+TSharedRef<SDMXControllers> FDMXEditor::CreateControllersWidget()
+{
+	return SNew(SDMXControllers)
+		.DMXEditor(SharedThis(this));
+}
+
+TSharedRef<SDMXFixtureTypes> FDMXEditor::CreateFixtureTypesWidget()
+{
+	return SNew(SDMXFixtureTypes)
+		.DMXEditor(SharedThis(this));
+}
+
+TSharedRef<SDMXFixturePatch> FDMXEditor::CreateFixturePatchWidget()
+{
+	return SNew(SDMXFixturePatch)
+		.DMXEditor(SharedThis(this));
+}
+
+#undef LOCTEXT_NAMESPACE
diff --git a/Engine/Plugins/Experimental/DigitalMultiplex/DMXEngine/Source/DMXEditor/Private/DMXEditorModule.cpp b/Engine/Plugins/Experimental/DigitalMultiplex/DMXEngine/Source/DMXEditor/Private/DMXEditorModule.cpp
new file mode 100644
--- /dev/null
+++ b/Engine/Plugins/Experimental/DigitalMultiplex/DMXEngine/Source/DMXEditor/Private/DMXEditorModule.cpp
@@ -0,0 +1,155 @@
+// Copyright Epic Games, Inc. All Rights Reserved.
+
+#include "DMXEditorModule.h"
+#include "Library/DMXLibrary.h"
+#include "Library/DMXEntityReference.h"
+#include "Library/DMXEntity.h"
+#include "Game/DMXComponent.h"
+#include "DMXEditor.h"
+#include "DMXEditorStyle.h"
+#include "DMXProtocolTypes.h"
+#include "AssetTools/AssetTypeActions_DMXEditorLibrary.h"
+#include "Customizations/DMXEditorPropertyEditorCustomization.h"
+
+#include "AssetToolsModule.h"
+#include "PropertyEditorModule.h"
+#include "AssetRegistryModule.h"
+#include "AssetToolsModule.h"
+
+const FName FDMXEditorModule::DMXEditorAppIdentifier(TEXT("DMXEditorApp"));
+EAssetTypeCategories::Type FDMXEditorModule::DMXEditorAssetCategory;
+
+const FName FDMXEditorModule::ModuleName = TEXT("DMXEditor");
+
+#define LOCTEXT_NAMESPACE "DMXEditorModule"
+
+void FDMXEditorModule::StartupModule()
+{
+    static const FName AssetRegistryName(TEXT("AssetRegistry"));
+    static const FName AssetToolsModuleame(TEXT("AssetTools"));
+
+	FPropertyEditorModule& PropertyModule = FModuleManager::GetModuleChecked<FPropertyEditorModule>("PropertyEditor");
+
+    // Make sure AssetRegistry is loaded and get reference to the module
+    FModuleManager::LoadModuleChecked<FAssetRegistryModule>(AssetRegistryName);
+    FModuleManager::LoadModuleChecked<FAssetToolsModule>(AssetToolsModuleame);
+
+    FDMXEditorStyle::Initialize();
+
+	MenuExtensibilityManager = MakeShared<FExtensibilityManager>();
+	ToolBarExtensibilityManager = MakeShared<FExtensibilityManager>();
+	SharedDMXEditorCommands = MakeShared<FUICommandList>();
+
+	RegisterPropertyTypeCustomizations();
+	RegisterObjectCustomizations();
+
+	IAssetTools& AssetTools = FModuleManager::LoadModuleChecked<FAssetToolsModule>("AssetTools").Get();
+	DMXEditorAssetCategory = AssetTools.RegisterAdvancedAssetCategory(FName(TEXT("DMX")), LOCTEXT("DmxCategory", "DMX"));
+	RegisterAssetTypeAction(AssetTools, MakeShared<FAssetTypeActions_DMXEditorLibrary>());
+
+	PropertyModule.NotifyCustomizationModuleChanged();
+}
+
+void FDMXEditorModule::ShutdownModule()
+{
+	FDMXEditorStyle::Shutdown();
+
+	MenuExtensibilityManager.Reset();
+	ToolBarExtensibilityManager.Reset();
+	SharedDMXEditorCommands.Reset();
+
+	if (FModuleManager::Get().IsModuleLoaded("AssetTools"))
+	{
+		IAssetTools& AssetTools = FModuleManager::GetModuleChecked<FAssetToolsModule>("AssetTools").Get();
+		for (TSharedPtr<IAssetTypeActions>& AssetIt : CreatedAssetTypeActions)
+		{
+			AssetTools.UnregisterAssetTypeActions(AssetIt.ToSharedRef());
+		}
+	}
+	CreatedAssetTypeActions.Empty();
+
+
+	FCoreUObjectDelegates::GetPreGarbageCollectDelegate().RemoveAll(this);
+	FCoreDelegates::OnPostEngineInit.RemoveAll(this);
+
+
+	if (FModuleManager::Get().IsModuleLoaded("PropertyEditor"))
+	{
+		FPropertyEditorModule& PropertyModule = FModuleManager::GetModuleChecked<FPropertyEditorModule>("PropertyEditor");
+
+		// Unregister all classes customized by name
+		for (const FName& ClassName : RegisteredClassNames)
+		{
+			PropertyModule.UnregisterCustomClassLayout(ClassName);
+		}
+
+		// Unregister all structures
+		for (const FName& PropertyTypeName : RegisteredPropertyTypes)
+		{
+			PropertyModule.UnregisterCustomPropertyTypeLayout(PropertyTypeName);
+		}
+
+		PropertyModule.NotifyCustomizationModuleChanged();
+	}
+}
+
+void FDMXEditorModule::RegisterCustomClassLayout(FName ClassName, FOnGetDetailCustomizationInstance DetailLayoutDelegate)
+{
+	check(ClassName != NAME_None);
+
+	RegisteredClassNames.Add(ClassName);
+
+	static FName PropertyEditor("PropertyEditor");
+	FPropertyEditorModule& PropertyModule = FModuleManager::GetModuleChecked<FPropertyEditorModule>(PropertyEditor);
+	PropertyModule.RegisterCustomClassLayout(ClassName, DetailLayoutDelegate);
+}
+
+
+void FDMXEditorModule::RegisterCustomPropertyTypeLayout(FName PropertyTypeName, FOnGetPropertyTypeCustomizationInstance PropertyTypeLayoutDelegate)
+{
+	check(PropertyTypeName != NAME_None);
+
+	RegisteredPropertyTypes.Add(PropertyTypeName);
+
+	static FName PropertyEditor("PropertyEditor");
+	FPropertyEditorModule& PropertyModule = FModuleManager::GetModuleChecked<FPropertyEditorModule>(PropertyEditor);
+	PropertyModule.RegisterCustomPropertyTypeLayout(PropertyTypeName, PropertyTypeLayoutDelegate);
+}
+
+void FDMXEditorModule::RegisterPropertyTypeCustomizations()
+{
+	RegisterCustomPropertyTypeLayout(FDMXProtocolName::StaticStruct()->GetFName(), FOnGetPropertyTypeCustomizationInstance::CreateStatic(&FDMXCustomizationFactory::MakeInstance<FNameListCustomization<FDMXProtocolName>>, MakeAttributeLambda(&FDMXProtocolName::GetPossibleValues), (FSimpleMulticastDelegate*)nullptr));
+	RegisterCustomPropertyTypeLayout(FDMXFixtureCategory::StaticStruct()->GetFName(), FOnGetPropertyTypeCustomizationInstance::CreateStatic(&FDMXCustomizationFactory::MakeInstance<FNameListCustomization<FDMXFixtureCategory>>, MakeAttributeLambda(&FDMXFixtureCategory::GetPossibleValues), &FDMXFixtureCategory::OnPossibleValuesUpdated));
+
+	// Customizations for the Entity Reference types
+	RegisterCustomPropertyTypeLayout(FDMXEntityControllerRef::StaticStruct()->GetFName(), FOnGetPropertyTypeCustomizationInstance::CreateStatic(&FDMXCustomizationFactory::MakeInstance<FDMXEntityReferenceCustomization>));
+	RegisterCustomPropertyTypeLayout(FDMXEntityFixtureTypeRef::StaticStruct()->GetFName(), FOnGetPropertyTypeCustomizationInstance::CreateStatic(&FDMXCustomizationFactory::MakeInstance<FDMXEntityReferenceCustomization>));
+	RegisterCustomPropertyTypeLayout(FDMXEntityFixturePatchRef::StaticStruct()->GetFName(), FOnGetPropertyTypeCustomizationInstance::CreateStatic(&FDMXCustomizationFactory::MakeInstance<FDMXEntityReferenceCustomization>));
+}
+
+void FDMXEditorModule::RegisterObjectCustomizations()
+{
+}
+
+void FDMXEditorModule::RegisterAssetTypeAction(IAssetTools& AssetTools, TSharedRef<IAssetTypeActions> Action)
+{
+	AssetTools.RegisterAssetTypeActions(Action);
+	CreatedAssetTypeActions.Add(Action);
+}
+
+FDMXEditorModule& FDMXEditorModule::Get()
+{
+	return FModuleManager::GetModuleChecked<FDMXEditorModule>(ModuleName);
+}
+
+TSharedRef<FDMXEditor> FDMXEditorModule::CreateEditor(const EToolkitMode::Type Mode, const TSharedPtr<class IToolkitHost>& InitToolkitHost, UDMXLibrary * DMXLibrary)
+{
+	TSharedRef<FDMXEditor> DMXEditor(new FDMXEditor());
+	DMXEditor->InitEditor(Mode, InitToolkitHost, DMXLibrary);
+
+	return DMXEditor;
+}
+
+IMPLEMENT_MODULE(FDMXEditorModule, DMXEditor)
+
+#undef LOCTEXT_NAMESPACE
diff --git a/Engine/Plugins/Experimental/DigitalMultiplex/DMXEngine/Source/DMXEditor/Private/DMXEditorTabs.cpp b/Engine/Plugins/Experimental/DigitalMultiplex/DMXEngine/Source/DMXEditor/Private/DMXEditorTabs.cpp
new file mode 100644
--- /dev/null
+++ b/Engine/Plugins/Experimental/DigitalMultiplex/DMXEngine/Source/DMXEditor/Private/DMXEditorTabs.cpp
@@ -0,0 +1,9 @@
+// Copyright Epic Games, Inc. All Rights Reserved.
+
+#include "DMXEditorTabs.h"
+
+const FName FDMXEditorTabs::DMXControllersId(TEXT("GenericEditor_Properties"));
+const FName FDMXEditorTabs::DMXFixtureTypesEditorTabId("GenericEditor_DMXFixtureTypesEditor");
+const FName FDMXEditorTabs::DMXFixturePatchEditorTabId("GenericEditor_DMXFixturePatchEditor");
+const FName FDMXEditorTabs::DMXInputConsoleEditorTabId("GenericEditor_DMXInputConsole");
+const FName FDMXEditorTabs::DMXOutputConsoleEditorTabId("GenericEditor_DMXOutputConsole");
diff --git a/Engine/Plugins/Experimental/DigitalMultiplex/DMXEngine/Source/DMXEditor/Private/DMXEditorTabs.h b/Engine/Plugins/Experimental/DigitalMultiplex/DMXEngine/Source/DMXEditor/Private/DMXEditorTabs.h
new file mode 100644
--- /dev/null
+++ b/Engine/Plugins/Experimental/DigitalMultiplex/DMXEngine/Source/DMXEditor/Private/DMXEditorTabs.h
@@ -0,0 +1,28 @@
+// Copyright Epic Games, Inc. All Rights Reserved.
+
+#pragma once
+
+#include "CoreMinimal.h"
+
+/**  DMX editor tabs identifiers */
+struct FDMXEditorTabs
+{
+	/**	The tab id for the dmx constollers tab */
+	static const FName DMXControllersId;
+
+	/**	The tab id for the dmx fixture types tab */
+	static const FName DMXFixtureTypesEditorTabId;
+
+	/**	The tab id for the dmx fixture patch tab */
+	static const FName DMXFixturePatchEditorTabId;
+
+	/**	The tab id for the dmx monitor tab */
+	static const FName DMXInputConsoleEditorTabId;
+
+	/**	The tab id for the dmx output console tab */
+	static const FName DMXOutputConsoleEditorTabId;
+
+	// Disable default constructor
+	FDMXEditorTabs() = delete;
+};
+
diff --git a/Engine/Plugins/Experimental/DigitalMultiplex/DMXEngine/Source/DMXEditor/Private/DMXEditorUtils.cpp b/Engine/Plugins/Experimental/DigitalMultiplex/DMXEngine/Source/DMXEditor/Private/DMXEditorUtils.cpp
new file mode 100644
--- /dev/null
+++ b/Engine/Plugins/Experimental/DigitalMultiplex/DMXEngine/Source/DMXEditor/Private/DMXEditorUtils.cpp
@@ -0,0 +1,595 @@
+// Copyright Epic Games, Inc. All Rights Reserved.
+
+#include "DMXEditorUtils.h"
+#include "Library/DMXLibrary.h"
+#include "Library/DMXEntityFader.h"
+#include "Library/DMXEntityController.h"
+#include "Library/DMXEntityFixtureType.h"
+#include "Library/DMXEntityFixturePatch.h"
+
+#include "ScopedTransaction.h"
+#include "UnrealExporter.h"
+#include "Exporters/Exporter.h"
+#include "HAL/PlatformApplicationMisc.h"
+#include "Factories.h"
+#include "UObject/Package.h"
+
+#define LOCTEXT_NAMESPACE "FDMXEditorUtils"
+// In blueprints name verification, it is said that '.' is known for causing problems
+#define DMX_INVALID_NAME_CHARACTERS TEXT(".")
+
+// Text object factory for pasting DMX Entities
+struct FDMXEntityObjectTextFactory : public FCustomizableTextObjectFactory
+{
+	/** Entities instantiated */
+	TArray<UDMXEntity*> NewEntities;
+
+	/** Constructs a new object factory from the given text buffer */
+	static TSharedRef<FDMXEntityObjectTextFactory> Get(const FString& InTextBuffer)
+	{
+		// Construct a new instance
+		TSharedPtr<FDMXEntityObjectTextFactory> FactoryPtr = MakeShareable(new FDMXEntityObjectTextFactory());
+		check(FactoryPtr.IsValid());
+
+		// Create new objects if we're allowed to
+		if (FactoryPtr->CanCreateObjectsFromText(InTextBuffer))
+		{
+			EObjectFlags ObjectFlags = RF_Transactional;
+
+			// Use the transient package initially for creating the objects, since the variable name is used when copying
+			FactoryPtr->ProcessBuffer(GetTransientPackage(), ObjectFlags, InTextBuffer);
+		}
+
+		return FactoryPtr.ToSharedRef();
+	}
+
+protected:
+	/** Constructor; protected to only allow this class to instance itself */
+	FDMXEntityObjectTextFactory()
+		: FCustomizableTextObjectFactory(GWarn)
+	{
+	}
+
+	//~ Begin FCustomizableTextObjectFactory implementation
+	virtual bool CanCreateClass(UClass* ObjectClass, bool& bOmitSubObjs) const override
+	{
+		// Allow DMX Entity types to be created
+		return ObjectClass->IsChildOf(UDMXEntity::StaticClass());
+	}
+
+	virtual void ProcessConstructedObject(UObject* NewObject) override
+	{
+		check(NewObject);
+
+		if (UDMXEntity* NewEntity = Cast<UDMXEntity>(NewObject))
+		{
+			// If this is a Fixture Type and the first object was a Patch, we
+			// don't add the Type to the array. SDMXEntityList will deal with it
+			// from the reference on the Patch(es)
+			if (NewEntities.Num() > 0 && NewEntities[0]->GetClass()->IsChildOf<UDMXEntityFixturePatch>()
+				&& NewEntity->GetClass()->IsChildOf<UDMXEntityFixtureType>())
+			{
+				return;
+			}
+
+			NewEntities.Add(NewEntity);
+		}
+	}
+	//~ End FCustomizableTextObjectFactory implementation
+};
+
+
+bool FDMXEditorUtils::GetNameAndIndexFromString(const FString& InString, FString& OutName, int32& OutIndex)
+{
+	OutName = InString.TrimEnd();
+
+	// If there's an index at the end of the name, erase it
+	int32 DigitIndex = OutName.Len();
+	while (DigitIndex > 0 && OutName[DigitIndex - 1] >= '0' && OutName[DigitIndex - 1] <= '9')
+	{
+		--DigitIndex;
+	}
+
+	bool bHadIndex = false;
+	if (DigitIndex < OutName.Len() && DigitIndex > -1)
+	{
+		OutIndex = FCString::Atoi(*OutName.RightChop(DigitIndex));
+		OutName = OutName.Left(DigitIndex);
+		bHadIndex = true;
+	}
+	else
+	{
+		OutIndex = 0;
+	}
+
+	if (OutName.EndsWith(TEXT("_")))
+	{
+		OutName = OutName.LeftChop(1);
+	}
+	OutName.TrimEnd();
+
+	return bHadIndex;
+}
+
+FString FDMXEditorUtils::GenerateUniqueNameFromExisting(const TSet<FString>& InExistingNames, const FString& InBaseName)
+{
+	if (!InBaseName.IsEmpty() && !InExistingNames.Contains(InBaseName))
+	{
+		return InBaseName;
+	}
+
+	FString FinalName;
+	FString BaseName;
+
+	if (InBaseName.IsEmpty())
+	{
+		BaseName = TEXT("Default name");
+	}
+	else
+	{
+		// If there's an index at the end of the name, erase it
+		int32 Index = 0;
+		GetNameAndIndexFromString(InBaseName, BaseName, Index);
+	}
+
+	int32 Count = 1;
+	FinalName = BaseName;
+	// Add Count to the BaseName, increasing Count, until it's a non-existent name
+	do
+	{
+		// Calculate the number of digits in the number, adding 2 (1 extra to correctly count digits, another to account for the '_' that will be added to the name
+		int32 CountLength = Count > 0 ? (int32)FGenericPlatformMath::LogX(10.0f, Count) + 2 : 2;
+
+		// If the length of the final string will be too long, cut off the end so we can fit the number
+		if (CountLength + BaseName.Len() > NAME_SIZE)
+		{
+			BaseName = BaseName.Left(NAME_SIZE - CountLength);
+		}
+
+		FinalName = FString::Printf(TEXT("%s_%d"), *BaseName, Count);
+		++Count;
+	} while (InExistingNames.Contains(FinalName));
+
+	return FinalName;
+}
+
+FString FDMXEditorUtils::FindUniqueEntityName(const UDMXLibrary* InLibrary, TSubclassOf<UDMXEntity> InEntityClass, const FString& InBaseName /*= TEXT("")*/)
+{
+	check(InLibrary != nullptr);
+
+	// Get existing names for the current entity type
+	TSet<FString> EntityNames;
+	InLibrary->ForEachEntityOfType(InEntityClass, [&EntityNames](UDMXEntity* Entity)
+		{
+			EntityNames.Add(Entity->GetDisplayName());
+		});
+
+	FString BaseName = InBaseName;
+
+	// If no base name was set, use the entity class name as base
+	if (BaseName.IsEmpty())
+	{
+		BaseName = InEntityClass->GetDisplayNameText().ToString();
+	}
+
+	return GenerateUniqueNameFromExisting(EntityNames, BaseName);
+}
+
+void FDMXEditorUtils::SetNewFixtureFunctionsNames(UDMXEntityFixtureType* InFixtureType)
+{
+	check(InFixtureType != nullptr);
+
+	// We'll only populate this Set if we find an item with no name.
+	// Otherwise we can save some for loops.
+	TSet<FString> ModesNames;
+
+	// Iterate over all of the Fixture's Modes, Functions and Sub Functions, creating names for the
+	// ones with a blank name.
+	for (FDMXFixtureMode& Mode : InFixtureType->Modes)
+	{
+		// Do we need to name this mode?
+		if (Mode.ModeName.IsEmpty())
+		{
+			// Cache existing names only once, when needed.
+			if (ModesNames.Num() == 0)
+			{
+				// Cache the existing modes' names
+				for (FDMXFixtureMode& NamedMode : InFixtureType->Modes)
+				{
+					if (!NamedMode.ModeName.IsEmpty())
+					{
+						ModesNames.Add(NamedMode.ModeName);
+					}
+				}
+			}
+
+			Mode.ModeName = GenerateUniqueNameFromExisting(ModesNames, TEXT("Mode"));
+			ModesNames.Add(Mode.ModeName);
+		}
+
+		// Name this mode's functions
+		TSet<FString> FunctionsNames;
+		for (FDMXFixtureFunction& Function : Mode.Functions)
+		{
+			if (Function.FunctionName.IsEmpty())
+			{
+				// Cache existing names only once, when needed.
+				if (FunctionsNames.Num() == 0)
+				{
+					for (FDMXFixtureFunction& NamedFunction : Mode.Functions)
+					{
+						if (!NamedFunction.FunctionName.IsEmpty())
+						{
+							FunctionsNames.Add(NamedFunction.FunctionName);
+						}
+					}
+				}
+
+				Function.FunctionName = GenerateUniqueNameFromExisting(FunctionsNames, TEXT("Function"));
+				FunctionsNames.Add(Function.FunctionName);
+			}
+
+			// Name this function's Sub Functions
+			TSet<FString> SubFunctionsNames;
+			for (FDMXFixtureSubFunction& SubFunction : Function.SubFunctions)
+			{
+				if (SubFunction.FunctionName.IsEmpty())
+				{
+					// Cache existing names only once, when needed.
+					if (SubFunctionsNames.Num() == 0)
+					{
+						for (FDMXFixtureSubFunction& NamedSubFunction : Function.SubFunctions)
+						{
+							if (!NamedSubFunction.FunctionName.IsEmpty())
+							{
+								SubFunctionsNames.Add(NamedSubFunction.FunctionName);
+							}
+						}
+					}
+
+					SubFunction.FunctionName = GenerateUniqueNameFromExisting(SubFunctionsNames, TEXT("SubFunction"));
+					SubFunctionsNames.Add(SubFunction.FunctionName);
+				}
+			}
+		}
+	}
+}
+
+bool FDMXEditorUtils::AddEntity(UDMXLibrary* InLibrary, const FString& NewEntityName, TSubclassOf<UDMXEntity> NewEntityClass, UDMXEntity** OutNewEntity /*= nullptr*/)
+{
+	// Don't allow entities with empty names
+	if (NewEntityName.IsEmpty())
+	{
+		return false;
+	}
+
+	// Mark library as pending save and store current state for undo
+	const FScopedTransaction NewEntityTransaction(LOCTEXT("NewEntityTransaction", "Add new Entity to DMX Library"));
+	InLibrary->Modify();
+
+	// Create new entity
+	*OutNewEntity = InLibrary->GetOrCreateEntityObject(NewEntityName, NewEntityClass);
+
+	return true;
+}
+
+bool FDMXEditorUtils::ValidateEntityName(const FString& NewEntityName, const UDMXLibrary* InLibrary, UClass* InEntityClass, FText& OutReason)
+{
+	if (NewEntityName.Len() > NAME_SIZE)
+	{
+		OutReason = LOCTEXT("NameTooLong", "The name is too long");
+		return false;
+	}
+
+	for (const TCHAR& Character : DMX_INVALID_NAME_CHARACTERS)
+	{
+		if (NewEntityName.Contains(&Character))
+		{
+			OutReason = FText::Format(LOCTEXT("NameWithInvalidCharacters", "Name can not contain: {0}"),
+				FText::FromString(&Character));
+			return false;
+		}
+	}
+
+	// Check against existing names for the current entity type
+	bool bNameIsUsed = false;
+	InLibrary->ForEachEntityOfTypeWithBreak(InEntityClass, [&bNameIsUsed, &NewEntityName](UDMXEntity* Entity)
+		{
+			if (Entity->GetDisplayName() == NewEntityName)
+			{
+				bNameIsUsed = true;
+				return false; // Break the loop
+			}
+			return true; // Keep checking Entities' names
+		});
+
+	if (bNameIsUsed)
+	{
+		OutReason = LOCTEXT("ExistingEntityName", "Name already exists");
+		return false;
+	}
+	else
+	{
+		OutReason = FText::GetEmpty();
+		return true;
+	}
+}
+
+UDMXEntityFader* FDMXEditorUtils::CreateFaderTemplate(const UDMXLibrary* InLibrary)
+{
+	FString BaseName(TEXT(""));
+	FString EntityName = FindUniqueEntityName(InLibrary, UDMXEntityFader::StaticClass(), BaseName);
+
+	UDMXEntityFader* OutputConsoleFaderTemplate = NewObject<UDMXEntityFader>(GetTransientPackage(), UDMXEntityFader::StaticClass(), NAME_None, RF_Transient);
+	OutputConsoleFaderTemplate->SetName(EntityName);
+
+	return OutputConsoleFaderTemplate;
+}
+
+void FDMXEditorUtils::RenameEntity(UDMXLibrary* InLibrary, UDMXEntity* InEntity, const FString& NewName)
+{
+	if (InEntity == nullptr)
+	{
+		return;
+	}
+
+	if (!NewName.IsEmpty() && !NewName.Equals(InEntity->GetDisplayName()))
+	{
+		const FScopedTransaction Transaction(LOCTEXT("RenameEntity", "Rename Entity"));
+		InEntity->Modify();
+
+		// Update the name
+		InEntity->SetName(NewName);
+	}
+}
+
+bool FDMXEditorUtils::IsEntityUsed(const UDMXLibrary* InLibrary, const UDMXEntity* InEntity)
+{
+	if (InLibrary != nullptr && InEntity != nullptr)
+	{
+		if (const UDMXEntityFixtureType* EntityAsFixtureType = Cast<UDMXEntityFixtureType>(InEntity))
+		{
+			bool bIsUsed = false;
+			InLibrary->ForEachEntityOfTypeWithBreak<UDMXEntityFixturePatch>([&](UDMXEntityFixturePatch* Patch)
+				{
+					if (Patch->ParentFixtureTypeTemplate == InEntity)
+					{
+						bIsUsed = true;
+						return false;
+					}
+					return true;
+				});
+
+			return bIsUsed;
+		}
+		else
+		{
+			return false;
+		}
+	}
+
+	return false;
+}
+
+void FDMXEditorUtils::RemoveEntities(UDMXLibrary* InLibrary, const TArray<UDMXEntity*>&& InEntities)
+{
+	if (InLibrary != nullptr)
+	{
+		for (UDMXEntity* EntityToDelete : InEntities)
+		{
+			// Fix references to this Entity
+			if (UDMXEntityFixtureType* AsFixtureType = Cast<UDMXEntityFixtureType>(EntityToDelete))
+			{
+				// Find Fixture Patches using this Fixture Type and null their templates
+				InLibrary->ForEachEntityOfType<UDMXEntityFixturePatch>([&AsFixtureType](UDMXEntityFixturePatch* Patch)
+					{
+						if (Patch->ParentFixtureTypeTemplate == AsFixtureType)
+						{
+							Patch->Modify();
+							Patch->ParentFixtureTypeTemplate = nullptr;
+						}
+					});
+			}
+
+			InLibrary->Modify();
+			InLibrary->RemoveEntity(EntityToDelete);
+		}
+	}
+}
+
+void FDMXEditorUtils::CopyEntities(const TArray<UDMXEntity*>&& EntitiesToCopy)
+{
+	// Clear the mark state for saving.
+	UnMarkAllObjects(EObjectMark(OBJECTMARK_TagExp | OBJECTMARK_TagImp));
+
+	const FExportObjectInnerContext Context;
+	FStringOutputDevice Archive;
+
+	// Stores duplicates of the Fixture Type Templates because they can't be parsed being children of
+	// a DMX Library asset since they're private objects.
+	TMap<FName, UDMXEntityFixtureType*> CopiedPatchTemplates;
+
+	// Export the component object(s) to text for copying
+	for (UDMXEntity* Entity : EntitiesToCopy)
+	{
+		check(Entity && Entity->GetParentLibrary());
+
+		// Fixture Patches require copying their template because it's a reference to a private object
+		if (UDMXEntityFixturePatch* AsPatch = Cast<UDMXEntityFixturePatch>(Entity))
+		{
+			if (AsPatch->ParentFixtureTypeTemplate != nullptr)
+			{
+				bool bExportedTemplate = false;
+
+				// Try to get a cached duplicate of the template
+				UDMXEntityFixtureType* DuplicateFixtureType;
+				if (UDMXEntityFixtureType** CachedTemplate = CopiedPatchTemplates.Find(AsPatch->ParentFixtureTypeTemplate->GetFName()))
+				{
+					DuplicateFixtureType = *CachedTemplate;
+					bExportedTemplate = true;
+				}
+				else
+				{
+					// Copy the template to the transient package to make the Patch reference the copy
+					FObjectDuplicationParameters DuplicationParams(AsPatch->ParentFixtureTypeTemplate, GetTransientPackage());
+					DuplicationParams.DestName = AsPatch->ParentFixtureTypeTemplate->GetFName();
+
+					DuplicateFixtureType = CastChecked<UDMXEntityFixtureType>(StaticDuplicateObjectEx(DuplicationParams));
+					// Keep same entity ID to find the original Template when pasting
+					DuplicateFixtureType->ReplicateID(AsPatch->ParentFixtureTypeTemplate);
+
+					// Cache this copy so we don't copy the same template over and over for several Patches
+					CopiedPatchTemplates.Add(DuplicateFixtureType->GetFName(), DuplicateFixtureType);
+				}
+
+				// We'll temporarily change the ParentFixtureTypeTemplate of the Patch to copy it
+				// with a reference to the duplicate Fixture Type
+				UDMXEntityFixtureType* OriginalTemplate = AsPatch->ParentFixtureTypeTemplate;
+				AsPatch->ParentFixtureTypeTemplate = DuplicateFixtureType;
+				// Export the Patch referencing the duplicate template
+				UExporter::ExportToOutputDevice(&Context, AsPatch, nullptr, Archive, TEXT("copy"), 0, PPF_ExportsNotFullyQualified | PPF_Copy | PPF_Delimited, false, GetTransientPackage());
+				if (!bExportedTemplate)
+				{
+					// Export the template after the Patch, to make interpretation easier when pasting it back
+					UExporter::ExportToOutputDevice(&Context, DuplicateFixtureType, nullptr, Archive, TEXT("copy"), 4, PPF_ExportsNotFullyQualified | PPF_Copy | PPF_Delimited, false, GetTransientPackage());
+				}
+				// Revert the patch to it's original, private template
+				AsPatch->ParentFixtureTypeTemplate = OriginalTemplate;
+			}
+			else // Template is null
+			{
+				// Export the entity object to the given string
+				UExporter::ExportToOutputDevice(&Context, AsPatch, nullptr, Archive, TEXT("copy"), 0, PPF_ExportsNotFullyQualified | PPF_Copy | PPF_Delimited, false, GetTransientPackage());
+			}
+		}
+		else
+		{
+			// Export the entity object to the given string
+			UExporter::ExportToOutputDevice(&Context, Entity, nullptr, Archive, TEXT("copy"), 0, PPF_ExportsNotFullyQualified | PPF_Copy | PPF_Delimited, false, GetTransientPackage());
+		}
+	}
+
+	// Speed up the deletion of the copies. We don't need them anymore
+	for (const TPair<FName, UDMXEntityFixtureType*>& CopiedTemplate : CopiedPatchTemplates)
+	{
+		CopiedTemplate.Value->ConditionalBeginDestroy();
+	}
+
+	// Copy text to clipboard
+	FString ExportedText = Archive;
+	FPlatformApplicationMisc::ClipboardCopy(*ExportedText);
+}
+
+bool FDMXEditorUtils::CanPasteEntities()
+{
+	FString ClipboardContent;
+	FPlatformApplicationMisc::ClipboardPaste(ClipboardContent);
+
+	// Obtain the entity object text factory for the clipboard content and return whether or not we can use it
+	TSharedRef<FDMXEntityObjectTextFactory> Factory = FDMXEntityObjectTextFactory::Get(ClipboardContent);
+	return Factory->NewEntities.Num() > 0;
+}
+
+void FDMXEditorUtils::GetEntitiesFromClipboard(TArray<UDMXEntity*>& OutNewObjects)
+{
+	// Get the text from the clipboard
+	FString TextToImport;
+	FPlatformApplicationMisc::ClipboardPaste(TextToImport);
+
+	// Get a new component object factory for the clipboard content
+	TSharedRef<FDMXEntityObjectTextFactory> Factory = FDMXEntityObjectTextFactory::Get(TextToImport);
+
+	// Return the created component mappings
+	OutNewObjects = MoveTemp(Factory->NewEntities);
+}
+
+bool FDMXEditorUtils::AreFixtureTypesIdentical(const UDMXEntityFixtureType* A, const UDMXEntityFixtureType* B)
+{
+	if (A == B)
+	{
+		return true;
+	}
+	if (A == nullptr || B == nullptr)
+	{
+		return false;
+	}
+	if (A->GetClass() != B->GetClass())
+	{
+		return false;
+	}
+
+	// Compare each UProperty in the projects
+	const UStruct* Struct = UDMXEntityFixtureType::StaticClass();
+	TPropertyValueIterator<const FProperty> ItA(Struct, A);
+	TPropertyValueIterator<const FProperty> ItB(Struct, B);
+
+	static const FName NAME_ParentLibrary = TEXT("ParentLibrary");
+	static const FName NAME_Id = TEXT("Id");
+
+	for (; ItA && ItB; ++ItA, ++ItB)
+	{
+		const FProperty* Property = ItA->Key;
+
+		// Name and Id don't have to be identical
+		if (Property->GetFName() == GET_MEMBER_NAME_CHECKED(UDMXEntity, Name)
+			|| Property->GetFName() == NAME_ParentLibrary) // Can't GET_MEMBER_NAME... with private properties
+		{
+			continue;
+		}
+
+		if (Property->GetFName() == NAME_Id)
+		{
+			// Skip all properties from GUID struct
+			for (int32 PropertyCount = 0; PropertyCount < 4; ++PropertyCount)
+			{
+				++ItA;
+				++ItB;
+			}
+			continue;
+		}
+
+		const void* ValueA = ItA->Value;
+		const void* ValueB = ItB->Value;
+
+		if (!Property->Identical(ValueA, ValueB))
+		{
+			return false;
+		}
+	}
+
+	return true;
+}
+
+FText FDMXEditorUtils::GetEntityTypeNameText(TSubclassOf<UDMXEntity> EntityClass, bool bPlural /*= false*/)
+{
+	if (EntityClass->IsChildOf(UDMXEntityController::StaticClass()))
+	{
+		return FText::Format(
+			LOCTEXT("EntityTypeName_Controller", "{0}|plural(one=Controller, other=Controllers)"),
+			bPlural ? 2 : 1
+		);
+	}
+	else if (EntityClass->IsChildOf(UDMXEntityFixtureType::StaticClass()))
+	{
+		return FText::Format(
+			LOCTEXT("EntityTypeName_FixtureType", "Fixture {0}|plural(one=Type, other=Types)"),
+			bPlural ? 2 : 1
+		);
+	}
+	else if (EntityClass->IsChildOf(UDMXEntityFixturePatch::StaticClass()))
+	{
+		return FText::Format(
+			LOCTEXT("EntityTypeName_FixturePatch", "Fixture {0}|plural(one=Patch, other=Patches)"),
+			bPlural ? 2 : 1
+		);
+	}
+	else
+	{
+		return FText::Format(
+			LOCTEXT("EntityTypeName_NotImplemented", "{0}|plural(one=Entity, other=Entities)"),
+			bPlural ? 2 : 1
+		);
+	}
+}
+
+#undef LOCTEXT_NAMESPACE
diff --git a/Engine/Plugins/Experimental/DigitalMultiplex/DMXEngine/Source/DMXEditor/Private/Factories/DMXGDTFFactory.cpp b/Engine/Plugins/Experimental/DigitalMultiplex/DMXEngine/Source/DMXEditor/Private/Factories/DMXGDTFFactory.cpp
new file mode 100644
--- /dev/null
+++ b/Engine/Plugins/Experimental/DigitalMultiplex/DMXEngine/Source/DMXEditor/Private/Factories/DMXGDTFFactory.cpp
@@ -0,0 +1,231 @@
+// Copyright Epic Games, Inc. All Rights Reserved.
+
+#include "Factories/DMXGDTFFactory.h"
+#include "Factories/DMXGDTFImportUI.h"
+#include "DMXEditorLog.h"
+#include "Factories/DMXGDTFImporter.h"
+#include "Library/DMXImportGDTF.h"
+
+#include "Misc/Paths.h"
+#include "Editor.h"
+#include "EditorReimportHandler.h"
+#include "AssetImportTask.h"
+#include "Misc/FeedbackContext.h"
+#include "Misc/CommandLine.h"
+#include "HAL/FileManager.h"
+
+const TCHAR* UDMXGDTFFactory::Extension = TEXT("gdtf");
+
+UDMXGDTFImportUI::UDMXGDTFImportUI()
+    : bUseSubDirectory(false)
+    , bImportXML(true)
+    , bImportTextures(false)
+    , bImportModels(false)
+{
+}
+
+void UDMXGDTFImportUI::ResetToDefault()
+{
+    bUseSubDirectory = false;
+    bImportXML = true;
+    bImportTextures = false;
+    bImportModels = false;
+}
+
+UDMXGDTFFactory::UDMXGDTFFactory()
+{
+    SupportedClass = nullptr;
+	Formats.Add(TEXT("gdtf;General Device Type Format"));
+
+	bCreateNew = false;
+	bText = false;
+	bEditorImport = true;
+	bOperationCanceled = false;
+}
+
+void UDMXGDTFFactory::CleanUp()
+{
+	Super::CleanUp();
+
+    bShowOption = true;
+}
+
+bool UDMXGDTFFactory::ConfigureProperties()
+{
+    Super::ConfigureProperties();
+    EnableShowOption();
+
+    return true;
+}
+
+void UDMXGDTFFactory::PostInitProperties()
+{
+	Super::PostInitProperties();
+
+    ImportUI = NewObject<UDMXGDTFImportUI>(this, NAME_None, RF_NoFlags);
+}
+
+bool UDMXGDTFFactory::DoesSupportClass(UClass* Class)
+{
+    return Class == UDMXImportGDTF::StaticClass();
+}
+
+UClass* UDMXGDTFFactory::ResolveSupportedClass()
+{
+    return UDMXImportGDTF::StaticClass();
+}
+
+UObject* UDMXGDTFFactory::FactoryCreateFile(UClass* InClass, UObject* InParent, FName InName, EObjectFlags Flags, const FString& InFilename, const TCHAR* Parms, FFeedbackContext* Warn, bool& bOutOperationCanceled)
+{
+    FString FileExtension = FPaths::GetExtension(InFilename);
+    const TCHAR* Type = *FileExtension;
+
+    if (!IFileManager::Get().FileExists(*InFilename))
+    {
+        UE_LOG_DMXEDITOR(Error, TEXT("Failed to load file '%s'"), *InFilename)
+        return nullptr;
+    }
+
+    ParseParms(Parms);
+
+    CA_ASSUME(InParent);
+
+    if( bOperationCanceled )
+    {
+        bOutOperationCanceled = true;
+        GEditor->GetEditorSubsystem<UImportSubsystem>()->BroadcastAssetPostImport(this, nullptr);
+        return nullptr;
+    }
+
+    GEditor->GetEditorSubsystem<UImportSubsystem>()->BroadcastAssetPreImport(this, InClass, InParent, InName, Type);
+
+    UDMXImportGDTF* CreatedObject = nullptr;
+    UObject* ExistingObject = nullptr;
+    if (InParent != nullptr)
+    {
+        ExistingObject = StaticFindObject(UObject::StaticClass(), InParent, *(InName.ToString()));
+        if (ExistingObject)
+        {
+            bShowOption = false;
+        }
+    }
+
+    // Prepare import options
+    FDMXGDTFImportArgs ImportArgs;
+    ImportArgs.ImportUI = ImportUI;
+    ImportArgs.Name = InName;
+    ImportArgs.Parent = InParent;
+    ImportArgs.CurrentFilename = CurrentFilename;
+    ImportArgs.Flags = Flags;
+    ImportArgs.bCancelOperation = bOperationCanceled;
+    TUniquePtr<FDMXGDTFImporter> Importer = MakeUnique<FDMXGDTFImporter>(ImportArgs);
+
+    // Set Import UI
+    if (FParse::Param(FCommandLine::Get(), TEXT("NoDMXImportOption")))
+    {
+        bShowOption = false;
+    }
+    bool bIsAutomated = IsAutomatedImport();
+    bool bShowImportDialog = bShowOption && !bIsAutomated;
+    bool bImportAll = false;
+    FDMXGDTFImporter::GetImportOptions(Importer, ImportUI, bShowImportDialog, InParent->GetPathName(), bOperationCanceled, bImportAll, UFactory::CurrentFilename);
+    bOutOperationCanceled = bOperationCanceled;
+    if( bImportAll )
+    {
+        // If the user chose to import all, we don't show the dialog again and use the same settings for each object until importing another set of files
+        bShowOption = false;
+    }
+
+    bool bImportError = false;
+
+    if (!ImportUI->bImportXML && !ImportUI->bImportModels && !ImportUI->bImportTextures)
+    {
+        Warn->Log(ELogVerbosity::Error, TEXT("Nothing to Import") );
+        bImportError = true;
+    }
+
+    // Try to load and parse the content
+    if ( !Importer->AttemptImportFromFile() )
+    {
+        Warn->Log(ELogVerbosity::Error, TEXT("Attempt Import") );
+        bImportError = true;
+    }
+
+    // Import was successful, create new file
+    if (!bImportError)
+    {
+		// Import to the Editor
+		CreatedObject = Importer->Import();
+
+        if (CreatedObject != nullptr)
+        {
+            CreatedObject->SourceFilename = GetCurrentFilename();
+        }
+		GEditor->GetEditorSubsystem<UImportSubsystem>()->BroadcastAssetPostImport(this, CreatedObject);
+    }
+
+    // We do not need importer any more
+    Importer.Reset();
+
+    return CreatedObject;
+}
+
+bool UDMXGDTFFactory::FactoryCanImport(const FString& Filename)
+{
+	const FString TargetExtension = FPaths::GetExtension(Filename);
+
+	if(TargetExtension == UDMXGDTFFactory::Extension)
+	{
+		return true;
+	}
+
+	return false;
+}
+
+bool UDMXGDTFFactory::CanReimport(UObject* Obj, TArray<FString>& OutFilenames)
+{
+	UDMXImportGDTF* GDTFReimport = Cast<UDMXImportGDTF>(Obj);
+	if (GDTFReimport)
+	{
+		OutFilenames.Add(GDTFReimport->SourceFilename);
+		return true;
+	}
+	return true;
+}
+
+void UDMXGDTFFactory::SetReimportPaths(UObject* Obj, const TArray<FString>& NewReimportPaths)
+{
+	UDMXImportGDTF* GDTFReimport = Cast<UDMXImportGDTF>(Obj);
+	if (GDTFReimport && ensure(NewReimportPaths.Num() == 1))
+	{
+        GDTFReimport->SourceFilename = NewReimportPaths[0];
+	}
+}
+
+EReimportResult::Type UDMXGDTFFactory::Reimport(UObject* InObject)
+{
+	UDMXImportGDTF* GDTFReimport = Cast<UDMXImportGDTF>(InObject);
+
+	if (!GDTFReimport)
+	{
+		return EReimportResult::Failed;
+	}
+
+	if (GDTFReimport->SourceFilename.IsEmpty() || !FPaths::FileExists(GDTFReimport->SourceFilename))
+	{
+		return EReimportResult::Failed;
+	}
+
+	bool OutCanceled = false;
+	if (ImportObject(InObject->GetClass(), InObject->GetOuter(), *InObject->GetName(), RF_Public | RF_Standalone, GDTFReimport->SourceFilename, nullptr, OutCanceled))
+	{
+		return EReimportResult::Succeeded;
+	}
+
+	return OutCanceled ? EReimportResult::Cancelled : EReimportResult::Failed;
+}
+
+int32 UDMXGDTFFactory::GetPriority() const
+{
+    return ImportPriority;
+}
diff --git a/Engine/Plugins/Experimental/DigitalMultiplex/DMXEngine/Source/DMXEditor/Private/Factories/DMXGDTFImportUI.h b/Engine/Plugins/Experimental/DigitalMultiplex/DMXEngine/Source/DMXEditor/Private/Factories/DMXGDTFImportUI.h
new file mode 100644
--- /dev/null
+++ b/Engine/Plugins/Experimental/DigitalMultiplex/DMXEngine/Source/DMXEditor/Private/Factories/DMXGDTFImportUI.h
@@ -0,0 +1,35 @@
+// Copyright Epic Games, Inc. All Rights Reserved.
+
+#pragma once
+
+#include "CoreMinimal.h"
+#include "UObject/ObjectMacros.h"
+#include "UObject/Object.h"
+#include "DMXGDTFImportUI.generated.h"
+
+UCLASS(config = Editor, HideCategories=Object, MinimalAPI)
+class UDMXGDTFImportUI 
+    : public UObject
+{
+	GENERATED_BODY()
+
+public:
+    UDMXGDTFImportUI();
+
+	void ResetToDefault();
+
+public:
+    UPROPERTY(EditAnywhere, Category = "DMX")
+    bool bUseSubDirectory;
+
+    UPROPERTY(EditAnywhere, Category = "DMX")
+    bool bImportXML;
+
+    UPROPERTY(EditAnywhere, Category = "DMX")
+    bool bImportTextures;
+
+    UPROPERTY(EditAnywhere, Category = "DMX")
+    bool bImportModels;
+};
+
+
diff --git a/Engine/Plugins/Experimental/DigitalMultiplex/DMXEngine/Source/DMXEditor/Private/Factories/DMXGDTFImporter.cpp b/Engine/Plugins/Experimental/DigitalMultiplex/DMXEngine/Source/DMXEditor/Private/Factories/DMXGDTFImporter.cpp
new file mode 100644
--- /dev/null
+++ b/Engine/Plugins/Experimental/DigitalMultiplex/DMXEngine/Source/DMXEditor/Private/Factories/DMXGDTFImporter.cpp
@@ -0,0 +1,663 @@
+// Copyright Epic Games, Inc. All Rights Reserved.
+
+#include "Factories/DMXGDTFImporter.h"
+#include "Library/DMXImportGDTF.h"
+#include "DMXEditorLog.h"
+#include "Widgets/SDMXGDTFOptionWindow.h"
+#include "Factories/DMXGDTFImportUI.h"
+
+#include "Framework/Application/SlateApplication.h"
+#include "Interfaces/IMainFrameModule.h"
+#include "Modules/ModuleManager.h"
+#include "ObjectTools.h"
+#include "HAL/PlatformFilemanager.h"
+#include "HAL/FileManagerGeneric.h"
+#include "Misc/FileHelper.h"
+#include "PackageTools.h"
+#include "Serialization/BufferArchive.h"
+#include "Containers/UnrealString.h"
+#include "Factories/TextureFactory.h"
+#include "Engine/Texture.h"
+#include "AssetRegistryModule.h"
+#include "AssetToolsModule.h"
+#include "EditorFramework/AssetImportData.h"
+#include "XmlNode.h"
+#include "XmlFile.h"
+#include "Engine/Texture2D.h"
+
+#define LOCTEXT_NAMESPACE "DMXGDTFImporter"
+
+FDMXGDTFImporter::FDMXGDTFImporter(const FDMXGDTFImportArgs& InImportArgs)
+    : ImportArgs(InImportArgs)
+    , bIsXMLParsedSuccessfully(false)
+{
+    ImportName = ObjectTools::SanitizeObjectName(ImportArgs.Name.ToString());
+}
+
+UDMXImportGDTF* FDMXGDTFImporter::Import()
+{
+    UDMXImportGDTF* CreatedObject = nullptr;
+
+    // Import GDTF XML into UObject
+    if (ImportArgs.ImportUI->bImportXML && bIsXMLParsedSuccessfully)
+    {
+        CreatedObject = CreateGDTFDesctription();
+    }
+
+    return CreatedObject;
+}
+
+UPackage* FDMXGDTFImporter::GetPackage(const FString& InImportName)
+{
+    UPackage* Package = nullptr;
+    FString NewPackageName;
+
+    if (!ImportArgs.ImportUI->bUseSubDirectory && ImportArgs.Parent != nullptr && ImportArgs.Parent->IsA(UPackage::StaticClass()))
+    {
+        Package = StaticCast<UPackage*>(ImportArgs.Parent.Get());
+    }
+
+    if (Package == nullptr)
+    {
+        if (ImportArgs.Parent != nullptr && ImportArgs.Parent->GetOutermost() != nullptr)
+        {
+            NewPackageName = FPackageName::GetLongPackagePath(ImportArgs.Parent->GetOutermost()->GetName()) + "/" + ImportName;
+        }
+        else
+        {
+            return nullptr;
+        }
+
+        if (ImportArgs.ImportUI->bUseSubDirectory)
+        {
+            NewPackageName /= InImportName;
+        }
+
+        NewPackageName = UPackageTools::SanitizePackageName(NewPackageName);
+        LastPackageName = NewPackageName;
+        Package = CreatePackage(NULL, *NewPackageName);
+        Package->FullyLoad();
+    }
+
+    return Package;
+}
+
+bool FDMXGDTFImporter::AttemptImportFromFile()
+{
+    if (ImportArgs.ImportUI->bImportXML)
+    {
+        if (!ParseXML())
+        {
+            UE_LOG_DMXEDITOR(Warning, TEXT("Error parsing XML"));
+
+            return false;
+        }
+    }
+
+    return true;
+}
+
+UDMXImportGDTF* FDMXGDTFImporter::CreateGDTFDesctription()
+{
+    UPackage* Package = GetPackage(ImportName);
+
+    // Create new GDTF asset
+    UDMXImportGDTF* ImportObject = NewObject<UDMXImportGDTF>(Package, FName(*ImportName), ImportArgs.Flags | RF_Public);
+
+    UDMXImportGDTFFixtureType* GDTFFixtureType = ImportObject->CreateNewObject<UDMXImportGDTFFixtureType>();
+    ImportObject->FixtureType = GDTFFixtureType;
+
+    UDMXImportGDTFAttributeDefinitions* GDTFAttributeDefinitions = ImportObject->CreateNewObject<UDMXImportGDTFAttributeDefinitions>();
+    ImportObject->AttributeDefinitions = GDTFAttributeDefinitions;
+
+    UDMXImportGDTFWheels* GDTFWheels = ImportObject->CreateNewObject<UDMXImportGDTFWheels>();
+    ImportObject->Wheels = GDTFWheels;
+
+    UDMXImportGDTFPhysicalDescriptions* GDTFPhysicalDescriptions = ImportObject->CreateNewObject<UDMXImportGDTFPhysicalDescriptions>();
+    ImportObject->PhysicalDescriptions = GDTFPhysicalDescriptions;
+
+    UDMXImportGDTFModels* GDTFModels = ImportObject->CreateNewObject<UDMXImportGDTFModels>();
+    ImportObject->Models = GDTFModels;
+
+    UDMXImportGDTFGeometries* GDTFGeometries = ImportObject->CreateNewObject<UDMXImportGDTFGeometries>();
+    ImportObject->Geometries = GDTFGeometries;
+
+    UDMXImportGDTFDMXModes* GDTFDMXModes = ImportObject->CreateNewObject<UDMXImportGDTFDMXModes>();
+    ImportObject->DMXModes = GDTFDMXModes;
+
+    UDMXImportGDTFProtocols* GDTFProtocols = ImportObject->CreateNewObject<UDMXImportGDTFProtocols>();
+    ImportObject->Protocols = GDTFProtocols;
+
+    if (bIsXMLParsedSuccessfully)
+    {
+        if (FXmlNode* XmlNode = XMLFile->GetRootNode())
+        {
+            if (const FXmlNode* FixtureTypeNode = XmlNode->FindChildNode("FixtureType"))
+            {
+                GDTFFixtureType->Description = FixtureTypeNode->GetAttribute("Description");
+                GDTFFixtureType->FixtureTypeID = FixtureTypeNode->GetAttribute("FixtureTypeID");
+                GDTFFixtureType->LongName = FixtureTypeNode->GetAttribute("LongName");
+                GDTFFixtureType->Manufacturer = FixtureTypeNode->GetAttribute("Manufacturer");
+                GDTFFixtureType->Name = FName(*FixtureTypeNode->GetAttribute("Name"));
+                GDTFFixtureType->ShortName = FixtureTypeNode->GetAttribute("ShortName");
+                GDTFFixtureType->RefFT = FixtureTypeNode->GetAttribute("RefFT");
+
+                if (const FXmlNode* AttributeDefinitionsNode = FixtureTypeNode->FindChildNode("AttributeDefinitions"))
+                {
+                    if (const FXmlNode* ActivationGroupsNode = AttributeDefinitionsNode->FindChildNode("ActivationGroups"))
+                    {
+                        for (const FXmlNode* ActivationGroupNode : ActivationGroupsNode->GetChildrenNodes())
+                        {
+                            FDMXImportGDTFActivationGroup ActivationGroup;
+                            ActivationGroup.Name = FName(*ActivationGroupNode->GetAttribute("Name"));
+                            GDTFAttributeDefinitions->ActivationGroups.Add(ActivationGroup);
+                        }
+                    }
+
+                    if (const FXmlNode* FeatureGroupsNode = AttributeDefinitionsNode->FindChildNode("FeatureGroups"))
+                    {
+                        for (const FXmlNode* FeatureGroupNode : FeatureGroupsNode->GetChildrenNodes())
+                        {
+                            if (FeatureGroupNode != nullptr)
+                            {
+                                FDMXImportGDTFFeatureGroup FeatureGroup;
+                                FeatureGroup.Name = FName(*FeatureGroupNode->GetAttribute("Name"));
+                                FeatureGroup.Pretty = FeatureGroupNode->GetAttribute("Pretty");
+
+                                for (const FXmlNode* FeatureNode : FeatureGroupNode->GetChildrenNodes())
+                                {
+                                    FDMXImportGDTFFeature Feature;
+                                    Feature.Name = FName(*FeatureNode->GetAttribute("Name"));
+                                    FeatureGroup.Features.Add(Feature);
+                                }
+
+                                GDTFAttributeDefinitions->FeatureGroups.Add(FeatureGroup);
+                            }
+                        }
+                    }
+
+                    if (const FXmlNode* AttributesNode = AttributeDefinitionsNode->FindChildNode("Attributes"))
+                    {
+                        for (const FXmlNode* AttributeNode : AttributesNode->GetChildrenNodes())
+                        {
+                            if (AttributeNode != nullptr)
+                            {
+                                FDMXImportGDTFAttribute Attribute ;
+                                Attribute.Name = FName(*AttributeNode->GetAttribute("Name"));
+                                Attribute.Pretty = AttributeNode->GetAttribute("Pretty");
+                                Attribute.ActivationGroup.Name = FName(*AttributeNode->GetAttribute("ActivationGroup"));
+                                GDTFAttributeDefinitions->FindFeature(AttributeNode->GetAttribute("Feature"), Attribute.Feature);
+                                Attribute.MainAttribute = AttributeNode->GetAttribute("MainAttribute");
+                                Attribute.PhysicalUnit = DMXImport::GetEnumValueFromString<EDMXImportGDTFPhysicalUnit>(AttributeNode->GetAttribute("PhysicalUnit"));
+                                Attribute.Color = DMXImport::ParseColorCIE(AttributeNode->GetAttribute("Color"));
+                                GDTFAttributeDefinitions->Attributes.Add(Attribute);
+                            }
+                        }
+                    }
+                }
+
+                if (const FXmlNode* WheelsNode = FixtureTypeNode->FindChildNode("Wheels"))
+                {
+                    for (const FXmlNode* WheelNode : WheelsNode->GetChildrenNodes())
+                    {
+                        FDMXImportGDTFWheel ImportWheel;
+                        ImportWheel.Name = FName(*WheelNode->GetAttribute("Name"));
+
+                        for (const FXmlNode* SlotNode : WheelNode->GetChildrenNodes())
+                        {
+                            FDMXImportGDTFWheelSlot ImportWheelSlot;
+                            ImportWheelSlot.Name = FName(*SlotNode->GetAttribute("Name"));
+                            ImportWheelSlot.Color = DMXImport::ParseColorCIE(SlotNode->GetAttribute("Color"));
+                            ImportWheel.Slots.Add(ImportWheelSlot);
+                        }
+
+                        GDTFWheels->Wheels.Add(ImportWheel);
+                    }
+                }
+
+                if (const FXmlNode* PhysicalDescriptionsNode = FixtureTypeNode->FindChildNode("PhysicalDescriptions"))
+                {
+                    if (const FXmlNode* EmittersNode = PhysicalDescriptionsNode->FindChildNode("Emitters"))
+                    {
+                        for (const FXmlNode* EmitterNode : EmittersNode->GetChildrenNodes())
+                        {
+                            FDMXImportGDTFEmitter ImportEmitter;
+
+                            ImportEmitter.Name = FName(*EmitterNode->GetAttribute("Name"));
+                            ImportEmitter.Color = DMXImport::ParseColorCIE(EmitterNode->GetAttribute("Color"));
+                            ImportEmitter.DiodePart = EmitterNode->GetAttribute("DiodePart");
+                            LexTryParseString(ImportEmitter.DominantWaveLength, *EmitterNode->GetAttribute("DominantWaveLength"));
+                            ImportEmitter.DiodePart = EmitterNode->GetAttribute("DiodePart");
+
+                            if (const FXmlNode* MeasurementNode = EmitterNode->FindChildNode("Measurement"))
+                            {
+                                FDMXImportGDTFMeasurement ImportMeasurement;
+
+                                LexTryParseString(ImportMeasurement.Physical, *MeasurementNode->GetAttribute("DominantWaveLength"));
+                                LexTryParseString(ImportMeasurement.LuminousIntensity, *MeasurementNode->GetAttribute("LuminousIntensity"));
+                                LexTryParseString(ImportMeasurement.Transmission, *MeasurementNode->GetAttribute("Transmission"));
+                                ImportMeasurement.InterpolationTo = DMXImport::GetEnumValueFromString<EDMXImportGDTFInterpolationTo>(MeasurementNode->GetAttribute("InterpolationTo"));
+
+                                for (const FXmlNode* MeasurementPointNode : MeasurementNode->GetChildrenNodes())
+                                {
+                                    FDMXImportGDTFMeasurementPoint ImportMeasurementPoint;
+
+                                    LexTryParseString(ImportMeasurementPoint.Energy, *MeasurementPointNode->GetAttribute("Energy"));
+                                    LexTryParseString(ImportMeasurementPoint.WaveLength, *MeasurementPointNode->GetAttribute("WaveLength"));
+
+                                    ImportMeasurement.MeasurementPoints.Add(ImportMeasurementPoint);
+                                }
+
+                                ImportEmitter.Measurement = ImportMeasurement;
+                            }
+
+                            GDTFPhysicalDescriptions->Emitters.Add(ImportEmitter);
+                        }
+                    }
+
+                    if (const FXmlNode* ColorSpaceNode = PhysicalDescriptionsNode->FindChildNode("ColorSpace"))
+                    {
+                        FDMXImportGDTFColorSpace ImportColorSpace;
+                        ImportColorSpace.Mode = DMXImport::GetEnumValueFromString<EDMXImportGDTFMode>(ColorSpaceNode->GetAttribute("Mode"));
+                        ImportColorSpace.Description = ColorSpaceNode->GetAttribute("Description");
+                        ImportColorSpace.Red = DMXImport::ParseColorCIE(ColorSpaceNode->GetAttribute("Red"));
+                        ImportColorSpace.Blue = DMXImport::ParseColorCIE(ColorSpaceNode->GetAttribute("Blue"));
+                        ImportColorSpace.Green = DMXImport::ParseColorCIE(ColorSpaceNode->GetAttribute("Green"));
+                        ImportColorSpace.WhitePoint = DMXImport::ParseColorCIE(ColorSpaceNode->GetAttribute("WhitePoint"));
+
+                        GDTFPhysicalDescriptions->ColorSpace = ImportColorSpace;
+                    }
+                }
+
+                if (const FXmlNode* ModelsNode = FixtureTypeNode->FindChildNode("Models"))
+                {
+                    for (const FXmlNode* ModelNode : ModelsNode->GetChildrenNodes())
+                    {
+                        FDMXImportGDTFModel ImportModel;
+
+                        ImportModel.Name = FName(*ModelNode->GetAttribute("Name"));
+                        LexTryParseString(ImportModel.Length, *ModelNode->GetAttribute("Length"));
+                        LexTryParseString(ImportModel.Width, *ModelNode->GetAttribute("Width"));
+                        LexTryParseString(ImportModel.Height, *ModelNode->GetAttribute("Height"));
+                        ImportModel.PrimitiveType = DMXImport::GetEnumValueFromString<EDMXImportGDTFPrimitiveType>(ModelNode->GetAttribute("PrimitiveType"));
+
+                        GDTFModels->Models.Add(ImportModel);
+                    }
+                }
+
+                if (const FXmlNode* GeometriesNode = FixtureTypeNode->FindChildNode("Geometries"))
+                {
+                    for (const FXmlNode* GeometryNode : GeometriesNode->GetChildrenNodes())
+                    {
+                        FDMXImportGDTFGeneralGeometry ImportGeneralGeometry;
+                        ImportGeneralGeometry.Name = FName(*GeometryNode->GetAttribute("Name"));
+                        ImportGeneralGeometry.Model = FName(*GeometryNode->GetAttribute("Model"));
+                        ImportGeneralGeometry.Position = DMXImport::ParseMatrix(GeometryNode->GetAttribute("Position"));
+
+                        if (const FXmlNode* GeneralAxisNode = GeometryNode->FindChildNode("Axis"))
+                        {
+                            FDMXImportGDTFGeneralAxis ImportGeneralAxis;
+                            ImportGeneralAxis.Name = FName(*GeneralAxisNode->GetAttribute("Name"));
+                            ImportGeneralAxis.Model = FName(*GeneralAxisNode->GetAttribute("Model"));
+                            ImportGeneralAxis.Position = DMXImport::ParseMatrix(GeneralAxisNode->GetAttribute("Position"));
+
+                            for (const FXmlNode* TypeAxisNode : GeneralAxisNode->GetChildrenNodes())
+                            {
+                                FDMXImportGDTFTypeAxis ImportTypeAxis;
+                                ImportTypeAxis.Name = FName(*TypeAxisNode->GetAttribute("Name"));
+                                ImportTypeAxis.Model = FName(*TypeAxisNode->GetAttribute("Model"));
+                                ImportTypeAxis.Position = DMXImport::ParseMatrix(TypeAxisNode->GetAttribute("Position"));
+
+                                for (const FXmlNode* BeamNode : TypeAxisNode->GetChildrenNodes())
+                                {
+                                    FDMXImportGDTFBeam ImportBeam;
+                                    ImportBeam.Name = FName(*BeamNode->GetAttribute("Name"));
+                                    ImportBeam.Model = FName(*BeamNode->GetAttribute("Model"));
+                                    ImportBeam.Position = DMXImport::ParseMatrix(BeamNode->GetAttribute("Position"));
+                                    ImportBeam.LampType = DMXImport::GetEnumValueFromString<EDMXImportGDTFLampType>(BeamNode->GetAttribute("LampType"));
+
+                                    LexTryParseString(ImportBeam.PowerConsumption, *BeamNode->GetAttribute("PowerConsumption"));
+                                    LexTryParseString(ImportBeam.LuminousFlux, *BeamNode->GetAttribute("LuminousFlux"));
+                                    LexTryParseString(ImportBeam.ColorTemperature, *BeamNode->GetAttribute("ColorTemperature"));
+                                    LexTryParseString(ImportBeam.BeamAngle, *BeamNode->GetAttribute("BeamAngle"));
+                                    LexTryParseString(ImportBeam.FieldAngle, *BeamNode->GetAttribute("FieldAngle"));
+                                    LexTryParseString(ImportBeam.BeamRadius, *BeamNode->GetAttribute("BeamRadius"));
+                                    ImportBeam.BeamType = DMXImport::GetEnumValueFromString<EDMXImportGDTFBeamType>(BeamNode->GetAttribute("BeamType"));
+                                    LexTryParseString(ImportBeam.ColorRenderingIndex, *BeamNode->GetAttribute("ColorRenderingIndex"));
+
+                                    ImportTypeAxis.Beams.Add(ImportBeam);
+                                }
+                                ImportGeneralAxis.Axis.Add(ImportTypeAxis);
+                            }
+                            ImportGeneralGeometry.Axis = ImportGeneralAxis;
+                        }
+
+                        if (const FXmlNode* TypeGeometryNode = GeometryNode->FindChildNode("Geometry"))
+                        {
+                            FDMXImportGDTFTypeGeometry ImportTypeGeometry;
+                            ImportTypeGeometry.Name = FName(*TypeGeometryNode->GetAttribute("Name"));
+                            ImportTypeGeometry.Model = FName(*TypeGeometryNode->GetAttribute("Model"));
+                            ImportTypeGeometry.Position = DMXImport::ParseMatrix(TypeGeometryNode->GetAttribute("Position"));
+
+                            ImportGeneralGeometry.Geometry = ImportTypeGeometry;
+                        }
+
+                        if (const FXmlNode* FilterBeamNode = GeometryNode->FindChildNode("FilterBeam"))
+                        {
+                            FDMXImportGDTFFilterBeam ImportFilterBeam  ;
+                            ImportFilterBeam.Name = FName(*FilterBeamNode->GetAttribute("Name"));
+                            ImportFilterBeam.Model = FName(*FilterBeamNode->GetAttribute("Model"));
+                            ImportFilterBeam.Position = DMXImport::ParseMatrix(FilterBeamNode->GetAttribute("Position"));
+
+                            ImportGeneralGeometry.FilterBeam = ImportFilterBeam ;
+                        }
+
+                        if (const FXmlNode* FilterColorNode = GeometryNode->FindChildNode("FilterColor"))
+                        {
+                            FDMXImportGDTFFilterColor ImportFilterColor;
+                            ImportFilterColor.Name = FName(*FilterColorNode->GetAttribute("Name"));
+                            ImportFilterColor.Model = FName(*FilterColorNode->GetAttribute("Model"));
+                            ImportFilterColor.Position = DMXImport::ParseMatrix(FilterColorNode->GetAttribute("Position"));
+
+                            ImportGeneralGeometry.FilterColor = ImportFilterColor;
+                        }
+
+                        if (const FXmlNode* FilterGoboNode = GeometryNode->FindChildNode("FilterGobo"))
+                        {
+                            FDMXImportGDTFFilterGobo ImportFilterGobo;
+                            ImportFilterGobo.Name = FName(*FilterGoboNode->GetAttribute("Name"));
+                            ImportFilterGobo.Model = FName(*FilterGoboNode->GetAttribute("Model"));
+                            ImportFilterGobo.Position = DMXImport::ParseMatrix(FilterGoboNode->GetAttribute("Position"));
+
+                            ImportGeneralGeometry.FilterGobo = ImportFilterGobo;
+                        }
+
+                        if (const FXmlNode* FilterShaperNode = GeometryNode->FindChildNode("FilterShaper"))
+                        {
+                            FDMXImportGDTFFilterShaper ImportFilterShaper;
+                            ImportFilterShaper.Name = FName(*FilterShaperNode->GetAttribute("Name"));
+                            ImportFilterShaper.Model = FName(*FilterShaperNode->GetAttribute("Model"));
+                            ImportFilterShaper.Position = DMXImport::ParseMatrix(FilterShaperNode->GetAttribute("Position"));
+
+                            ImportGeneralGeometry.FilterShaper = ImportFilterShaper;
+                        }
+
+                        if (const FXmlNode* GeometryReferenceNode = GeometryNode->FindChildNode("GeometryReference"))
+                        {
+                            FDMXImportGDTFGeometryReference ImportGeometryReference;
+                            ImportGeometryReference.Name = FName(*GeometryReferenceNode->GetAttribute("Name"));
+                            ImportGeometryReference.Model = FName(*GeometryReferenceNode->GetAttribute("Model"));
+                            ImportGeometryReference.Position = DMXImport::ParseMatrix(GeometryReferenceNode->GetAttribute("Position"));
+
+                            for (const FXmlNode* BreakNode : GeometryReferenceNode->GetChildrenNodes())
+                            {
+                                FDMXImportGDTFBreak ImportBreak;
+
+                                LexTryParseString(ImportBreak.DMXOffset, *BreakNode->GetAttribute("DMXOffset"));
+                                LexTryParseString(ImportBreak.DMXBreak, *BreakNode->GetAttribute("DMXBreak"));
+
+                                ImportGeometryReference.Breaks.Add(ImportBreak);
+                            }
+
+                            ImportGeneralGeometry.GeometryReference = ImportGeometryReference;
+                        }
+
+                        GDTFGeometries->GeneralGeometry.Add(ImportGeneralGeometry);
+                    }
+                }
+
+                if (const FXmlNode* DMXModesNode = FixtureTypeNode->FindChildNode("DMXModes"))
+                {
+                    for (const FXmlNode* DMXModeNode : DMXModesNode->GetChildrenNodes())
+                    {
+                        FDMXImportGDTFDMXMode DMXImportGDTFDMXMode;
+                        DMXImportGDTFDMXMode.Name = FName(*DMXModeNode->GetAttribute("Name"));
+                        DMXImportGDTFDMXMode.Geometry = FName(*DMXModeNode->GetAttribute("Geometry"));
+
+                        if (const FXmlNode* DMXChannelsNode = DMXModeNode->FindChildNode("DMXChannels"))
+                        {
+                            for (const FXmlNode* DMXChannelNode : DMXChannelsNode->GetChildrenNodes())
+                            {
+                                FDMXImportGDTFDMXChannel ImportDMXChannel;
+                                LexTryParseString(ImportDMXChannel.DMXBreak, *DMXChannelNode->GetAttribute("DMXBreak"));
+                                ImportDMXChannel.ParseOffset(DMXChannelNode->GetAttribute("Offset"));
+                                ImportDMXChannel.Default = FDMXImportGDTFDMXValue(DMXChannelNode->GetAttribute("Default"));
+                                ImportDMXChannel.Highlight = FDMXImportGDTFDMXValue(DMXChannelNode->GetAttribute("Highlight"));
+                                ImportDMXChannel.Geometry = FName(*DMXChannelNode->GetAttribute("Geometry"));
+
+                                if (const FXmlNode* LogicalChannelNode = DMXChannelNode->FindChildNode("LogicalChannel"))
+                                {
+                                    FDMXImportGDTFLogicalChannel ImportLogicalChannel;
+
+                                    GDTFAttributeDefinitions->FindAtributeByName(*LogicalChannelNode->GetAttribute("Attribute"), ImportLogicalChannel.Attribute);
+                                    ImportLogicalChannel.Snap = DMXImport::GetEnumValueFromString<EDMXImportGDTFSnap>(LogicalChannelNode->GetAttribute("Snap"));
+                                    ImportLogicalChannel.Master = DMXImport::GetEnumValueFromString<EDMXImportGDTFMaster>(LogicalChannelNode->GetAttribute("Master"));
+
+                                    LexTryParseString(ImportLogicalChannel.MibFade, *LogicalChannelNode->GetAttribute("MibFade"));
+                                    LexTryParseString(ImportLogicalChannel.DMXChangeTimeLimit, *LogicalChannelNode->GetAttribute("DMXChangeTimeLimit"));
+
+                                    if (const FXmlNode* ChannelFunctionNode = LogicalChannelNode->FindChildNode("ChannelFunction"))
+                                    {
+                                        FDMXImportGDTFChannelFunction ImportChannelFunction;
+                                        ImportChannelFunction.Name = FName(*ChannelFunctionNode->GetAttribute("Name"));
+                                        GDTFAttributeDefinitions->FindAtributeByName(*ChannelFunctionNode->GetAttribute("Attribute"), ImportChannelFunction.Attribute);
+                                        ImportChannelFunction.OriginalAttribute = *ChannelFunctionNode->GetAttribute("OriginalAttribute");
+                                        ImportChannelFunction.DMXFrom = FDMXImportGDTFDMXValue(ChannelFunctionNode->GetAttribute("DMXFrom"));
+                                        ImportChannelFunction.DMXValue = FDMXImportGDTFDMXValue(ChannelFunctionNode->GetAttribute("DMXValue"));
+
+                                        LexTryParseString(ImportChannelFunction.PhysicalFrom, *ChannelFunctionNode->GetAttribute("PhysicalFrom"));
+                                        LexTryParseString(ImportChannelFunction.PhysicalTo, *ChannelFunctionNode->GetAttribute("PhysicalTo"));
+                                        LexTryParseString(ImportChannelFunction.RealFade, *ChannelFunctionNode->GetAttribute("RealFade"));
+
+                                        GDTFWheels->FindWeelByName(*ChannelFunctionNode->GetAttribute("Wheel"), ImportChannelFunction.Wheel);
+                                        GDTFPhysicalDescriptions->FindEmitterByName(*ChannelFunctionNode->GetAttribute("Emitter"), ImportChannelFunction.Emitter);
+
+                                        ImportChannelFunction.DMXInvert = DMXImport::GetEnumValueFromString<EDMXImportGDTFDMXInvert>(ChannelFunctionNode->GetAttribute("DMXInvert"));
+                                        ImportChannelFunction.ModeMaster = *ChannelFunctionNode->GetAttribute("ModeMaster");
+                                        ImportChannelFunction.ModeFrom = FDMXImportGDTFDMXValue(ChannelFunctionNode->GetAttribute("ModeFrom"));
+                                        ImportChannelFunction.ModeTo = FDMXImportGDTFDMXValue(ChannelFunctionNode->GetAttribute("ModeTo"));
+
+                                        for (const FXmlNode* ChannelSetNode : ChannelFunctionNode->GetChildrenNodes())
+                                        {
+                                            FDMXImportGDTFChannelSet ImportChannelSet;
+
+                                            ImportChannelSet.Name = *ChannelSetNode->GetAttribute("Name");
+                                            ImportChannelSet.DMXFrom = FDMXImportGDTFDMXValue(ChannelSetNode->GetAttribute("DMXFrom"));
+                                            LexTryParseString(ImportChannelSet.PhysicalFrom , *ChannelSetNode->GetAttribute("PhysicalFrom"));
+                                            LexTryParseString(ImportChannelSet.PhysicalTo, *ChannelSetNode->GetAttribute("PhysicalTo"));
+                                            LexTryParseString(ImportChannelSet.WheelSlotIndex, *ChannelSetNode->GetAttribute("WheelSlotIndex"));
+
+                                            ImportChannelFunction.ChannelSets.Add(ImportChannelSet);
+                                        }
+
+                                        ImportLogicalChannel.ChannelFunction = ImportChannelFunction;
+                                    }
+
+                                    ImportDMXChannel.LogicalChannel = ImportLogicalChannel;
+                                }
+
+                                DMXImportGDTFDMXMode.DMXChannels.Add(ImportDMXChannel);
+                            }
+                        }
+
+                        if (const FXmlNode* RelationsNode = DMXModeNode->FindChildNode("Relations"))
+                        {
+                            for (const FXmlNode* RelationNode : RelationsNode->GetChildrenNodes())
+                            {
+                                FDMXImportGDTFRelation ImportRelation;
+                                ImportRelation.Name = RelationNode->GetAttribute("Name");
+                                ImportRelation.Master = RelationNode->GetAttribute("Master");
+                                ImportRelation.Follower = RelationNode->GetAttribute("Follower");
+                                ImportRelation.Type = DMXImport::GetEnumValueFromString<EDMXImportGDTFType>(RelationNode->GetAttribute("Type"));
+
+                                DMXImportGDTFDMXMode.Relations.Add(ImportRelation);
+                            }
+                        }
+
+                        if (const FXmlNode* FTMacrosNode = DMXModeNode->FindChildNode("FTMacros"))
+                        {
+                            for (const FXmlNode* FTMacroNode : FTMacrosNode->GetChildrenNodes())
+                            {
+                                FDMXImportGDTFFTMacro ImportFTMacro;
+                                ImportFTMacro.Name = FName(*FTMacroNode->GetAttribute("Name"));
+                                DMXImportGDTFDMXMode.FTMacros.Add(ImportFTMacro);
+                            }
+                        }
+
+                        GDTFDMXModes->DMXModes.Add(DMXImportGDTFDMXMode);
+                    }
+                }
+
+                if (const FXmlNode* ProtocolsNode = FixtureTypeNode->FindChildNode("Protocols"))
+                {
+                    for (const FXmlNode* ProtocolNode : ProtocolsNode->GetChildrenNodes())
+                    {
+                        GDTFProtocols->Protocols.Add(FName(*ProtocolNode->GetTag()));
+                    }
+                }
+            }
+        }
+    }
+
+    return ImportObject;
+}
+
+bool FDMXGDTFImporter::ParseXML()
+{
+    //Load the Compressed data array
+    TArray<uint8> GDTFFileData;
+    if (!FFileHelper::LoadFileToArray(GDTFFileData, *ImportArgs.CurrentFilename))
+    {
+        UE_LOG_DMXEDITOR(Warning, TEXT("InFilename %s"), *ImportArgs.CurrentFilename);
+        return false;
+    }
+
+    // Convert to readable string
+    FString CharString;
+    {
+        const uint8* DataPtr = GDTFFileData.GetData();
+        int32 DataCount = GDTFFileData.Num();
+        CharString.Empty(DataCount);
+        while (DataCount)
+        {
+            uint8 Value = *DataPtr;
+            CharString += TCHAR(Value);
+
+            ++DataPtr;
+            DataCount--;
+        }
+    }
+
+    FString BeginXMLString("<?xml version=");
+    FString EndXMLString("</GDTF>");
+    int32 BeginXMLStringIndex = CharString.Find(BeginXMLString, ESearchCase::CaseSensitive);
+    int32 EndXMLStringIndex = CharString.Find(EndXMLString, ESearchCase::CaseSensitive, ESearchDir::FromEnd);
+    FString ResultXMLString = CharString.Mid(BeginXMLStringIndex, EndXMLStringIndex - BeginXMLStringIndex + EndXMLString.Len());
+
+    // @TODO. UE4 XML Parser can't handle symbol < or > inside attribute and just break the parsing
+    // This small algorithm remove < and > from the XML line
+    TArray<FString> XMLStringLines;
+    ResultXMLString.ParseIntoArrayLines(XMLStringLines, /*bCullEmpty*/false);
+    FString FinalXMLString;
+    for (FString& Line : XMLStringLines)
+    {
+        TCHAR LeftToken = '<';
+        TCHAR RightToken = '>';
+        TArray<int32> LeftTockensIndexes;
+        TArray<int32> RightTockensIndexes;
+
+        for (int32 CharIndex = 0; CharIndex < Line.Len(); ++CharIndex)
+        {
+            TCHAR CurrentChar = Line[CharIndex];
+
+            if (CurrentChar == LeftToken)
+            {
+                LeftTockensIndexes.Add(CharIndex);
+            }
+            else if (CurrentChar == RightToken)
+            {
+                RightTockensIndexes.Add(CharIndex);
+            }
+        }
+
+        int32 LeftTockensNum = LeftTockensIndexes.Num() - 1;
+        for (int32 LeftCharIndex = 1; LeftCharIndex < LeftTockensNum; ++LeftCharIndex)
+        {
+            Line[LeftTockensIndexes[LeftCharIndex]] = TCHAR(' ');
+        }
+
+        int32 RightTockensNum = RightTockensIndexes.Num() - 1;
+        for (int32 RightCharIndex = 0; RightCharIndex < RightTockensNum; ++RightCharIndex)
+        {
+            Line[RightTockensIndexes[RightCharIndex]] = TCHAR(' ');
+        }
+
+        FinalXMLString.Append(Line + LINE_TERMINATOR);
+    }
+
+    XMLFile = MakeShared<FXmlFile>(FinalXMLString, EConstructMethod::ConstructFromBuffer);
+    bIsXMLParsedSuccessfully = XMLFile->IsValid();
+
+    return bIsXMLParsedSuccessfully;
+}
+
+void FDMXGDTFImporter::GetImportOptions(const TUniquePtr<FDMXGDTFImporter>& Importer, UDMXGDTFImportUI* ImportUI, bool bShowOptionDialog, const FString& FullPath, bool& OutOperationCanceled, bool& bOutImportAll, const FString& InFilename)
+{
+    OutOperationCanceled = false;
+
+    if ( bShowOptionDialog )
+    {
+        TSharedPtr<SWindow> ParentWindow;
+
+        if( FModuleManager::Get().IsModuleLoaded( "MainFrame" ) )
+        {
+            IMainFrameModule& MainFrame = FModuleManager::LoadModuleChecked<IMainFrameModule>( "MainFrame" );
+            ParentWindow = MainFrame.GetParentWindow();
+        }
+
+        // Compute centered window position based on max window size, which include when all categories are expanded
+        const float ImportWindowWidth = 410.0f;
+        const float ImportWindowHeight = 750.0f;
+        FVector2D ImportWindowSize = FVector2D(ImportWindowWidth, ImportWindowHeight); // Max window size it can get based on current slate
+
+        FSlateRect WorkAreaRect = FSlateApplicationBase::Get().GetPreferredWorkArea();
+        FVector2D DisplayTopLeft(WorkAreaRect.Left, WorkAreaRect.Top);
+        FVector2D DisplaySize(WorkAreaRect.Right - WorkAreaRect.Left, WorkAreaRect.Bottom - WorkAreaRect.Top);
+
+        FVector2D WindowPosition = (DisplayTopLeft + (DisplaySize - ImportWindowSize) / 2.0f);
+
+		TSharedRef<SWindow> Window = SNew(SWindow)
+			.Title(LOCTEXT("GDTFImportOpionsTitle", "GDTF Import Options"))
+			.SizingRule(ESizingRule::Autosized)
+			.AutoCenter(EAutoCenter::None)
+			.ClientSize(ImportWindowSize)
+			.ScreenPosition(WindowPosition);
+
+		TSharedPtr<SDMXGDTFOptionWindow> OptionWindow;
+		Window->SetContent
+		(
+			SAssignNew(OptionWindow, SDMXGDTFOptionWindow)
+			.ImportUI(ImportUI)
+			.WidgetWindow(Window)
+			.FullPath(FText::FromString(FullPath))
+			.MaxWindowHeight(ImportWindowHeight)
+			.MaxWindowWidth(ImportWindowWidth)
+		);
+
+		// @todo: we can make this slow as showing progress bar later
+		FSlateApplication::Get().AddModalWindow(Window, ParentWindow, false);
+
+        if (OptionWindow->ShouldImport())
+        {
+            bOutImportAll = OptionWindow->ShouldImportAll();
+
+            return;
+        }
+        else
+        {
+            OutOperationCanceled = true;
+        }
+    }
+}
+
+#undef LOCTEXT_NAMESPACE
diff --git a/Engine/Plugins/Experimental/DigitalMultiplex/DMXEngine/Source/DMXEditor/Private/Factories/DMXGDTFImporter.h b/Engine/Plugins/Experimental/DigitalMultiplex/DMXEngine/Source/DMXEditor/Private/Factories/DMXGDTFImporter.h
new file mode 100644
--- /dev/null
+++ b/Engine/Plugins/Experimental/DigitalMultiplex/DMXEngine/Source/DMXEditor/Private/Factories/DMXGDTFImporter.h
@@ -0,0 +1,66 @@
+// Copyright Epic Games, Inc. All Rights Reserved.
+
+#pragma once
+
+#include "CoreMinimal.h"
+#include "UObject/WeakObjectPtr.h"
+
+class UDMXImportGDTF;
+class UDMXGDTFImportUI;
+class FXmlFile;
+class UPackage;
+
+struct FDMXGDTFImportArgs
+{
+public:
+    FDMXGDTFImportArgs()
+        : Parent(nullptr)
+        , Name(NAME_None)
+        , Flags(RF_NoFlags)
+        , ImportUI(nullptr)
+        , bCancelOperation(false)
+    {}
+
+    TWeakObjectPtr<UObject> Parent;
+    FName Name;
+    FString CurrentFilename;
+    EObjectFlags Flags;
+    TWeakObjectPtr<UDMXGDTFImportUI> ImportUI;
+    bool bCancelOperation;
+};
+
+/**
+ * GDTF Importer is read and parse GDTF input file
+ */
+class FDMXGDTFImporter
+{
+public:
+    FDMXGDTFImporter(const FDMXGDTFImportArgs& InImportArgs);
+
+    /** Try to load the file and parse the data */
+    bool AttemptImportFromFile();
+
+    /** Create new GDTF UObject from imported file */
+    UDMXImportGDTF* Import();
+
+private:
+    bool ParseXML();
+
+    /** Create GDTF from imported xml */
+    UDMXImportGDTF* CreateGDTFDesctription();
+
+    UPackage* GetPackage(const FString& InAssetName);
+public:
+
+    static void GetImportOptions(const TUniquePtr<FDMXGDTFImporter>& Importer, UDMXGDTFImportUI* ImportUI, bool bShowOptionDialog, const FString& FullPath, bool& OutOperationCanceled, bool& bOutImportAll, const FString& InFilename);
+private:
+    TSharedPtr<FXmlFile> XMLFile;
+
+    FDMXGDTFImportArgs ImportArgs;
+
+    bool bIsXMLParsedSuccessfully;
+
+    FString LastPackageName;
+
+    FString ImportName;
+};
\ No newline at end of file
diff --git a/Engine/Plugins/Experimental/DigitalMultiplex/DMXEngine/Source/DMXEditor/Private/Modes/DMXEditorApplicationMode.cpp b/Engine/Plugins/Experimental/DigitalMultiplex/DMXEngine/Source/DMXEditor/Private/Modes/DMXEditorApplicationMode.cpp
new file mode 100644
--- /dev/null
+++ b/Engine/Plugins/Experimental/DigitalMultiplex/DMXEngine/Source/DMXEditor/Private/Modes/DMXEditorApplicationMode.cpp
@@ -0,0 +1,78 @@
+// Copyright Epic Games, Inc. All Rights Reserved.
+
+#include "Modes/DMXEditorApplicationMode.h"
+#include "DMXEditor.h"
+#include "Toolbars/DMXEditorToolbar.h"
+#include "DMXEditorTabs.h"
+#include "Tabs/DMXEditorTabFactories.h"
+
+#define LOCTEXT_NAMESPACE "DMXEditorApplicationMode"
+
+const FName FDMXEditorApplicationMode::DefaultsMode(TEXT("DefaultsName"));
+
+FText FDMXEditorApplicationMode::GetLocalizedMode(const FName InMode)
+{
+	static TMap< FName, FText > LocModes;
+
+	if (LocModes.Num() == 0)
+	{
+		LocModes.Add(DefaultsMode, LOCTEXT("DMXDefaultsMode", "Defaults"));
+	}
+
+	check(InMode != NAME_None);
+	return LocModes[InMode];
+}
+
+FDMXEditorDefaultApplicationMode::FDMXEditorDefaultApplicationMode(TSharedPtr<FDMXEditor> InDMXEditor)
+	: FApplicationMode(FDMXEditorApplicationMode::DefaultsMode, FDMXEditorApplicationMode::GetLocalizedMode)
+	, DMXEditorCachedPtr(InDMXEditor)
+{
+	// 1. Create and register Tabs
+	DefaultsTabFactories.RegisterFactory(MakeShared<FDMXEditorControllersSummoner>(InDMXEditor));
+	DefaultsTabFactories.RegisterFactory(MakeShared<FDMXEditorFixtureTypesSummoner>(InDMXEditor));
+	DefaultsTabFactories.RegisterFactory(MakeShared<FDMXEditorFixturePatchSummoner>(InDMXEditor));
+	DefaultsTabFactories.RegisterFactory(MakeShared<FDMXEditorInputConsoleSummoner>(InDMXEditor));
+	DefaultsTabFactories.RegisterFactory(MakeShared<FDMXEditorOutputConsoleSummoner>(InDMXEditor));
+
+	// 2. REGISATER TAB LAYOUT
+	TabLayout = FTabManager::NewLayout("Standalone_SimpleAssetEditor_Layout_v5")
+		->AddArea
+		(
+			FTabManager::NewPrimaryArea()->SetOrientation(Orient_Vertical)
+			->Split
+			(
+				FTabManager::NewStack()
+				->SetSizeCoefficient(0.1f)
+				->SetHideTabWell(true)
+				->AddTab(InDMXEditor->GetToolbarTabId(), ETabState::OpenedTab)
+			)
+			->Split
+			(
+				FTabManager::NewStack()
+				->AddTab(FDMXEditorTabs::DMXControllersId, ETabState::OpenedTab)
+				->AddTab(FDMXEditorTabs::DMXFixtureTypesEditorTabId, ETabState::OpenedTab)
+				->AddTab(FDMXEditorTabs::DMXFixturePatchEditorTabId, ETabState::OpenedTab)
+				->AddTab(FDMXEditorTabs::DMXInputConsoleEditorTabId, ETabState::OpenedTab)
+				->AddTab(FDMXEditorTabs::DMXOutputConsoleEditorTabId, ETabState::OpenedTab)
+				->SetForegroundTab(FDMXEditorTabs::DMXControllersId)
+			)
+		);
+
+	// 3. SETUP TOOLBAR
+	InDMXEditor->GetToolbarBuilder()->AddCompileToolbar(ToolbarExtender);
+}
+
+void FDMXEditorDefaultApplicationMode::RegisterTabFactories(TSharedPtr<FTabManager> InTabManager)
+{
+	if (TSharedPtr<FDMXEditor> DMXEditorPtr = DMXEditorCachedPtr.Pin())
+	{
+		DMXEditorPtr->RegisterToolbarTab(InTabManager.ToSharedRef());
+
+		// Setup all tab factories
+		DMXEditorPtr->PushTabFactories(DefaultsTabFactories);
+
+		FApplicationMode::RegisterTabFactories(InTabManager);
+	}
+}
+
+#undef LOCTEXT_NAMESPACE
diff --git a/Engine/Plugins/Experimental/DigitalMultiplex/DMXEngine/Source/DMXEditor/Private/Style/DMXEditorStyle.cpp b/Engine/Plugins/Experimental/DigitalMultiplex/DMXEngine/Source/DMXEditor/Private/Style/DMXEditorStyle.cpp
new file mode 100644
--- /dev/null
+++ b/Engine/Plugins/Experimental/DigitalMultiplex/DMXEngine/Source/DMXEditor/Private/Style/DMXEditorStyle.cpp
@@ -0,0 +1,66 @@
+// Copyright Epic Games, Inc. All Rights Reserved.
+
+#include "DMXEditorStyle.h"
+
+#include "Styling/SlateStyleRegistry.h"
+#include "Framework/Application/SlateApplication.h"
+#include "Interfaces/IPluginManager.h"
+#include "Engine/Font.h"
+
+TSharedPtr<FSlateStyleSet> FDMXEditorStyle::StyleInstance = nullptr;
+
+void FDMXEditorStyle::Initialize()
+{
+	if (!StyleInstance.IsValid())
+	{
+		StyleInstance = Create();
+		FSlateStyleRegistry::RegisterSlateStyle(*StyleInstance);
+	}
+}
+
+void FDMXEditorStyle::Shutdown()
+{
+	ensureMsgf(StyleInstance.IsValid(), TEXT("%S called, but StyleInstance wasn't initialized"), __FUNCTION__);
+	FSlateStyleRegistry::UnRegisterSlateStyle(*StyleInstance);
+	ensure(StyleInstance.IsUnique());
+	StyleInstance.Reset();
+}
+
+FName FDMXEditorStyle::GetStyleSetName()
+{
+	return TEXT("DMXEditorStyle");
+}
+
+TSharedRef<FSlateStyleSet> FDMXEditorStyle::Create()
+{
+	TSharedRef<FSlateStyleSet> Style = MakeShared<FSlateStyleSet>(GetStyleSetName());
+	Style->SetContentRoot(IPluginManager::Get().FindPlugin("DMXEngine")->GetBaseDir() / TEXT("Resources"));
+
+	// Solid color brushes
+	Style->Set("DMXEditor.WhiteBrush", new FSlateColorBrush(FLinearColor(1.0f, 1.0f, 1.0f, 1.0f)));
+	Style->Set("DMXEditor.BlackBrush", new FSlateColorBrush(FLinearColor(0.0f, 0.0f, 0.0f, 1.0f)));
+
+	// Fonts
+	const TCHAR* FontPathRoboto = TEXT("Font'/Engine/EngineFonts/Roboto.Roboto'");
+	const UFont* FontRoboto = Cast<UFont>(StaticLoadObject(UFont::StaticClass(), nullptr, FontPathRoboto));
+	check(FontRoboto != nullptr);
+
+	Style->Set("DMXEditor.Font.InputChannelID", FSlateFontInfo(FontRoboto, 8, FName(TEXT("Light"))));
+	Style->Set("DMXEditor.Font.InputChannelValue", FSlateFontInfo(FontRoboto, 10, FName(TEXT("Regular"))));
+
+	return Style;
+}
+
+void FDMXEditorStyle::ReloadTextures()
+{
+	if (FSlateApplication::IsInitialized())
+	{
+		FSlateApplication::Get().GetRenderer()->ReloadTextureResources();
+	}
+}
+
+const ISlateStyle& FDMXEditorStyle::Get()
+{
+	check(StyleInstance.IsValid());
+	return *StyleInstance;
+}
\ No newline at end of file
diff --git a/Engine/Plugins/Experimental/DigitalMultiplex/DMXEngine/Source/DMXEditor/Private/Tabs/DMXEditorTabFactories.cpp b/Engine/Plugins/Experimental/DigitalMultiplex/DMXEngine/Source/DMXEditor/Private/Tabs/DMXEditorTabFactories.cpp
new file mode 100644
--- /dev/null
+++ b/Engine/Plugins/Experimental/DigitalMultiplex/DMXEngine/Source/DMXEditor/Private/Tabs/DMXEditorTabFactories.cpp
@@ -0,0 +1,141 @@
+// Copyright Epic Games, Inc. All Rights Reserved.
+
+#include "Tabs/DMXEditorTabFactories.h"
+#include "DMXEditorTabs.h"
+#include "DMXEditor.h"
+
+#include "Widgets/SDMXInputConsole.h"
+#include "Widgets/SDMXEntityEditor.h"
+
+#include "Widgets/Docking/SDockTab.h"
+
+#define LOCTEXT_NAMESPACE "DMXWorkflowTabFactory"
+
+TSharedRef<SDockTab> FDMXEditorPropertyTabSummoner::SpawnTab(const FWorkflowTabSpawnInfo& Info) const
+{
+	TSharedRef<SDockTab> NewTab = FWorkflowTabFactory::SpawnTab(Info);
+	NewTab->SetTag(GetIdentifier());
+
+	TSharedPtr<FDMXEditor> DMXEditor = StaticCastSharedPtr<FDMXEditor>(HostingApp.Pin());
+
+	return NewTab;
+}
+
+
+FDMXEditorControllersSummoner::FDMXEditorControllersSummoner(TSharedPtr<FAssetEditorToolkit> InHostingApp)
+	: FDMXEditorPropertyTabSummoner(FDMXEditorTabs::DMXControllersId, InHostingApp)
+{
+	TabLabel = LOCTEXT("DMXControllersTabLabel", "Controllers");
+	TabIcon = FSlateIcon(FEditorStyle::GetStyleSetName(), "Kismet.Tabs.BlueprintDefaults");
+
+	bIsSingleton = true;
+
+	ViewMenuDescription = LOCTEXT("DMXControllersView", "Controllers");
+	ViewMenuTooltip = LOCTEXT("DMXControllersViewTooltip", "Show the controllers view");
+}
+
+TSharedRef<SWidget> FDMXEditorControllersSummoner::CreateTabBody(const FWorkflowTabSpawnInfo& Info) const
+{
+	TSharedPtr<FDMXEditor> DMXEditor = StaticCastSharedPtr<FDMXEditor>(HostingApp.Pin());
+
+	return SNew(SBox)
+		.HAlign(HAlign_Fill)
+		.VAlign(VAlign_Fill)
+		[
+			DMXEditor->GetControllersTab()
+		];
+}
+
+FDMXEditorFixtureTypesSummoner::FDMXEditorFixtureTypesSummoner(TSharedPtr<FAssetEditorToolkit> InHostingApp)
+	: FDMXEditorPropertyTabSummoner(FDMXEditorTabs::DMXFixtureTypesEditorTabId, InHostingApp)
+{
+	TabLabel = LOCTEXT("DMXFixtureTypesTabLabel", "Fixture Types");
+	TabIcon = FSlateIcon(FEditorStyle::GetStyleSetName(), "Kismet.Tabs.BlueprintDefaults");
+
+	bIsSingleton = true;
+
+	ViewMenuDescription = LOCTEXT("DMXFixtureTypesView", "Fixture Types");
+	ViewMenuTooltip = LOCTEXT("DMXFixtureTypesViewTooltip", "Show the fixture types view");
+}
+
+TSharedRef<SWidget> FDMXEditorFixtureTypesSummoner::CreateTabBody(const FWorkflowTabSpawnInfo& Info) const
+{
+	TSharedPtr<FDMXEditor> DMXEditor = StaticCastSharedPtr<FDMXEditor>(HostingApp.Pin());
+
+	return SNew(SBox)
+		.HAlign(HAlign_Fill)
+		.VAlign(VAlign_Fill)
+		[
+			DMXEditor->GetFixtureTypesTab()
+		];
+}
+
+FDMXEditorFixturePatchSummoner::FDMXEditorFixturePatchSummoner(TSharedPtr<FAssetEditorToolkit> InHostingApp)
+	: FDMXEditorPropertyTabSummoner(FDMXEditorTabs::DMXFixturePatchEditorTabId, InHostingApp)
+{
+	TabLabel = LOCTEXT("DMXFixturePatchTabLabel", "Fixture Patch");
+	TabIcon = FSlateIcon(FEditorStyle::GetStyleSetName(), "Kismet.Tabs.BlueprintDefaults");
+
+	bIsSingleton = true;
+
+	ViewMenuDescription = LOCTEXT("DMXFixturePatchView", "Fixture Patch");
+	ViewMenuTooltip = LOCTEXT("DMXFixturePatchViewTooltip", "Show the fixture patch view");
+}
+
+TSharedRef<SWidget> FDMXEditorFixturePatchSummoner::CreateTabBody(const FWorkflowTabSpawnInfo& Info) const
+{
+	TSharedPtr<FDMXEditor> DMXEditor = StaticCastSharedPtr<FDMXEditor>(HostingApp.Pin());
+
+	return SNew(SBox)
+		.HAlign(HAlign_Fill)
+		.VAlign(VAlign_Fill)
+		[
+			DMXEditor->GetFixturePatchTab()
+		];
+}
+
+FDMXEditorInputConsoleSummoner::FDMXEditorInputConsoleSummoner(TSharedPtr<FAssetEditorToolkit> InHostingApp)
+	: FWorkflowTabFactory(FDMXEditorTabs::DMXInputConsoleEditorTabId, InHostingApp)
+{
+	TabLabel = LOCTEXT("DMXInputConsoleTabLabel", "Monitor");
+	TabIcon = FSlateIcon(FEditorStyle::GetStyleSetName(), "Kismet.Tabs.BlueprintDefaults");
+
+	bIsSingleton = true;
+
+	ViewMenuDescription = LOCTEXT("DMXInputConsoleView", "Monitor");
+	ViewMenuTooltip = LOCTEXT("DMXInputConsoleViewTooltip", "Show the monitor view");
+}
+
+TSharedRef<SWidget> FDMXEditorInputConsoleSummoner::CreateTabBody(const FWorkflowTabSpawnInfo& Info) const
+{
+	TSharedPtr<FDMXEditor> DMXEditor = StaticCastSharedPtr<FDMXEditor>(HostingApp.Pin());
+
+	return SNew(SBox)
+		.HAlign(HAlign_Fill)
+		.VAlign(VAlign_Fill)
+		[
+			DMXEditor->GetInputConsoleTab()
+		];
+}
+
+FDMXEditorOutputConsoleSummoner::FDMXEditorOutputConsoleSummoner(TSharedPtr<FAssetEditorToolkit> InHostingApp)
+	: FWorkflowTabFactory(FDMXEditorTabs::DMXOutputConsoleEditorTabId, InHostingApp)
+{
+	TabLabel = LOCTEXT("DMXOutputConsoleTabLabel", "Output Console");
+	TabIcon = FSlateIcon(FEditorStyle::GetStyleSetName(), "Kismet.Tabs.BlueprintDefaults");
+
+	bIsSingleton = true;
+
+	ViewMenuDescription = LOCTEXT("DMXOutputConsoleView", "Output Console");
+	ViewMenuTooltip = LOCTEXT("DMXOutputConsoleViewTooltip", "Show the output console view");
+}
+
+TSharedRef<SWidget> FDMXEditorOutputConsoleSummoner::CreateTabBody(const FWorkflowTabSpawnInfo& Info) const
+{
+	TSharedPtr<FDMXEditor> DMXEditor = StaticCastSharedPtr<FDMXEditor>(HostingApp.Pin());
+
+	return DMXEditor->GetOutputConsoleTab();
+}
+
+#undef LOCTEXT_NAMESPACE
+
