- DMX Editor

diff --git a/Engine/Plugins/Experimental/DigitalMultiplex/DMXEngine/Source/DMXEditor/Private/Tests/DMXEditor.spec.cpp b/Engine/Plugins/Experimental/DigitalMultiplex/DMXEngine/Source/DMXEditor/Private/Tests/DMXEditor.spec.cpp
new file mode 100644
--- /dev/null
+++ b/Engine/Plugins/Experimental/DigitalMultiplex/DMXEngine/Source/DMXEditor/Private/Tests/DMXEditor.spec.cpp
@@ -0,0 +1,575 @@
+// Copyright Epic Games, Inc. All Rights Reserved.
+
+#include "Misc/AutomationTest.h"
+
+#include "Factories/DMXEditorFactoryNew.h"
+#include "Library/DMXLibrary.h"
+#include "Interfaces/IDMXProtocol.h"
+#include "Interfaces/IDMXProtocolUniverse.h"
+#include "DMXProtocolArtNetModule.h"
+#include "DMXProtocolSACNModule.h"
+#include "Library/DMXLibrary.h"
+#include "Library/DMXEntityFader.h"
+#include "Library/DMXEntityController.h"
+#include "DMXEditor.h"
+#include "DMXProtocolTypes.h"
+#include "DMXEditorUtils.h"
+
+#include "Widgets/OutputFader/SDMXOutputFaderList.h"
+#include "Widgets/OutputFader/SDMXFader.h"
+#include "Widgets/SDMXInputConsole.h"
+#include "Widgets/SDMXInputInfoSelecter.h"
+#include "Widgets/SDMXInputInfo.h"
+#include "Widgets/Common/SSpinBoxVertical.h"
+
+#include "Widgets/Input/SComboBox.h"
+#include "Widgets/Input/SSpinBox.h"
+
+#include "Utils.h"
+#include "Dom/JsonObject.h"
+#include "Engine/Engine.h"
+#include "AssetData.h"
+#include "Editor.h"
+
+#include "AssetRegistryModule.h"
+#include "AssetRegistryState.h"
+#include "AssetData.h"
+
+#if WITH_DEV_AUTOMATION_TESTS
+
+struct FDMXEditorHelper
+{
+	static void ResetUniverses()
+	{
+		const TMap<FName, TSharedPtr<IDMXProtocol>>& Protocols = IDMXProtocol::GetProtocols();
+		for (TMap<FName, TSharedPtr<IDMXProtocol>>::TConstIterator ProtocolIt = Protocols.CreateConstIterator(); ProtocolIt; ++ProtocolIt)
+		{
+			if (TSharedPtr<IDMXProtocol> ProtocolPtr = ProtocolIt->Value)
+			{
+				ProtocolPtr->RemoveAllUniverses();
+			}
+		}
+	}
+
+	FDMXEditorHelper()
+	{
+		DMXEditor = MakeShared<FDMXEditor>();
+		UPackage* Package = GetTransientPackage();
+		UDMXEditorFactoryNew* Factory = NewObject<UDMXEditorFactoryNew>(Package, MakeUniqueObjectName(GetTransientPackage(), UDMXEditorFactoryNew::StaticClass()));
+		const FName NewLibraryName = MakeUniqueObjectName(Package, UDMXLibrary::StaticClass(), *FString::Printf(TEXT("%d_LIB"), FMath::RandRange(0, 1000)));
+
+		DMXLibrary = ImportObject<UDMXLibrary>(Package, *NewLibraryName.ToString(), RF_Public | RF_Standalone | RF_Transient, *NewLibraryName.ToString(), nullptr, Factory);
+
+		FName NewTemplateName = MakeUniqueObjectName(DMXLibrary, UDMXEntityFader::StaticClass());
+		NewFaderTemplate = FDMXEditorUtils::CreateFaderTemplate(DMXLibrary);
+		FaderEntity = Cast<UDMXEntityFader>(DMXLibrary->GetOrCreateEntityObject(TEXT(""), UDMXEntityFader::StaticClass()));
+		FaderEntity->SetName(NewFaderTemplate->GetDisplayName());
+	}
+
+	TSharedPtr<FDMXEditor> DMXEditor;
+	UDMXLibrary* DMXLibrary;
+	UDMXEntityFader* FaderEntity;
+	UDMXEntityFader* NewFaderTemplate;
+};
+
+
+static const uint8 TestChannelValue = 155;
+static const uint8 ExistingUniverse = 1;
+static const uint8 NonExistingUniverse = 12;
+static const uint8 ChannelValue = 50;
+
+IMPLEMENT_SIMPLE_AUTOMATION_TEST(FDMXEditorFadersArtNetTest, "VirtualProduction.DMX.Editor.Faders.ArtNet",
+	EAutomationTestFlags::EditorContext | EAutomationTestFlags::EngineFilter);
+
+bool FDMXEditorFadersArtNetTest::RunTest(const FString& Parameters)
+{
+	FDMXEditorHelper::ResetUniverses();
+
+	// Create Helper
+	TSharedPtr<FDMXEditorHelper> Helper = MakeShared<FDMXEditorHelper>();
+
+	TSharedPtr<IDMXProtocol> DMXProtocol = IDMXProtocol::Get(FDMXProtocolArtNetModule::NAME_Artnet);
+
+	// Add first universe
+	{
+		FDMXUniverse DMXUniverse;
+		DMXUniverse.Channel = 10;
+		DMXUniverse.UniverseNumber = 1;
+		DMXUniverse.DMXProtocolDirectionality = EDMXProtocolDirectionality::EOutput;
+		Helper->FaderEntity->Universes.Add(DMXUniverse);
+	}
+
+	// Add second universe
+	{
+		FDMXUniverse DMXUniverse;
+		DMXUniverse.Channel = 20;
+		DMXUniverse.UniverseNumber = 2;
+		DMXUniverse.DMXProtocolDirectionality = EDMXProtocolDirectionality::EOutput;
+		Helper->FaderEntity->Universes.Add(DMXUniverse);
+	}
+	Helper->FaderEntity->DeviceProtocol = FDMXProtocolName(FDMXProtocolArtNetModule::NAME_Artnet);
+	Helper->FaderEntity->PostEditChange();
+
+	Helper->DMXEditor->InitEditor(EToolkitMode::Standalone, nullptr, Helper->DMXLibrary);
+
+	TSharedPtr<SDMXOutputFaderList> DMXOutputFaderList = SNew(SDMXOutputFaderList)
+		.DMXEditor(Helper->DMXEditor)
+		.FaderTemplate(Helper->NewFaderTemplate);
+
+	TestTrue(TEXT("Should be only one fader now"), DMXOutputFaderList->GetFaderWidgets().Num() == 1);
+
+	DMXOutputFaderList->GetFaderWidgets()[0]->GetFaderBoxVertical()->SetValue(ChannelValue);
+
+	// Check the data inside universe buffers
+	TArray<TSharedPtr<IDMXProtocolUniverse, ESPMode::ThreadSafe>> Universes;
+	Universes.Add(DMXProtocol->GetUniverseById(Helper->FaderEntity->Universes[0].UniverseNumber));
+	Universes.Add(DMXProtocol->GetUniverseById(Helper->FaderEntity->Universes[1].UniverseNumber));
+
+	TestTrue(TEXT("Universe should be exists"),
+		Universes[0].IsValid());
+	TestTrue(TEXT("Universe should be exists"),
+		Universes[1].IsValid());
+
+	TestEqual(TEXT("Buffer value should be same"),
+		Universes[0]->GetOutputDMXBuffer()->GetDMXData()[Helper->FaderEntity->Universes[0].Channel], ChannelValue);
+	TestEqual(TEXT("Buffer value should be same"),
+		Universes[1]->GetOutputDMXBuffer()->GetDMXData()[Helper->FaderEntity->Universes[1].Channel], ChannelValue);
+
+	Helper->DMXEditor->CloseWindow();
+
+	FDMXEditorHelper::ResetUniverses();
+
+	return true;
+}
+
+IMPLEMENT_SIMPLE_AUTOMATION_TEST(FDMXEditorFadersSACNTest, "VirtualProduction.DMX.Editor.Faders.sACN",
+	EAutomationTestFlags::EditorContext | EAutomationTestFlags::EngineFilter);
+
+bool FDMXEditorFadersSACNTest::RunTest(const FString& Parameters)
+{
+	FDMXEditorHelper::ResetUniverses();
+
+	// Create Helper
+	TSharedPtr<FDMXEditorHelper> Helper = MakeShared<FDMXEditorHelper>();
+
+	TSharedPtr<IDMXProtocol> DMXProtocol = IDMXProtocol::Get(FDMXProtocolSACNModule::NAME_SACN);
+
+	// Add first universe
+	{
+		FDMXUniverse DMXUniverse;
+		DMXUniverse.Channel = 10;
+		DMXUniverse.UniverseNumber = 31;
+		DMXUniverse.DMXProtocolDirectionality = EDMXProtocolDirectionality::EOutput;
+		Helper->FaderEntity->Universes.Add(DMXUniverse);
+	}
+
+	// Add second universe
+	{
+		FDMXUniverse DMXUniverse;
+		DMXUniverse.Channel = 20;
+		DMXUniverse.UniverseNumber = 51;
+		DMXUniverse.DMXProtocolDirectionality = EDMXProtocolDirectionality::EOutput;
+		Helper->FaderEntity->Universes.Add(DMXUniverse);
+	}
+
+	Helper->FaderEntity->DeviceProtocol = FDMXProtocolName(FDMXProtocolSACNModule::NAME_SACN);
+	Helper->FaderEntity->PostEditChange();
+
+	Helper->DMXEditor->InitEditor(EToolkitMode::Standalone, nullptr, Helper->DMXLibrary);
+
+	TSharedPtr<SDMXOutputFaderList> DMXOutputFaderList = SNew(SDMXOutputFaderList)
+		.DMXEditor(Helper->DMXEditor)
+		.FaderTemplate(Helper->NewFaderTemplate);
+
+	TestTrue(TEXT("Should be only one fader now"), DMXOutputFaderList->GetFaderWidgets().Num() == 1);
+
+	DMXOutputFaderList->GetFaderWidgets()[0]->GetFaderBoxVertical()->SetValue(ChannelValue);
+
+	// Check the data inside universe buffers
+	TArray<TSharedPtr<IDMXProtocolUniverse, ESPMode::ThreadSafe>> Universes;
+	Universes.Add(DMXProtocol->GetUniverseById(Helper->FaderEntity->Universes[0].UniverseNumber));
+	Universes.Add(DMXProtocol->GetUniverseById(Helper->FaderEntity->Universes[1].UniverseNumber));
+
+	TestTrue(TEXT("Universe should be exists"),
+		Universes[0].IsValid());
+	TestTrue(TEXT("Universe should be exists"),
+		Universes[1].IsValid());
+
+	TestEqual(TEXT("Buffer value should be same"),
+		Universes[0]->GetOutputDMXBuffer()->GetDMXData()[Helper->FaderEntity->Universes[0].Channel], ChannelValue);
+	TestEqual(TEXT("Buffer value should be same"),
+		Universes[1]->GetOutputDMXBuffer()->GetDMXData()[Helper->FaderEntity->Universes[1].Channel], ChannelValue);
+
+	Helper->DMXEditor->CloseWindow();
+
+	FDMXEditorHelper::ResetUniverses();
+
+	return true;
+}
+
+IMPLEMENT_SIMPLE_AUTOMATION_TEST(FDMXEditorControllersEmptyTest, "VirtualProduction.DMX.Editor.Controllers.Empty",
+	EAutomationTestFlags::EditorContext | EAutomationTestFlags::EngineFilter);
+
+bool FDMXEditorControllersEmptyTest::RunTest(const FString& Parameters)
+{
+	FDMXEditorHelper::ResetUniverses();
+
+	FDMXEditorHelper Helpers[] =
+	{
+		FDMXEditorHelper(),
+		FDMXEditorHelper()
+	};
+
+	Helpers[0].DMXEditor->InitEditor(EToolkitMode::Standalone, nullptr, Helpers[0].DMXLibrary);
+	Helpers[1].DMXEditor->InitEditor(EToolkitMode::Standalone, nullptr, Helpers[1].DMXLibrary);
+
+	const TMap<FName, TSharedPtr<IDMXProtocol>>& Protocols = IDMXProtocol::GetProtocols();
+	for (TMap<FName, TSharedPtr<IDMXProtocol>>::TConstIterator ProtocolIt = Protocols.CreateConstIterator(); ProtocolIt; ++ProtocolIt)
+	{
+		if (TSharedPtr<IDMXProtocol> ProtocolPtr = ProtocolIt->Value)
+		{
+			TestEqual(TEXT("Verify number of universes Is 0"), 0, ProtocolPtr->GetUniversesNum());
+		}
+	}
+
+	// Check our files
+	static const FName AssetRegistryName(TEXT("AssetRegistry"));
+	FAssetRegistryModule& AssetRegistryModule = FModuleManager::LoadModuleChecked<FAssetRegistryModule>(AssetRegistryName);
+	TArray<FAssetData> Items;
+	AssetRegistryModule.Get().GetAssetsByClass(UDMXLibrary::StaticClass()->GetFName(), Items);
+
+	uint32 OurAssetsCount = 0;
+	if (Items.Num() > 0)
+	{
+		for (const FAssetData& Item : Items)
+		{
+			UObject* Asset = FindObject<UObject>(nullptr, *Item.ObjectPath.ToString());
+			if (Asset != nullptr)
+			{
+				if (Asset == Helpers[0].DMXLibrary || Asset == Helpers[1].DMXLibrary)
+				{
+					OurAssetsCount++;
+				}
+			}
+		}
+	}
+	TestEqual(TEXT("Verify number of assets Is 0"), 0, OurAssetsCount);
+
+	Helpers[0].DMXEditor->CloseWindow();
+	Helpers[1].DMXEditor->CloseWindow();
+
+	FDMXEditorHelper::ResetUniverses();
+
+	return true;
+}
+
+IMPLEMENT_SIMPLE_AUTOMATION_TEST(FDMXEditorControllersNonEmptyTest, "VirtualProduction.DMX.Editor.Controllers.NonEmpty",
+	EAutomationTestFlags::EditorContext | EAutomationTestFlags::EngineFilter);
+
+bool FDMXEditorControllersNonEmptyTest::RunTest(const FString& Parameters)
+{
+	FDMXEditorHelper::ResetUniverses();
+
+	FDMXEditorHelper Helpers[] =
+	{
+		FDMXEditorHelper(),
+		FDMXEditorHelper()
+	};
+
+
+	// Add one Universe for editor 0
+	{
+		FName NewLibraryName = MakeUniqueObjectName(GetTransientPackage(), UDMXLibrary::StaticClass());
+		UDMXEntityController* Controller = Cast<UDMXEntityController>(Helpers[0].DMXLibrary->GetOrCreateEntityObject(NewLibraryName.ToString(), UDMXEntityController::StaticClass()));
+		FDMXUniverse Universe;
+		Universe.UniverseNumber = 2;
+		Controller->DeviceProtocol = FDMXProtocolName(FDMXProtocolArtNetModule::NAME_Artnet);
+		Controller->Universes.Add(Universe);
+		Controller->PostEditChange();
+	}
+
+	// Add same Universe for editor 0
+	{
+		FName NewLibraryName = MakeUniqueObjectName(GetTransientPackage(), UDMXLibrary::StaticClass());
+		UDMXEntityController* Controller = Cast<UDMXEntityController>(Helpers[0].DMXLibrary->GetOrCreateEntityObject(NewLibraryName.ToString(), UDMXEntityController::StaticClass()));
+		FDMXUniverse Universe;
+		Universe.UniverseNumber = 2;
+		Controller->DeviceProtocol = FDMXProtocolName(FDMXProtocolArtNetModule::NAME_Artnet);
+		Controller->Universes.Add(Universe);
+		Controller->PostEditChange();
+	}
+
+	// Add Universe for editor 1
+	{
+		FName NewLibraryName = MakeUniqueObjectName(GetTransientPackage(), UDMXLibrary::StaticClass());
+		UDMXEntityController* Controller = Cast<UDMXEntityController>(Helpers[1].DMXLibrary->GetOrCreateEntityObject(NewLibraryName.ToString(), UDMXEntityController::StaticClass()));
+		FDMXUniverse Universe;
+		Universe.UniverseNumber = 675;
+		Controller->DeviceProtocol = FDMXProtocolName(FDMXProtocolSACNModule::NAME_SACN);
+		Controller->Universes.Add(Universe);
+		Controller->PostEditChange();
+	}
+
+	// Add same Universe for editor 1
+	{
+		FName NewLibraryName = MakeUniqueObjectName(GetTransientPackage(), UDMXLibrary::StaticClass());
+		UDMXEntityController* Controller = Cast<UDMXEntityController>(Helpers[1].DMXLibrary->GetOrCreateEntityObject(NewLibraryName.ToString(), UDMXEntityController::StaticClass()));
+		FDMXUniverse Universe;
+		Universe.UniverseNumber = 675;
+		Controller->DeviceProtocol = FDMXProtocolName(FDMXProtocolSACNModule::NAME_SACN);
+		Controller->Universes.Add(Universe);
+		Controller->PostEditChange();
+	}
+
+	Helpers[0].DMXEditor->InitEditor(EToolkitMode::Standalone, nullptr, Helpers[0].DMXLibrary);
+	Helpers[1].DMXEditor->InitEditor(EToolkitMode::Standalone, nullptr, Helpers[1].DMXLibrary);
+
+	uint32 NumUniverses = 0;
+	const TMap<FName, TSharedPtr<IDMXProtocol>>& Protocols = IDMXProtocol::GetProtocols();
+	for (TMap<FName, TSharedPtr<IDMXProtocol>>::TConstIterator ProtocolIt = Protocols.CreateConstIterator(); ProtocolIt; ++ProtocolIt)
+	{
+		if (TSharedPtr<IDMXProtocol> ProtocolPtr = ProtocolIt->Value)
+		{
+			NumUniverses += ProtocolPtr->GetUniversesNum();
+		}
+	}
+
+	TestEqual(TEXT("Verify number of universes Is 2"), 2, NumUniverses);
+
+	Helpers[0].DMXEditor->CloseWindow();
+	Helpers[1].DMXEditor->CloseWindow();
+
+	FDMXEditorHelper::ResetUniverses();
+
+	return true;
+}
+
+IMPLEMENT_SIMPLE_AUTOMATION_TEST(FDMXEditorInputConsoleArtNetExistingTest, "VirtualProduction.DMX.Editor.InputConsole.ArtNet.Existing",
+	EAutomationTestFlags::EditorContext | EAutomationTestFlags::EngineFilter);
+
+bool FDMXEditorInputConsoleArtNetExistingTest::RunTest(const FString& Parameters)
+{
+	// Reset Universes
+	FDMXEditorHelper::ResetUniverses();
+
+	// Create Helper
+	TSharedPtr<FDMXEditorHelper> Helper = MakeShared<FDMXEditorHelper>();
+
+	// 2. Create Universe
+	FName NewLibraryName = MakeUniqueObjectName(GetTransientPackage(), UDMXLibrary::StaticClass());
+	UDMXEntityController* Controller = Cast<UDMXEntityController>(Helper->DMXLibrary->GetOrCreateEntityObject(NewLibraryName.ToString(), UDMXEntityController::StaticClass()));
+	FDMXUniverse Universe;
+	Universe.UniverseNumber = ExistingUniverse;
+	Controller->DeviceProtocol = FDMXProtocolName(FDMXProtocolArtNetModule::NAME_Artnet);
+	Controller->Universes.Add(Universe);
+	Controller->PostEditChange();
+
+	// 3. Create editor
+	Helper->DMXEditor->InitEditor(EToolkitMode::Standalone, nullptr, Helper->DMXLibrary);
+
+	// 4. Set Input Tab
+	TSharedRef<SDMXInputConsole> InputConsole = Helper->DMXEditor->GetInputConsoleTab();
+	const TSharedRef<SDMXInputInfo>& InputInfo = InputConsole->GetInputInfo();
+	const TSharedRef<SDMXInputInfoSelecter>& InfoSelecter = InputConsole->GetInputInfoSelecter();
+	const TSharedRef<SSpinBox<uint16>>& UniverseField = InfoSelecter->GetUniverseField();
+
+	InfoSelecter->SetProtocol(FDMXProtocolArtNetModule::NAME_Artnet);
+	UniverseField->SetValue(ExistingUniverse);
+
+	// 5. Send DMX
+	TSharedPtr<IDMXProtocol> DMXProtocol = IDMXProtocol::Get(FDMXProtocolArtNetModule::NAME_Artnet);
+	IDMXFragmentMap FragmentMap;
+	FragmentMap.Add(0, TestChannelValue);
+	DMXProtocol->SendDMXFragment(ExistingUniverse, FragmentMap);
+
+	// Check the input console values
+	AddCommand(new FDelayedFunctionLatentCommand([=] {
+		TSharedRef<SDMXInputConsole> InputConsole = Helper->DMXEditor->GetInputConsoleTab();
+		const TSharedPtr<SDMXInputInfo>& InputInfo = InputConsole->GetInputInfo();
+
+		// 6. Check input console
+		const TArray<uint8>& ChannelsValues = InputInfo->GetChannelsValues();
+		TestEqual(TEXT("Verify ChannelsValue"), TestChannelValue, ChannelsValues[0]);
+
+		FDMXEditorHelper::ResetUniverses();
+
+		// 7. Close the Editor
+		Helper->DMXEditor->CloseWindow();
+
+	}, 0.2f));
+
+	return true;
+}
+
+IMPLEMENT_SIMPLE_AUTOMATION_TEST(FDMXEditorInputConsoleArtNetNonExistingTest, "VirtualProduction.DMX.Editor.InputConsole.ArtNet.NonExisting",
+	EAutomationTestFlags::EditorContext | EAutomationTestFlags::EngineFilter);
+
+bool FDMXEditorInputConsoleArtNetNonExistingTest::RunTest(const FString& Parameters)
+{
+	// Reset Universes
+	FDMXEditorHelper::ResetUniverses();
+
+	// Create Helper
+	TSharedPtr<FDMXEditorHelper> Helper = MakeShared<FDMXEditorHelper>();
+
+	// 2. Create Universe
+	FName NewLibraryName = MakeUniqueObjectName(GetTransientPackage(), UDMXLibrary::StaticClass());
+	UDMXEntityController* Controller = Cast<UDMXEntityController>(Helper->DMXLibrary->GetOrCreateEntityObject(NewLibraryName.ToString(), UDMXEntityController::StaticClass()));
+	FDMXUniverse Universe;
+	Universe.UniverseNumber = ExistingUniverse;
+	Controller->DeviceProtocol = FDMXProtocolName(FDMXProtocolArtNetModule::NAME_Artnet);
+	Controller->Universes.Add(Universe);
+	Controller->PostEditChange();
+
+	// 3. Create editor
+	Helper->DMXEditor->InitEditor(EToolkitMode::Standalone, nullptr, Helper->DMXLibrary);
+
+	// 4. Set Input Tab
+	TSharedRef<SDMXInputConsole> InputConsole = Helper->DMXEditor->GetInputConsoleTab();
+	const TSharedRef<SDMXInputInfo>& InputInfo = InputConsole->GetInputInfo();
+	const TSharedRef<SDMXInputInfoSelecter>& InfoSelecter = InputConsole->GetInputInfoSelecter();
+	const TSharedRef<SSpinBox<uint16>>& UniverseField = InfoSelecter->GetUniverseField();
+
+	InfoSelecter->SetProtocol(FDMXProtocolArtNetModule::NAME_Artnet);
+	UniverseField->SetValue(NonExistingUniverse);
+
+	// 5. Send DMX
+	TSharedPtr<IDMXProtocol> DMXProtocol = IDMXProtocol::Get(FDMXProtocolArtNetModule::NAME_Artnet);
+	IDMXFragmentMap FragmentMap;
+	FragmentMap.Add(0, TestChannelValue);
+	DMXProtocol->SendDMXFragment(ExistingUniverse, FragmentMap);
+
+	// Check the input console values
+	AddCommand(new FDelayedFunctionLatentCommand([=] {
+		TSharedRef<SDMXInputConsole> InputConsole = Helper->DMXEditor->GetInputConsoleTab();
+		const TSharedPtr<SDMXInputInfo>& InputInfo = InputConsole->GetInputInfo();
+
+		// 6. Check input console
+		const TArray<uint8>& ChannelsValues = InputInfo->GetChannelsValues();
+		TestNotEqual(TEXT("Verify ChannelsValue"), TestChannelValue, ChannelsValues[0]);
+
+		FDMXEditorHelper::ResetUniverses();
+
+		// 7. Close the Editor
+		Helper->DMXEditor->CloseWindow();
+
+	}, 0.2f));
+
+	return true;
+}
+
+IMPLEMENT_SIMPLE_AUTOMATION_TEST(FDMXEditorInputConsoleSACNExistingTest, "VirtualProduction.DMX.Editor.InputConsole.SACN.Existing",
+	EAutomationTestFlags::EditorContext | EAutomationTestFlags::EngineFilter);
+
+bool FDMXEditorInputConsoleSACNExistingTest::RunTest(const FString& Parameters)
+{
+	// Reset Universes
+	FDMXEditorHelper::ResetUniverses();
+
+	// Create Helper
+	TSharedPtr<FDMXEditorHelper> Helper = MakeShared<FDMXEditorHelper>();
+
+	// 2. Create Universe
+	FName NewLibraryName = MakeUniqueObjectName(GetTransientPackage(), UDMXLibrary::StaticClass());
+	UDMXEntityController* Controller = Cast<UDMXEntityController>(Helper->DMXLibrary->GetOrCreateEntityObject(NewLibraryName.ToString(), UDMXEntityController::StaticClass()));
+	FDMXUniverse Universe;
+	Universe.UniverseNumber = ExistingUniverse;
+	Controller->DeviceProtocol = FDMXProtocolName(FDMXProtocolSACNModule::NAME_SACN);
+	Controller->Universes.Add(Universe);
+	Controller->PostEditChange();
+
+	// 3. Create editor
+	Helper->DMXEditor->InitEditor(EToolkitMode::Standalone, nullptr, Helper->DMXLibrary);
+
+	// 4. Set Input Tab
+	TSharedRef<SDMXInputConsole> InputConsole = Helper->DMXEditor->GetInputConsoleTab();
+	const TSharedRef<SDMXInputInfo>& InputInfo = InputConsole->GetInputInfo();
+	const TSharedRef<SDMXInputInfoSelecter>& InfoSelecter = InputConsole->GetInputInfoSelecter();
+	const TSharedRef<SSpinBox<uint16>>& UniverseField = InfoSelecter->GetUniverseField();
+
+	InfoSelecter->SetProtocol(FDMXProtocolSACNModule::NAME_SACN);
+	UniverseField->SetValue(ExistingUniverse);
+
+	// 5. Send DMX
+	TSharedPtr<IDMXProtocol> DMXProtocol = IDMXProtocol::Get(FDMXProtocolSACNModule::NAME_SACN);
+	IDMXFragmentMap FragmentMap;
+	FragmentMap.Add(0, TestChannelValue);
+	DMXProtocol->SendDMXFragment(ExistingUniverse, FragmentMap);
+
+	// Check the input console values
+	AddCommand(new FDelayedFunctionLatentCommand([=] {
+		TSharedRef<SDMXInputConsole> InputConsole = Helper->DMXEditor->GetInputConsoleTab();
+		const TSharedPtr<SDMXInputInfo>& InputInfo = InputConsole->GetInputInfo();
+
+		// 6. Check input console
+		const TArray<uint8>& ChannelsValues = InputInfo->GetChannelsValues();
+		TestEqual(TEXT("Verify ChannelsValue"), TestChannelValue, ChannelsValues[0]);
+
+		FDMXEditorHelper::ResetUniverses();
+
+		// 7. Close the Editor
+		Helper->DMXEditor->CloseWindow();
+
+	}, 0.2f));
+
+	return true;
+}
+
+IMPLEMENT_SIMPLE_AUTOMATION_TEST(FDMXEditorInputConsoleSACNNonExistingTest, "VirtualProduction.DMX.Editor.InputConsole.SACN.NonExisting",
+	EAutomationTestFlags::EditorContext | EAutomationTestFlags::EngineFilter);
+
+bool FDMXEditorInputConsoleSACNNonExistingTest::RunTest(const FString& Parameters)
+{
+	// Reset Universes
+	FDMXEditorHelper::ResetUniverses();
+
+	// Create Helper
+	TSharedPtr<FDMXEditorHelper> Helper = MakeShared<FDMXEditorHelper>();
+
+	// 2. Create Universe
+	FName NewLibraryName = MakeUniqueObjectName(GetTransientPackage(), UDMXLibrary::StaticClass());
+	UDMXEntityController* Controller = Cast<UDMXEntityController>(Helper->DMXLibrary->GetOrCreateEntityObject(NewLibraryName.ToString(), UDMXEntityController::StaticClass()));
+	FDMXUniverse Universe;
+	Universe.UniverseNumber = ExistingUniverse;
+	Controller->DeviceProtocol = FDMXProtocolName(FDMXProtocolSACNModule::NAME_SACN);
+	Controller->Universes.Add(Universe);
+	Controller->PostEditChange();
+
+	// 3. Create editor
+	Helper->DMXEditor->InitEditor(EToolkitMode::Standalone, nullptr, Helper->DMXLibrary);
+
+	// 4. Set Input Tab
+	TSharedRef<SDMXInputConsole> InputConsole = Helper->DMXEditor->GetInputConsoleTab();
+	const TSharedRef<SDMXInputInfo>& InputInfo = InputConsole->GetInputInfo();
+	const TSharedRef<SDMXInputInfoSelecter>& InfoSelecter = InputConsole->GetInputInfoSelecter();
+	const TSharedRef<SSpinBox<uint16>>& UniverseField = InfoSelecter->GetUniverseField();
+
+	InfoSelecter->SetProtocol(FDMXProtocolSACNModule::NAME_SACN);
+	UniverseField->SetValue(NonExistingUniverse);
+
+	// 5. Send DMX
+	TSharedPtr<IDMXProtocol> DMXProtocol = IDMXProtocol::Get(FDMXProtocolSACNModule::NAME_SACN);
+	IDMXFragmentMap FragmentMap;
+	FragmentMap.Add(0, TestChannelValue);
+	DMXProtocol->SendDMXFragment(ExistingUniverse, FragmentMap);
+
+	// Check the input console values
+	AddCommand(new FDelayedFunctionLatentCommand([=] {
+		TSharedRef<SDMXInputConsole> InputConsole = Helper->DMXEditor->GetInputConsoleTab();
+		const TSharedPtr<SDMXInputInfo>& InputInfo = InputConsole->GetInputInfo();
+
+		// 6. Check input console
+		const TArray<uint8>& ChannelsValues = InputInfo->GetChannelsValues();
+		TestNotEqual(TEXT("Verify ChannelsValue"), TestChannelValue, ChannelsValues[0]);
+
+		FDMXEditorHelper::ResetUniverses();
+
+		// 7. Close the Editor
+		Helper->DMXEditor->CloseWindow();
+
+	}, 0.2f));
+
+	return true;
+}
+
+#endif
\ No newline at end of file
diff --git a/Engine/Plugins/Experimental/DigitalMultiplex/DMXEngine/Source/DMXEditor/Private/Toolbars/DMXEditorToolbar.cpp b/Engine/Plugins/Experimental/DigitalMultiplex/DMXEngine/Source/DMXEditor/Private/Toolbars/DMXEditorToolbar.cpp
new file mode 100644
--- /dev/null
+++ b/Engine/Plugins/Experimental/DigitalMultiplex/DMXEngine/Source/DMXEditor/Private/Toolbars/DMXEditorToolbar.cpp
@@ -0,0 +1,37 @@
+// Copyright Epic Games, Inc. All Rights Reserved.
+
+#include "Toolbars/DMXEditorToolbar.h"
+#include "DMXEditor.h"
+#include "Framework/MultiBox/MultiBoxBuilder.h"
+
+#define LOCTEXT_NAMESPACE "KismetToolbar"
+
+void FDMXEditorToolbar::AddCompileToolbar(TSharedPtr<FExtender> Extender)
+{
+	TSharedPtr<FDMXEditor> DMXEditorPtr = DMXEditor.Pin();
+
+	Extender->AddToolBarExtension(
+		"Asset",
+		EExtensionHook::After,
+		DMXEditorPtr->GetToolkitCommands(),
+		FToolBarExtensionDelegate::CreateSP(this, &FDMXEditorToolbar::FillCompileToolbar));
+}
+
+FSlateIcon FDMXEditorToolbar::GetStatusImage() const
+{
+	return FSlateIcon(FEditorStyle::GetStyleSetName(), "Kismet.Status.Good");
+}
+
+FText FDMXEditorToolbar::GetStatusTooltip() const
+{
+	return LOCTEXT("Default_Status", "Good to go");
+}
+
+void FDMXEditorToolbar::FillCompileToolbar(FToolBarBuilder& ToolbarBuilder)
+{
+	// TODO. Implement custom toolbar
+	ToolbarBuilder.BeginSection("CompileToolbar");
+	ToolbarBuilder.EndSection();
+}
+
+#undef LOCTEXT_NAMESPACE
\ No newline at end of file
diff --git a/Engine/Plugins/Experimental/DigitalMultiplex/DMXEngine/Source/DMXEditor/Private/Widgets/Common/SSpinBoxVertical.cpp b/Engine/Plugins/Experimental/DigitalMultiplex/DMXEngine/Source/DMXEditor/Private/Widgets/Common/SSpinBoxVertical.cpp
new file mode 100644
--- /dev/null
+++ b/Engine/Plugins/Experimental/DigitalMultiplex/DMXEngine/Source/DMXEditor/Private/Widgets/Common/SSpinBoxVertical.cpp
@@ -0,0 +1,22 @@
+// Copyright Epic Games, Inc. All Rights Reserved.
+#include "Widgets/Common/SSpinBoxVertical.h"
+//#include "SlatePrivatePCH.h"
+
+float SpinBoxVertComputeExponentSliderFraction(float FractionFilled, float StartFractionFilled, float SliderExponent)
+{
+	if (FractionFilled <= StartFractionFilled)
+	{
+		float DeltaFraction = (StartFractionFilled - FractionFilled) / StartFractionFilled;
+		float LeftFractionFilled = 1.0f - FMath::Pow(1.0f - DeltaFraction, SliderExponent);
+		FractionFilled = StartFractionFilled - (StartFractionFilled*LeftFractionFilled);
+	}
+	else
+	{
+		float DeltaFraction = (FractionFilled - StartFractionFilled) / (1.0f - StartFractionFilled);
+		float RightFractionFilled = 1.0f - FMath::Pow(1.0f - DeltaFraction, SliderExponent);
+		FractionFilled = StartFractionFilled + (1.0f - StartFractionFilled) * RightFractionFilled;
+	}
+	return FractionFilled;
+}
+
+
diff --git a/Engine/Plugins/Experimental/DigitalMultiplex/DMXEngine/Source/DMXEditor/Private/Widgets/Common/SSpinBoxVertical.h b/Engine/Plugins/Experimental/DigitalMultiplex/DMXEngine/Source/DMXEditor/Private/Widgets/Common/SSpinBoxVertical.h
new file mode 100644
--- /dev/null
+++ b/Engine/Plugins/Experimental/DigitalMultiplex/DMXEngine/Source/DMXEditor/Private/Widgets/Common/SSpinBoxVertical.h
@@ -0,0 +1,1065 @@
+// Copyright Epic Games, Inc. All Rights Reserved.
+
+#pragma once
+
+#include "CoreMinimal.h"
+#include "EditorStyleSet.h"
+#include "Fonts/SlateFontInfo.h"
+#include "Framework/Application/SlateApplication.h"
+#include "Framework/MultiBox/MultiBoxExtender.h"
+#include "Input/CursorReply.h"
+#include "Input/Events.h"
+#include "Input/Reply.h"
+#include "InputCoreTypes.h"
+#include "Layout/Margin.h"
+#include "Layout/Visibility.h"
+#include "Misc/Attribute.h"
+#include "Rendering/DrawElements.h"
+#include "Styling/CoreStyle.h"
+#include "Styling/SlateColor.h"
+#include "Styling/SlateTypes.h"
+#include "Templates/IsIntegral.h"
+#include "Widgets/DeclarativeSyntaxSupport.h"
+#include "Widgets/Images/SImage.h"
+#include "Widgets/Input/NumericTypeInterface.h"
+#include "Widgets/Input/SEditableText.h"
+#include "Widgets/Layout/SBox.h"
+#include "Widgets/SBoxPanel.h"
+#include "Widgets/SCompoundWidget.h"
+#include "Widgets/Text/STextBlock.h"
+
+/*
+ * This function compute a slider position by simulating two log on both side of the neutral value
+ * Example a slider going from 0.0 to 2.0 with a neutral value of 1.0, the user will have a lot of precision around the neutral value
+ * on both side.
+ |
+ ||                              |
+ | -_                          _-
+ |   --__                  __--
+ |       ----__________----
+ ----------------------------------
+  0              1               2
+
+  The function return a float representing the slider fraction used to position the slider handle
+  FractionFilled: this is the value slider position with no exponent
+  StartFractionFilled: this is the neutral value slider position with no exponent
+  SliderExponent: this is the slider exponent
+*/
+float SpinBoxVertComputeExponentSliderFraction(float FractionFilled, float StartFractionFilled, float SliderExponent);
+
+/**
+ * A Slate SpinBox resembles traditional spin boxes in that it is a widget that provides
+ * keyboard-based and mouse-based manipulation of a numeric value.
+ * Mouse-based manipulation: drag anywhere on the spinbox to change the value.
+ * Keyboard-based manipulation: click on the spinbox to enter text mode.
+ */
+template<typename NumericType>
+class SSpinBoxVertical
+	: public SCompoundWidget
+{
+public:
+
+	/** Notification for numeric value change */
+	DECLARE_DELEGATE_OneParam( FOnValueChanged, NumericType );
+
+	/** Notification for numeric value committed */
+	DECLARE_DELEGATE_TwoParams( FOnValueCommitted, NumericType, ETextCommit::Type);
+
+	/** Notification when the max/min spinner values are changed (only apply if SupportDynamicSliderMaxValue or SupportDynamicSliderMinValue are true) */
+	DECLARE_DELEGATE_FourParams(FOnDynamicSliderMinMaxValueChanged, NumericType, TWeakPtr<SWidget>, bool, bool);
+
+	SLATE_BEGIN_ARGS(SSpinBoxVertical<NumericType>)
+		: _Style(&FCoreStyle::Get().GetWidgetStyle<FSpinBoxStyle>("SpinBox"))
+		, _Value(0)
+		, _MinValue(0)
+		, _MaxValue(10)
+		, _Delta(0)
+		, _ShiftMouseMovePixelPerDelta(1)
+		, _SupportDynamicSliderMaxValue(false)
+		, _SupportDynamicSliderMinValue(false)
+		, _SliderExponent(1)
+		, _Font( FCoreStyle::Get().GetFontStyle( TEXT( "NormalFont" ) ) )
+		, _ContentPadding(  FMargin( 2.0f, 1.0f) )
+		, _OnValueChanged()
+		, _OnValueCommitted()
+		, _ClearKeyboardFocusOnCommit( false )
+		, _SelectAllTextOnCommit( true )
+		, _MinDesiredWidth(0.0f)
+		{}
+	
+		/** The style used to draw this spinbox */
+		SLATE_STYLE_ARGUMENT( FSpinBoxStyle, Style )
+
+		/** The value to display */
+		SLATE_ATTRIBUTE( NumericType, Value )
+		/** The minimum value that can be entered into the text edit box */
+		SLATE_ATTRIBUTE( TOptional< NumericType >, MinValue )
+		/** The maximum value that can be entered into the text edit box */
+		SLATE_ATTRIBUTE( TOptional< NumericType >, MaxValue )
+		/** The minimum value that can be specified by using the slider, defaults to MinValue */
+		SLATE_ATTRIBUTE( TOptional< NumericType >, MinSliderValue )
+		/** The maximum value that can be specified by using the slider, defaults to MaxValue */
+		SLATE_ATTRIBUTE( TOptional< NumericType >, MaxSliderValue )
+		/** Delta to increment the value as the slider moves.  If not specified will determine automatically */
+		SLATE_ATTRIBUTE( NumericType, Delta )
+		/** How many pixel the mouse must move to change the value of the delta step */
+		SLATE_ATTRIBUTE( int32, ShiftMouseMovePixelPerDelta )
+		/** If we're an unbounded spinbox, what value do we divide mouse movement by before multiplying by Delta. Requires Delta to be set. */
+		SLATE_ATTRIBUTE( int32, LinearDeltaSensitivity)
+		/** Tell us if we want to support dynamically changing of the max value using ctrl */
+		SLATE_ATTRIBUTE(bool, SupportDynamicSliderMaxValue)
+		/** Tell us if we want to support dynamically changing of the min value using ctrl */
+		SLATE_ATTRIBUTE(bool, SupportDynamicSliderMinValue)
+		/** Called right after the max slider value is changed (only relevant if SupportDynamicSliderMaxValue is true) */
+		SLATE_EVENT(FOnDynamicSliderMinMaxValueChanged, OnDynamicSliderMaxValueChanged)
+		/** Called right after the min slider value is changed (only relevant if SupportDynamicSliderMinValue is true) */
+		SLATE_EVENT(FOnDynamicSliderMinMaxValueChanged, OnDynamicSliderMinValueChanged)
+		/** Use exponential scale for the slider */
+		SLATE_ATTRIBUTE( float, SliderExponent )
+		/** When use exponential scale for the slider which is the neutral value */
+		SLATE_ATTRIBUTE( NumericType, SliderExponentNeutralValue )
+		/** Font used to display text in the slider */
+		SLATE_ATTRIBUTE( FSlateFontInfo, Font )
+		/** Padding to add around this widget and its internal widgets */
+		SLATE_ATTRIBUTE( FMargin, ContentPadding )
+		/** Called when the value is changed by slider or typing */
+		SLATE_EVENT( FOnValueChanged, OnValueChanged )
+		/** Called when the value is committed (by pressing enter) */
+		SLATE_EVENT( FOnValueCommitted, OnValueCommitted )
+		/** Called right before the slider begins to move */
+		SLATE_EVENT( FSimpleDelegate, OnBeginSliderMovement )
+		/** Called right after the slider handle is released by the user */
+		SLATE_EVENT( FOnValueChanged, OnEndSliderMovement )
+		/** Whether to clear keyboard focus when pressing enter to commit changes */
+		SLATE_ATTRIBUTE( bool, ClearKeyboardFocusOnCommit )
+		/** Whether to select all text when pressing enter to commit changes */
+		SLATE_ATTRIBUTE( bool, SelectAllTextOnCommit )
+		/** Minimum width that a spin box should be */
+		SLATE_ATTRIBUTE( float, MinDesiredWidth )
+		/** How should the value be justified in the spinbox. */
+		SLATE_ATTRIBUTE( ETextJustify::Type, Justification )
+		/** Provide custom type conversion functionality to this spin box */
+		SLATE_ARGUMENT( TSharedPtr< INumericTypeInterface<NumericType> >, TypeInterface )
+		/** If refresh requests for the viewport should happen for all value changes **/
+		SLATE_ARGUMENT(bool, PreventThrottling)
+		/** Menu extender for the right-click context menu */
+		SLATE_EVENT( FMenuExtensionDelegate, ContextMenuExtender )
+
+	SLATE_END_ARGS()
+
+	SSpinBoxVertical()
+	{
+	}
+		
+	/**
+	 * Construct the widget
+	 * 
+	 * @param InArgs   A declaration from which to construct the widget
+	 */
+	void Construct( const FArguments& InArgs )
+	{
+		check(InArgs._Style);
+
+		Style = InArgs._Style;
+
+		ForegroundColor = InArgs._Style->ForegroundColor;
+		Interface = InArgs._TypeInterface.IsValid() ? InArgs._TypeInterface : MakeShared<TDefaultNumericTypeInterface<NumericType>>();
+
+		ValueAttribute = InArgs._Value;
+		OnValueChanged = InArgs._OnValueChanged;
+		OnValueCommitted = InArgs._OnValueCommitted;
+		OnBeginSliderMovement = InArgs._OnBeginSliderMovement;
+		OnEndSliderMovement = InArgs._OnEndSliderMovement;
+		MinDesiredWidth = InArgs._MinDesiredWidth;
+	
+		MinValue = InArgs._MinValue;
+		MaxValue = InArgs._MaxValue;
+		MinSliderValue = (InArgs._MinSliderValue.Get().IsSet()) ? InArgs._MinSliderValue : MinValue;
+		MaxSliderValue = (InArgs._MaxSliderValue.Get().IsSet()) ? InArgs._MaxSliderValue : MaxValue;
+
+		SupportDynamicSliderMaxValue = InArgs._SupportDynamicSliderMaxValue;
+		SupportDynamicSliderMinValue = InArgs._SupportDynamicSliderMinValue;
+		OnDynamicSliderMaxValueChanged = InArgs._OnDynamicSliderMaxValueChanged;
+		OnDynamicSliderMinValueChanged = InArgs._OnDynamicSliderMinValueChanged;
+
+		bPreventThrottling = InArgs._PreventThrottling;
+
+		// Update the max slider value based on the current value if we're in dynamic mode
+		NumericType CurrentMaxValue = GetMaxValue();
+		NumericType CurrentMinValue = GetMinValue();
+
+		if (SupportDynamicSliderMaxValue.Get() && ValueAttribute.Get() > GetMaxSliderValue())
+		{			
+			ApplySliderMaxValueChanged(ValueAttribute.Get() - GetMaxSliderValue(), true);
+		}
+		else if (SupportDynamicSliderMinValue.Get() && ValueAttribute.Get() < GetMinSliderValue())
+		{
+			ApplySliderMinValueChanged(ValueAttribute.Get() - GetMinSliderValue(), true);
+		}
+
+		UpdateIsSpinRangeUnlimited();
+	
+		SliderExponent = InArgs._SliderExponent;
+
+		SliderExponentNeutralValue = InArgs._SliderExponentNeutralValue;
+
+		DistanceDragged = 0.0f;
+		PreDragValue = 0;
+
+		Delta = InArgs._Delta;
+		ShiftMouseMovePixelPerDelta = InArgs._ShiftMouseMovePixelPerDelta;
+		LinearDeltaSensitivity = InArgs._LinearDeltaSensitivity;
+	
+		BackgroundHoveredBrush = &InArgs._Style->HoveredBackgroundBrush;
+		BackgroundBrush = &InArgs._Style->BackgroundBrush;
+		ActiveFillBrush = &InArgs._Style->ActiveFillBrush;
+		InactiveFillBrush = &InArgs._Style->InactiveFillBrush;
+		const FMargin& TextMargin = InArgs._Style->TextPadding;
+
+		bDragging = false;
+		PointerDraggingSliderIndex = INDEX_NONE;
+
+		CachedExternalValue = ValueAttribute.Get();
+		InternalValue = ValueAttribute.Get();
+
+		bIsTextChanging = false;
+
+		this->ChildSlot
+		.Padding( InArgs._ContentPadding )
+		[
+			SNew(SBox)
+			.HeightOverride(100)
+			.WidthOverride(25)
+			[
+				SNew(SVerticalBox)
+
+				+ SVerticalBox::Slot()
+			
+				.AutoHeight()
+				.Padding( TextMargin )
+				.HAlign(HAlign_Fill) 
+				.VAlign(VAlign_Center)
+				[
+					SNew(SHorizontalBox)
+
+					+SHorizontalBox::Slot()
+					.AutoWidth()
+					.HAlign(HAlign_Fill)
+					.VAlign(VAlign_Center)
+					[
+						SAssignNew(TextBlock, STextBlock)
+						.Font(InArgs._Font)
+						.Text( this, &SSpinBoxVertical<NumericType>::GetValueAsText )
+						.MinDesiredWidth( this, &SSpinBoxVertical<NumericType>::GetTextMinDesiredWidth )
+						.Justification(InArgs._Justification)
+					]
+				]
+
+				+ SVerticalBox::Slot()
+				.AutoHeight()
+				.Padding( TextMargin )
+				.HAlign(HAlign_Fill) 
+				.VAlign(VAlign_Center) 
+				[
+					SNew(SHorizontalBox)
+
+					+ SHorizontalBox::Slot()
+					.AutoWidth()
+					.HAlign(HAlign_Fill)
+					.VAlign(VAlign_Center)
+					[
+						SAssignNew(EditableText, SEditableText)
+						.Visibility( EVisibility::Collapsed )
+						.Font(InArgs._Font)
+						.SelectAllTextWhenFocused( true )
+						.Text( this, &SSpinBoxVertical<NumericType>::GetValueAsText )
+						.OnIsTypedCharValid(this, &SSpinBoxVertical<NumericType>::IsCharacterValid)
+						.OnTextChanged( this, &SSpinBoxVertical<NumericType>::TextField_OnTextChanged )
+						.OnTextCommitted( this, &SSpinBoxVertical<NumericType>::TextField_OnTextCommitted )
+						.ClearKeyboardFocusOnCommit( InArgs._ClearKeyboardFocusOnCommit )
+						.SelectAllTextOnCommit( InArgs._SelectAllTextOnCommit )
+						.MinDesiredWidth( this, &SSpinBoxVertical<NumericType>::GetTextMinDesiredWidth )
+						.VirtualKeyboardType(EKeyboardType::Keyboard_Number)
+						.Justification(InArgs._Justification)
+						.VirtualKeyboardTrigger(EVirtualKeyboardTrigger::OnAllFocusEvents)
+						.ContextMenuExtender( InArgs._ContextMenuExtender )
+					]
+				]
+
+				+SVerticalBox::Slot()
+				.FillHeight(1.0f)
+				.HAlign(HAlign_Fill) 
+				.VAlign(VAlign_Bottom) 
+				[
+					SNew(SImage)
+					.Image( &InArgs._Style->ArrowsImage )
+					.ColorAndOpacity( FSlateColor::UseForeground())
+				]
+			]
+		];
+	}
+	
+	virtual int32 OnPaint( const FPaintArgs& Args, const FGeometry& AllottedGeometry, const FSlateRect& MyCullingRect, FSlateWindowElementList& OutDrawElements, int32 LayerId, const FWidgetStyle& InWidgetStyle, bool bParentEnabled ) const override
+	{
+		const bool bActiveFeedback = IsHovered() || bDragging;
+
+		const FSlateBrush* BackgroundImage = bActiveFeedback ?
+			BackgroundHoveredBrush :
+			BackgroundBrush;
+
+		const FSlateBrush* FillImage = bActiveFeedback ?
+			ActiveFillBrush :
+			InactiveFillBrush;
+
+		const int32 BackgroundLayer = LayerId;
+
+		const bool bEnabled = ShouldBeEnabled( bParentEnabled );
+		const ESlateDrawEffect DrawEffects = bEnabled ? ESlateDrawEffect::None : ESlateDrawEffect::DisabledEffect;
+
+		FSlateDrawElement::MakeBox(
+			OutDrawElements,
+			BackgroundLayer,
+			AllottedGeometry.ToPaintGeometry(),
+			BackgroundImage,
+			//FEditorStyle::GetBrush("TextBlock.HighlightShape"),
+			DrawEffects,
+			BackgroundImage->GetTint(InWidgetStyle) * InWidgetStyle.GetColorAndOpacityTint()
+			);
+
+		const int32 FilledLayer = BackgroundLayer + 1;
+
+		//if there is a spin range limit, draw the filler bar
+		if (!bUnlimitedSpinRange)
+		{
+			double Value = ValueAttribute.Get();
+			NumericType CurrentDelta = Delta.Get();
+			if( CurrentDelta != 0.0f )
+			{
+				Value= Snap(Value, CurrentDelta); // snap floating point value to nearest Delta
+			}
+
+			float FractionFilled = Fraction(Value, GetMinSliderValue(), GetMaxSliderValue());
+			const float CachedSliderExponent = SliderExponent.Get();
+			if (CachedSliderExponent != 1)
+			{
+				if (SliderExponentNeutralValue.IsSet() && SliderExponentNeutralValue.Get() > GetMinSliderValue() && SliderExponentNeutralValue.Get() < GetMaxSliderValue())
+				{
+					//Compute a log curve on both side of the neutral value
+					float StartFractionFilled = Fraction(SliderExponentNeutralValue.Get(), GetMinSliderValue(), GetMaxSliderValue());
+					FractionFilled = SpinBoxVertComputeExponentSliderFraction(FractionFilled, StartFractionFilled, CachedSliderExponent);
+				}
+				else
+				{
+					FractionFilled = 1.0f - FMath::Pow( 1.0f - FractionFilled, CachedSliderExponent);
+				}
+			}
+			const FVector2D FillSize( AllottedGeometry.GetLocalSize().X, AllottedGeometry.GetLocalSize().Y * -FractionFilled);
+
+			if ( ! IsInTextMode() )
+			{
+				FSlateDrawElement::MakeBox(
+					OutDrawElements,
+					FilledLayer,
+					AllottedGeometry.ToPaintGeometry(FVector2D(0, AllottedGeometry.GetLocalSize().Y), FillSize),
+					FEditorStyle::GetBrush("TextBlock.HighlightShape"),
+					DrawEffects
+					//llImage->GetTint(InWidgetStyle) * InWidgetStyle.GetColorAndOpacityTint()
+					);
+			}
+		}
+
+		return FMath::Max( FilledLayer, SCompoundWidget::OnPaint(Args, AllottedGeometry, MyCullingRect, OutDrawElements, FilledLayer, InWidgetStyle, bEnabled ) );
+	}
+
+	/**
+	 * The system calls this method to notify the widget that a mouse button was pressed within it. This event is bubbled.
+	 *
+	 * @param MyGeometry The Geometry of the widget receiving the event
+	 * @param MouseEvent Information about the input event
+	 * @return Whether the event was handled along with possible requests for the system to take action.
+	 */
+	virtual FReply OnMouseButtonDown( const FGeometry& MyGeometry, const FPointerEvent& MouseEvent ) override
+	{
+		if ( MouseEvent.GetEffectingButton() == EKeys::LeftMouseButton && PointerDraggingSliderIndex == INDEX_NONE )
+		{
+			DistanceDragged = 0;
+			InternalValue = ValueAttribute.Get();
+			PreDragValue = InternalValue;
+			PointerDraggingSliderIndex = MouseEvent.GetPointerIndex();
+			CachedMousePosition = MouseEvent.GetScreenSpacePosition().IntPoint();
+
+			FReply ReturnReply = FReply::Handled().CaptureMouse(SharedThis(this)).UseHighPrecisionMouseMovement(SharedThis(this)).SetUserFocus(SharedThis(this), EFocusCause::Mouse);
+			if (bPreventThrottling) 
+			{
+				ReturnReply.PreventThrottling();
+			}
+			return ReturnReply;
+		}
+		else
+		{
+
+			return FReply::Unhandled();
+		}
+	}
+	
+	/**
+	 * The system calls this method to notify the widget that a mouse button was release within it. This event is bubbled.
+	 *
+	 * @param MyGeometry The Geometry of the widget receiving the event
+	 * @param MouseEvent Information about the input event
+	 * @return Whether the event was handled along with possible requests for the system to take action.
+	 */
+	virtual FReply OnMouseButtonUp( const FGeometry& MyGeometry, const FPointerEvent& MouseEvent ) override
+	{
+		if ( MouseEvent.GetEffectingButton() == EKeys::LeftMouseButton && PointerDraggingSliderIndex == MouseEvent.GetPointerIndex() )
+		{
+			if(!this->HasMouseCapture())
+			{
+				// Lost Capture - ensure reset
+				bDragging = false;
+				PointerDraggingSliderIndex = INDEX_NONE;
+				
+				return FReply::Unhandled();
+			}
+		
+			if( bDragging )
+			{
+				NumericType CurrentDelta = Delta.Get();
+				if (CurrentDelta != 0)
+				{
+					InternalValue = Snap(InternalValue, CurrentDelta);
+				}
+
+				NotifyValueCommitted();
+			}
+
+			bDragging = false;
+			PointerDraggingSliderIndex = INDEX_NONE;
+
+			FReply Reply = FReply::Handled().ReleaseMouseCapture();
+
+			if ( !MouseEvent.IsTouchEvent() )
+			{
+				Reply.SetMousePos(CachedMousePosition);
+			}
+
+			if ( DistanceDragged < FSlateApplication::Get().GetDragTriggerDistance() )
+			{
+				EnterTextMode();
+				Reply.SetUserFocus(EditableText.ToSharedRef(), EFocusCause::SetDirectly);
+			}
+
+			return Reply;
+		}
+		
+		return FReply::Unhandled();
+	}
+
+	void ApplySliderMaxValueChanged(float SliderDeltaToAdd, bool UpdateOnlyIfHigher)
+	{
+		check(SupportDynamicSliderMaxValue.Get());
+
+		NumericType NewMaxSliderValue = TNumericLimits<NumericType>::Min();
+		
+		if (MaxSliderValue.IsSet() && MaxSliderValue.Get().IsSet())
+		{
+			NewMaxSliderValue = GetMaxSliderValue();
+
+			if ((NewMaxSliderValue + SliderDeltaToAdd > GetMaxSliderValue() && UpdateOnlyIfHigher) || !UpdateOnlyIfHigher)
+			{
+				NewMaxSliderValue += SliderDeltaToAdd;
+
+				if (!MaxSliderValue.IsBound()) // simple value so we can update it without breaking the mechanic otherwise it must be handle by the callback implementer
+				{
+					SetMaxSliderValue(NewMaxSliderValue);
+				}
+			}
+		}
+
+		if (OnDynamicSliderMaxValueChanged.IsBound())
+		{
+			OnDynamicSliderMaxValueChanged.Execute(NewMaxSliderValue, TWeakPtr<SWidget>(AsShared()), true, UpdateOnlyIfHigher);
+		}
+	}
+
+	void ApplySliderMinValueChanged(float SliderDeltaToAdd, bool UpdateOnlyIfLower)
+	{
+		check(SupportDynamicSliderMaxValue.Get());
+
+		NumericType NewMinSliderValue = TNumericLimits<NumericType>::Min();
+		
+		if (MinSliderValue.IsSet() && MinSliderValue.Get().IsSet())
+		{
+			NewMinSliderValue = GetMinSliderValue();
+
+			if ((NewMinSliderValue + SliderDeltaToAdd < GetMinSliderValue() && UpdateOnlyIfLower) || !UpdateOnlyIfLower)
+			{
+				NewMinSliderValue += SliderDeltaToAdd;
+
+				if (!MinSliderValue.IsBound()) // simple value so we can update it without breaking the mechanic otherwise it must be handle by the callback implementer
+				{
+					SetMinSliderValue(NewMinSliderValue);
+				}
+			}
+		}		
+
+		if (OnDynamicSliderMinValueChanged.IsBound())
+		{
+			OnDynamicSliderMinValueChanged.Execute(NewMinSliderValue, TWeakPtr<SWidget>(AsShared()), true, UpdateOnlyIfLower);
+		}
+	}
+	
+	/**
+	 * The system calls this method to notify the widget that a mouse moved within it. This event is bubbled.
+	 *
+	 * @param MyGeometry The Geometry of the widget receiving the event
+	 * @param MouseEvent Information about the input event
+	 * @return Whether the event was handled along with possible requests for the system to take action.
+	 */
+	virtual FReply OnMouseMove( const FGeometry& MyGeometry, const FPointerEvent& MouseEvent ) override
+	{
+		if ( PointerDraggingSliderIndex == MouseEvent.GetPointerIndex() )
+		{
+			if(!this->HasMouseCapture())
+			{
+				// Lost the mouse capture - ensure reset
+				bDragging = false;
+				PointerDraggingSliderIndex = INDEX_NONE;
+				
+				return FReply::Unhandled();
+			}
+			
+			if (!bDragging)
+			{
+				DistanceDragged += FMath::Abs(MouseEvent.GetCursorDelta().Y);
+				if ( DistanceDragged > FSlateApplication::Get().GetDragTriggerDistance() )
+				{
+					ExitTextMode();
+					bDragging = true;
+					OnBeginSliderMovement.ExecuteIfBound();
+				}
+
+				// Cache the mouse, even if not dragging cache it.
+				CachedMousePosition = MouseEvent.GetScreenSpacePosition().IntPoint();
+			}
+			else
+			{
+				double NewValue = 0;
+
+				// Increments the spin based on delta mouse movement.
+
+				// A minimum slider width to use for calculating deltas in the slider-range space
+				const float MinSliderWidth = 100.f;
+				float SliderWidthInSlateUnits = FMath::Max(MyGeometry.GetDrawSize().X, MinSliderWidth);
+				
+				const int32 CachedShiftMouseMovePixelPerDelta = ShiftMouseMovePixelPerDelta.Get();
+				if (CachedShiftMouseMovePixelPerDelta > 1 && MouseEvent.IsShiftDown())
+				{
+					SliderWidthInSlateUnits *= CachedShiftMouseMovePixelPerDelta;
+				}
+
+				if (MouseEvent.IsControlDown())
+				{
+					float DeltaToAdd = MouseEvent.GetCursorDelta().Y / SliderWidthInSlateUnits;
+
+					if (SupportDynamicSliderMaxValue.Get() && InternalValue == GetMaxSliderValue())
+					{
+						ApplySliderMaxValueChanged(DeltaToAdd, false);
+					}
+					else if (SupportDynamicSliderMinValue.Get() && InternalValue == GetMinSliderValue())
+					{
+						ApplySliderMinValueChanged(DeltaToAdd, false);
+					}
+				}
+				
+				//if we have a range to draw in
+				if ( !bUnlimitedSpinRange) 
+				{
+					bool HasValidExponentNeutralValue = SliderExponentNeutralValue.IsSet() && SliderExponentNeutralValue.Get() > GetMinSliderValue() && SliderExponentNeutralValue.Get() < GetMaxSliderValue();
+
+					const float CachedSliderExponent = SliderExponent.Get();
+					// The amount currently filled in the spinbox, needs to be calculated to do deltas correctly.
+					float FractionFilled = Fraction(InternalValue, GetMinSliderValue(), GetMaxSliderValue());
+						
+					if (CachedSliderExponent != 1)
+					{
+						if (HasValidExponentNeutralValue)
+						{
+							//Compute a log curve on both side of the neutral value
+							float StartFractionFilled = Fraction(SliderExponentNeutralValue.Get(), GetMinSliderValue(), GetMaxSliderValue());
+							FractionFilled = SpinBoxVertComputeExponentSliderFraction(FractionFilled, StartFractionFilled, CachedSliderExponent);
+						}
+						else
+						{
+							FractionFilled = 1.0f - FMath::Pow( 1.0f - FractionFilled, CachedSliderExponent);
+						}
+					}
+					FractionFilled *= SliderWidthInSlateUnits;
+
+					// Now add the delta to the fraction filled, this causes the spin.
+					FractionFilled -= MouseEvent.GetCursorDelta().Y;
+						
+					// Clamp the fraction to be within the bounds of the geometry.
+					FractionFilled = FMath::Clamp(FractionFilled, 0.0f, SliderWidthInSlateUnits);
+
+					// Convert the fraction filled to a percent.
+					float Percent = FMath::Clamp(FractionFilled / SliderWidthInSlateUnits, 0.0f, 1.0f);
+					if (CachedSliderExponent != 1)
+					{
+						// Have to convert the percent to the proper value due to the exponent component to the spin.
+						if (HasValidExponentNeutralValue)
+						{
+							//Compute a log curve on both side of the neutral value
+							float StartFractionFilled = Fraction(SliderExponentNeutralValue.Get(), GetMinSliderValue(), GetMaxSliderValue());
+							Percent = SpinBoxVertComputeExponentSliderFraction(Percent, StartFractionFilled, 1.0/CachedSliderExponent);
+						}
+						else
+						{
+							Percent = 1.0f - FMath::Pow(1.0f - Percent, 1.0 / CachedSliderExponent);
+						}
+						
+						
+					}
+
+					NewValue = FMath::LerpStable<double>(GetMinSliderValue(), GetMaxSliderValue(), Percent);
+				}
+				else
+				{
+					// If this control has a specified delta and sensitivity then we use that instead of the current value for determining how much to change.
+					const float Sign = (MouseEvent.GetCursorDelta().Y > 0) ? 1.f : -1.f;
+					if (LinearDeltaSensitivity.IsSet() && Delta.IsSet() && Delta.Get() > 0)
+					{
+						const float MouseDelta = FMath::Abs(MouseEvent.GetCursorDelta().Y / LinearDeltaSensitivity.Get());
+						NewValue = InternalValue + (Sign * MouseDelta * FMath::Pow(Delta.Get(), SliderExponent.Get()));
+					}
+					else
+					{
+						const float MouseDelta = FMath::Abs(MouseEvent.GetCursorDelta().Y / SliderWidthInSlateUnits);
+						const double CurrentValue = FMath::Clamp<double>(FMath::Abs(InternalValue),1.0,TNumericLimits<NumericType>::Max());
+						NewValue = InternalValue + (Sign * MouseDelta * FMath::Pow(CurrentValue, SliderExponent.Get()));
+					}
+				}
+			
+				CommitValue( NewValue, CommittedViaSpin, ETextCommit::OnEnter );
+			}
+
+			return FReply::Handled();
+		}
+
+		return FReply::Unhandled();
+	}
+
+	virtual FCursorReply OnCursorQuery( const FGeometry& MyGeometry, const FPointerEvent& CursorEvent ) const override
+	{
+		return bDragging ? 
+			FCursorReply::Cursor( EMouseCursor::None ) :
+			FCursorReply::Cursor( EMouseCursor::ResizeUpDown );
+	}
+
+	virtual bool SupportsKeyboardFocus() const override
+	{
+		// SSpinBoxVertical is focusable.
+		return true;
+	}
+
+
+	virtual FReply OnFocusReceived( const FGeometry& MyGeometry, const FFocusEvent& InFocusEvent ) override
+	{
+		if ( !bDragging && (InFocusEvent.GetCause() == EFocusCause::Navigation || InFocusEvent.GetCause() == EFocusCause::SetDirectly) )
+		{
+			EnterTextMode();
+			return FReply::Handled().SetUserFocus(EditableText.ToSharedRef(), InFocusEvent.GetCause());
+		}
+		else
+		{
+			return FReply::Unhandled();
+		}
+	}
+	
+	virtual FReply OnKeyDown( const FGeometry& MyGeometry, const FKeyEvent& InKeyEvent ) override
+	{
+		const FKey Key = InKeyEvent.GetKey();
+		if ( Key == EKeys::Escape && HasMouseCapture())
+		{
+			bDragging = false;
+			PointerDraggingSliderIndex = INDEX_NONE;
+
+			InternalValue = PreDragValue;
+			NotifyValueCommitted();
+			return FReply::Handled().ReleaseMouseCapture().SetMousePos(CachedMousePosition);
+		}
+		else if ( Key == EKeys::Up || Key == EKeys::Right )
+		{
+			InternalValue = ValueAttribute.Get();
+			CommitValue( InternalValue + Delta.Get(), CommittedViaArrowKey, ETextCommit::OnEnter );
+			ExitTextMode();
+			return FReply::Handled();	
+		}
+		else if ( Key == EKeys::Down || Key == EKeys::Left )
+		{
+			InternalValue = ValueAttribute.Get();
+			CommitValue( InternalValue - Delta.Get(), CommittedViaArrowKey, ETextCommit::OnEnter );
+			ExitTextMode();
+			return FReply::Handled();
+		}
+		else if ( Key == EKeys::Enter )
+		{
+			InternalValue = ValueAttribute.Get();
+			EnterTextMode();
+			return FReply::Handled().SetUserFocus(EditableText.ToSharedRef(), EFocusCause::Navigation);
+		}
+		else
+		{
+			return FReply::Unhandled();
+		}
+	}
+	
+	virtual bool HasKeyboardFocus() const override
+	{
+		// The spinbox is considered focused when we are typing it text.
+		return SCompoundWidget::HasKeyboardFocus() || (EditableText.IsValid() && EditableText->HasKeyboardFocus());
+	}
+
+	/** Return the Value attribute */
+	TAttribute<NumericType> GetValueAttribute() const { return ValueAttribute; }
+
+	/** See the Value attribute */
+	float GetValue() const { return ValueAttribute.Get(); }
+	void SetValue(const TAttribute<NumericType>& InValueAttribute) 
+	{
+		ValueAttribute = InValueAttribute; 
+		CommitValue(InValueAttribute.Get(), ECommitMethod::CommittedViaTypeIn, ETextCommit::Default);
+	}
+
+	/** See the MinValue attribute */
+	NumericType GetMinValue() const { return MinValue.Get().Get(TNumericLimits<NumericType>::Lowest()); }
+	void SetMinValue(const TAttribute<TOptional<NumericType>>& InMinValue) 
+	{ 
+		MinValue = InMinValue;
+		UpdateIsSpinRangeUnlimited();
+	}
+
+	/** See the MaxValue attribute */
+	NumericType GetMaxValue() const { return MaxValue.Get().Get(TNumericLimits<NumericType>::Max()); }
+	void SetMaxValue(const TAttribute<TOptional<NumericType>>& InMaxValue) 
+	{ 
+		MaxValue = InMaxValue; 
+		UpdateIsSpinRangeUnlimited();
+	}
+
+	/** See the MinSliderValue attribute */
+	bool IsMinSliderValueBound() const { return MinSliderValue.IsBound(); }
+
+	NumericType GetMinSliderValue() const { return MinSliderValue.Get().Get(TNumericLimits<NumericType>::Lowest()); }
+	void SetMinSliderValue(const TAttribute<TOptional<NumericType>>& InMinSliderValue) 
+	{ 
+		MinSliderValue = (InMinSliderValue.Get().IsSet()) ? InMinSliderValue : MinValue;
+		UpdateIsSpinRangeUnlimited();
+	}
+
+	/** See the MaxSliderValue attribute */
+	bool IsMaxSliderValueBound() const { return MaxSliderValue.IsBound(); }
+
+	NumericType GetMaxSliderValue() const { return MaxSliderValue.Get().Get(TNumericLimits<NumericType>::Max()); }
+	void SetMaxSliderValue(const TAttribute<TOptional<NumericType>>& InMaxSliderValue) 
+	{ 
+		MaxSliderValue = (InMaxSliderValue.Get().IsSet()) ? InMaxSliderValue : MaxValue;;
+		UpdateIsSpinRangeUnlimited();
+	}
+
+	/** See the Delta attribute */
+	NumericType GetDelta() const { return Delta.Get(); }
+	void SetDelta(NumericType InDelta) { Delta.Set( InDelta ); }
+
+	/** See the SliderExponent attribute */
+	float GetSliderExponent() const { return SliderExponent.Get(); }
+	void SetSliderExponent(const TAttribute<float>& InSliderExponent) { SliderExponent = InSliderExponent; }
+	
+	/** See the MinDesiredWidth attribute */
+	float GetMinDesiredWidth() const { return SliderExponent.Get(); }
+	void SetMinDesiredWidth(const TAttribute<float>& InMinDesiredWidth) { MinDesiredWidth = InMinDesiredWidth; }
+
+protected:
+	/** Make the spinbox switch to keyboard-based input mode. */
+	void EnterTextMode()
+	{
+		TextBlock->SetVisibility( EVisibility::Collapsed );
+		EditableText->SetVisibility( EVisibility::Visible );
+	}
+	
+	/** Make the spinbox switch to mouse-based input mode. */
+	void ExitTextMode()
+	{
+		TextBlock->SetVisibility( EVisibility::Visible );
+		EditableText->SetVisibility( EVisibility::Collapsed );
+	}
+
+	/** @return the value being observed by the spinbox as a string */
+	FString GetValueAsString() const
+	{
+		return Interface->ToString(ValueAttribute.Get());
+	}
+
+	/** @return the value being observed by the spinbox as FText - todo: spinbox FText support (reimplement me) */
+	FText GetValueAsText() const
+	{
+		return FText::FromString(GetValueAsString());
+	}
+
+	/**
+	 * Invoked when the text in the text field changes
+	 *
+	 * @param NewText		The value of the text in the text field
+	 */
+	void TextField_OnTextChanged( const FText& NewText)
+	{
+		if (!bIsTextChanging)
+		{
+			TGuardValue<bool> TextChangedGuard(bIsTextChanging, true);
+
+			// Validate the text on change, and only accept text up until the first invalid character
+			FString Data = NewText.ToString();
+			int32 NumValidChars = Data.Len();
+
+			for (int32 Index = 0; Index < Data.Len(); ++Index)
+			{
+				if (!Interface->IsCharacterValid(Data[Index]))
+				{
+					NumValidChars = Index;
+					break;
+				}
+			}
+
+			if (NumValidChars < Data.Len())
+			{
+				FString ValidData = NumValidChars > 0 ? Data.Left(NumValidChars) : GetValueAsString();
+				EditableText->SetText(FText::FromString(ValidData));
+			}
+		}
+	}
+	
+	/**
+	 * Invoked when the text field commits its text.
+	 *
+	 * @param NewText		The value of text coming from the editable text field.
+	 * @param CommitInfo	Information about the source of the commit
+	 */
+	void TextField_OnTextCommitted( const FText& NewText, ETextCommit::Type CommitInfo )
+	{
+		if (CommitInfo != ETextCommit::OnEnter)
+		{
+			ExitTextMode();
+		}
+
+		TOptional<NumericType> NewValue = Interface->FromString(NewText.ToString(), ValueAttribute.Get());
+		if (NewValue.IsSet())
+		{
+			CommitValue( NewValue.GetValue(), CommittedViaTypeIn, CommitInfo );
+		}
+	}
+
+
+	/** How user changed the value in the spinbox */
+	enum ECommitMethod
+	{
+		CommittedViaSpin,
+		CommittedViaTypeIn,
+		CommittedViaArrowKey
+	};
+
+	/**
+	 * Call this method when the user's interaction has changed the value
+	 *
+	 * @param NewValue               Value resulting from the user's interaction
+	 * @param CommitMethod           Did the user type in the value or spin to it.
+	 * @param OriginalCommitInfo	 If the user typed in the value, information about the source of the commit
+	 */
+	void CommitValue( double NewValue, ECommitMethod CommitMethod, ETextCommit::Type OriginalCommitInfo )
+	{
+		if( CommitMethod == CommittedViaSpin || CommitMethod == CommittedViaArrowKey )
+		{
+			NewValue = FMath::Clamp<double>( NewValue, GetMinSliderValue(), GetMaxSliderValue() );
+		}
+
+		NewValue = FMath::Clamp<double>( NewValue, GetMinValue(), GetMaxValue() );
+
+		if ( !ValueAttribute.IsBound() )
+		{
+			ValueAttribute.Set( RoundIfIntegerValue( NewValue ) );
+		}
+
+		auto CurrentValue = ValueAttribute.Get();
+
+		// If not in spin mode, there is no need to jump to the value from the external source, continue to use the committed value.
+		if(CommitMethod == CommittedViaSpin)
+		{
+			// This will detect if an external force has changed the value. Internally it will abandon the delta calculated this tick and update the internal value instead.
+			if(CurrentValue != CachedExternalValue)
+			{
+				NewValue = CurrentValue;
+			}
+		}
+
+		// Update the internal value, this needs to be done before rounding.
+		InternalValue = NewValue;
+
+		// If needed, round this value to the delta. Internally the value is not held to the Delta but externally it appears to be.
+		if ( CommitMethod == CommittedViaSpin || CommitMethod == CommittedViaArrowKey )
+		{
+			NumericType CurrentDelta = Delta.Get();
+			if( CurrentDelta != 0 )
+			{
+				NewValue = Snap(NewValue, CurrentDelta); // snap numeric point value to nearest Delta
+			}
+		}		
+
+		// Update the max slider value based on the current value if we're in dynamic mode
+		if (SupportDynamicSliderMaxValue.Get() && ValueAttribute.Get() > GetMaxSliderValue())
+		{
+			ApplySliderMaxValueChanged(ValueAttribute.Get() - GetMaxSliderValue(), true);
+		}
+		else if (SupportDynamicSliderMinValue.Get() && ValueAttribute.Get() < GetMinSliderValue())
+		{
+			ApplySliderMinValueChanged(ValueAttribute.Get() - GetMinSliderValue(), true);
+		}
+
+		if( CommitMethod == CommittedViaTypeIn || CommitMethod == CommittedViaArrowKey )
+		{
+			OnValueCommitted.ExecuteIfBound( RoundIfIntegerValue( NewValue ), OriginalCommitInfo );
+		}
+
+		OnValueChanged.ExecuteIfBound( RoundIfIntegerValue( NewValue ) );
+
+		// Update the cache of the external value to what the user believes the value is now.
+		CachedExternalValue = ValueAttribute.Get();
+
+		// This ensures that dragging is cleared if focus has been removed from this widget in one of the delegate calls, such as when spawning a modal dialog.
+		if ( !this->HasMouseCapture() )
+		{
+			bDragging = false;
+			PointerDraggingSliderIndex = INDEX_NONE;
+		}
+	}
+	
+	void NotifyValueCommitted() const
+	{
+		// The internal value will have been clamped and rounded to the delta at this point, but integer values may still need to be rounded
+		// if the delta is 0.
+		const auto CurrentValue = RoundIfIntegerValue( InternalValue );
+		OnValueCommitted.ExecuteIfBound( CurrentValue, ETextCommit::OnEnter );
+		OnEndSliderMovement.ExecuteIfBound( CurrentValue );
+	}
+
+	/** @return true when we are in keyboard-based input mode; false otherwise */
+	bool IsInTextMode() const
+	{
+		return ( EditableText->GetVisibility() == EVisibility::Visible );
+	}
+
+	/** Calculates range fraction. Possible to use on full numeric range  */
+	static float Fraction(double InValue, NumericType InMinValue, NumericType InMaxValue)
+	{
+		const double HalfMax = InMaxValue*0.5;
+		const double HalfMin = InMinValue*0.5;
+		const double HalfVal = InValue*0.5;
+
+		return (float)FMath::Clamp((HalfVal - HalfMin)/(HalfMax - HalfMin), 0.0, 1.0);
+	}
+
+	/** Snaps value to a nearest grid multiple. Clamps value if it's outside numeric range  */
+	static double Snap(double InValue, double InGrid)
+	{
+		double Result = FMath::GridSnap(InValue, InGrid);
+		return FMath::Clamp<double>(Result, TNumericLimits<NumericType>::Lowest(), TNumericLimits<NumericType>::Max());
+	}
+
+private:
+	TAttribute<NumericType> ValueAttribute;
+	FOnValueChanged OnValueChanged;
+	FOnValueCommitted OnValueCommitted;
+	FSimpleDelegate OnBeginSliderMovement;
+	FOnValueChanged OnEndSliderMovement;
+	TSharedPtr<STextBlock> TextBlock;
+	TSharedPtr<SEditableText> EditableText;
+
+	/** Interface that defines conversion functionality for the templated type */
+	TSharedPtr< INumericTypeInterface<NumericType> > Interface;
+
+	/** True when no range is specified, spinner can be spun indefinitely */
+	bool bUnlimitedSpinRange;
+	void UpdateIsSpinRangeUnlimited()
+	{
+		bUnlimitedSpinRange = !((MinValue.Get().IsSet() && MaxValue.Get().IsSet()) || (MinSliderValue.Get().IsSet() && MaxSliderValue.Get().IsSet()));
+	}
+
+	const FSpinBoxStyle* Style;
+
+	const FSlateBrush* BackgroundHoveredBrush;
+	const FSlateBrush* BackgroundBrush;
+	const FSlateBrush* ActiveFillBrush;
+	const FSlateBrush* InactiveFillBrush;
+
+	float DistanceDragged;
+	TAttribute<NumericType> Delta;
+	TAttribute<int32> ShiftMouseMovePixelPerDelta;
+	TAttribute<int32> LinearDeltaSensitivity;
+	TAttribute<float> SliderExponent;
+	TAttribute<NumericType> SliderExponentNeutralValue;
+	TAttribute< TOptional<NumericType> > MinValue;
+	TAttribute< TOptional<NumericType> > MaxValue;
+	TAttribute< TOptional<NumericType> > MinSliderValue;
+	TAttribute< TOptional<NumericType> > MaxSliderValue;
+	TAttribute<bool> SupportDynamicSliderMaxValue;
+	TAttribute<bool> SupportDynamicSliderMinValue;
+	FOnDynamicSliderMinMaxValueChanged OnDynamicSliderMaxValueChanged;
+	FOnDynamicSliderMinMaxValueChanged OnDynamicSliderMinValueChanged;
+
+	/** Prevents the spinbox from being smaller than desired in certain cases (e.g. when it is empty) */
+	TAttribute<float> MinDesiredWidth;
+	float GetTextMinDesiredWidth() const
+	{
+		return FMath::Max(0.0f, MinDesiredWidth.Get() - Style->ArrowsImage.ImageSize.X);
+	}
+
+	/** Check whether a typed character is valid */
+	bool IsCharacterValid(TCHAR InChar) const
+	{
+		return Interface->IsCharacterValid(InChar);
+	}
+
+	/** Rounds the submitted value to the correct value if it's an integer. */
+	NumericType RoundIfIntegerValue( double ValueToRound ) const
+	{
+		return TIsIntegral<NumericType>::Value
+			? (NumericType)FMath::FloorToDouble( ValueToRound + 0.5 )
+			: (NumericType)ValueToRound;
+	}
+
+	/** Whether the user is dragging the slider */
+	bool bDragging;
+
+	/** Tracks which cursor is currently dragging the slider (e.g., the mouse cursor or a specific finger) */
+	int32 PointerDraggingSliderIndex;
+	
+	/** Cached mouse position to restore after scrolling. */
+	FIntPoint CachedMousePosition;
+
+	/** This value represents what the spinbox believes the value to be, regardless of delta and the user binding to an int. 
+		The spinbox will always count using floats between values, this is important to keep it flowing smoothly and feeling right, 
+		and most importantly not conflicting with the user truncating the value to an int. */
+	double InternalValue;
+
+	/** The state of InternalValue before a drag operation was started */
+	double PreDragValue;
+
+	/** This is the cached value the user believes it to be (usually different due to truncation to an int). Used for identifying 
+		external forces on the spinbox and syncing the internal value to them. Synced when a value is committed to the spinbox. */
+	NumericType CachedExternalValue;
+
+	/** Re-entrant guard for the text changed handler */
+	bool bIsTextChanging;
+
+	/* Holds whether or not to prevent throttling during mouse capture */
+	// When true, the viewport will be updated with every single change to the value during dragging
+	bool bPreventThrottling;
+};
diff --git a/Engine/Plugins/Experimental/DigitalMultiplex/DMXEngine/Source/DMXEditor/Private/Widgets/OutputFader/SDMXFader.cpp b/Engine/Plugins/Experimental/DigitalMultiplex/DMXEngine/Source/DMXEditor/Private/Widgets/OutputFader/SDMXFader.cpp
new file mode 100644
--- /dev/null
+++ b/Engine/Plugins/Experimental/DigitalMultiplex/DMXEngine/Source/DMXEditor/Private/Widgets/OutputFader/SDMXFader.cpp
@@ -0,0 +1,372 @@
+// Copyright Epic Games, Inc. All Rights Reserved.
+
+#include "Widgets/OutputFader/SDMXFader.h"
+
+#include "Interfaces/IDMXProtocol.h"
+#include "DMXEditor.h"
+#include "Library/DMXEntityFader.h"
+#include "DMXEditorLog.h"
+
+#include "Widgets/Common/SSpinBoxVertical.h"
+#include "Widgets/OutputFader/SDMXOutputFaderList.h"
+#include "Widgets/OutputFader/SDMXFaderChannel.h"
+
+#include "Widgets/Input/SCheckBox.h"
+#include "Widgets/Input/SButton.h"
+#include "Widgets/Layout/SScrollBox.h"
+#include "Widgets/Layout/SSeparator.h"
+
+#define LOCTEXT_NAMESPACE "SDMXFader"
+
+void SDMXFader::Construct(const FArguments& InArgs)
+{
+	WeakDMXEditor = InArgs._DMXEditor;
+	CurrentFaderValue = 0;
+
+	ChildSlot
+		.Padding(0.f)
+		[
+			SNew(SBox)
+			.WidthOverride(85)
+			[
+				SNew(SVerticalBox)
+				+ SVerticalBox::Slot()
+				.VAlign(VAlign_Fill)
+				.HAlign(HAlign_Fill)
+				.AutoHeight()
+				[
+					SNew(SHorizontalBox)
+					+ SHorizontalBox::Slot()
+					.VAlign(VAlign_Center)
+					.HAlign(HAlign_Fill)
+					.FillWidth(25)
+					[
+
+						SAssignNew(BackgroundBorder, SBorder)
+						.BorderImage(this, &SDMXFader::GetBorderImage)
+						[
+							SNew(SVerticalBox)
+							+ SVerticalBox::Slot()
+							.VAlign(VAlign_Top)
+							.HAlign(HAlign_Fill)
+							.AutoHeight()
+							[
+			
+								SNew(SVerticalBox)
+		
+								+ SVerticalBox::Slot()
+								.VAlign(VAlign_Top)
+								.HAlign(HAlign_Fill)
+								.AutoHeight()
+								[
+									SAssignNew(CustomFaderLabel, STextBlock)
+									.Justification(ETextJustify::Center)
+									.Text(InArgs._InText)
+								]
+							]
+							+ SVerticalBox::Slot()
+							.VAlign(VAlign_Top)
+							.HAlign(HAlign_Center)
+							.AutoHeight()
+							[
+								SNew(SVerticalBox)
+
+								+ SVerticalBox::Slot()
+								.VAlign(VAlign_Top)
+								.HAlign(HAlign_Center)
+								.Padding(0.f, 5.f, 0.f, 0.f)
+								.AutoHeight()
+								[
+									SAssignNew(FaderBoxVertical, SSpinBoxVertical<uint8>)
+									.MinValue(0).MaxValue(255)
+									.MinSliderValue(0).MaxSliderValue(255)
+									.OnValueChanged(this, &SDMXFader::HandleFaderChanged)
+									.MinDesiredWidth(25)
+							
+								]
+							]
+							+ SVerticalBox::Slot()
+								.VAlign(VAlign_Top)
+								.HAlign(HAlign_Center)
+								.AutoHeight()
+								[
+									SNew(SVerticalBox)
+
+									+ SVerticalBox::Slot()
+								.VAlign(VAlign_Top)
+								.HAlign(HAlign_Center)
+								.Padding(0.f, 5.f, 0.f, 0.f)
+								.AutoHeight()
+								[
+									SAssignNew(SendDMXCheckBox, SCheckBox)
+									.IsChecked(ECheckBoxState::Checked)
+								]
+							]
+							+ SVerticalBox::Slot()
+							.VAlign(VAlign_Top)
+							.HAlign(HAlign_Fill)
+							.AutoHeight()
+							[
+								SNew(SVerticalBox)
+
+								+ SVerticalBox::Slot()
+								.VAlign(VAlign_Top)
+								.HAlign(HAlign_Fill)
+								.Padding(0.f, 5.f, 0.f, 0.f)
+								.AutoHeight()
+								[
+									SNew(SHorizontalBox)
+									+ SHorizontalBox::Slot()
+									.VAlign(VAlign_Center)
+									.HAlign(HAlign_Fill)
+									.FillWidth(1)
+									[
+										SNew(SVerticalBox)
+										+ SVerticalBox::Slot()
+										.VAlign(VAlign_Fill)
+										.HAlign(HAlign_Fill)
+										.AutoHeight()
+										[
+											SNew(STextBlock)
+											.Justification(ETextJustify::Center)
+											.Text(LOCTEXT("Universe", "Uni"))
+										]
+										+SVerticalBox::Slot()
+										.VAlign(VAlign_Fill)
+										.HAlign(HAlign_Fill)
+										.AutoHeight()
+										[
+											SNew(SSeparator)
+											.Orientation(EOrientation::Orient_Horizontal)
+										]
+								
+									]
+									+ SHorizontalBox::Slot()
+									.VAlign(VAlign_Center)
+									.HAlign(HAlign_Fill)
+									.FillWidth(1)
+									[
+										SNew(SVerticalBox)
+
+										+ SVerticalBox::Slot()
+										.VAlign(VAlign_Fill)
+										.HAlign(HAlign_Fill)
+										.AutoHeight()
+										[
+											SNew(STextBlock)
+											.Justification(ETextJustify::Center)
+											.Text(LOCTEXT("Address", "Addr"))
+										]
+										+ SVerticalBox::Slot()
+										.VAlign(VAlign_Fill)
+										.HAlign(HAlign_Fill)
+										.AutoHeight()
+										[
+											SNew(SSeparator)
+											.Orientation(EOrientation::Orient_Horizontal)
+										]
+									]
+								]
+								+ SVerticalBox::Slot()
+								.VAlign(VAlign_Top)
+								.HAlign(HAlign_Fill)
+								.Padding(0.f, 5.f, 0.f, 0.f)
+								.MaxHeight(100)
+								[
+									SAssignNew(FaderChannelSlots, SScrollBox)
+									.Orientation(EOrientation::Orient_Vertical)
+									.ScrollBarAlwaysVisible(false)
+									.ScrollBarThickness(FVector2D(0.0f))
+									.ScrollBarPadding(0.f)
+								]
+							]
+						]
+					]
+					+ SHorizontalBox::Slot()
+					.VAlign(VAlign_Fill)
+					.HAlign(HAlign_Fill)
+					.FillWidth(1)
+					[
+						SNew(SBorder)
+						.BorderImage(FEditorStyle::GetBrush("DetailsView.CategoryMiddle_Hovered"))
+						.Padding(FMargin(0, 5, 0, 5))
+					]
+				]
+				+ SVerticalBox::Slot()
+				.VAlign(VAlign_Fill)
+				.HAlign(HAlign_Center)
+				[
+					SNew(STextBlock)
+					.Justification(ETextJustify::Center)
+					.Text(this, &SDMXFader::GetProtocolText)
+				]
+				+ SVerticalBox::Slot()
+				.VAlign(VAlign_Fill)
+				.HAlign(HAlign_Fill)
+				.AutoHeight()
+				[
+					SNew(SSeparator)
+					.Orientation(EOrientation::Orient_Horizontal)
+				]
+				+ SVerticalBox::Slot()
+				.VAlign(VAlign_Fill)
+				.HAlign(HAlign_Center)
+				.AutoHeight()
+				[
+					SNew(SBox)
+					.HeightOverride(20)
+					.WidthOverride(20)
+					[
+						SNew(SButton)
+						.Text(LOCTEXT("Close", "X"))
+						.HAlign(HAlign_Center)
+						.OnClicked(this, &SDMXFader::HandleRemoveFaderClicked)
+					]
+				]
+			]
+		];
+}
+
+FReply SDMXFader::HandleRemoveFaderClicked()
+{
+	if (TSharedPtr<SDMXOutputFaderList> FaderList = WeakFaderList.Pin())
+	{
+		if (TSharedPtr<SDMXFader> FaderPtr = FaderList->GetWeakSelectedFaderWidget().Pin())
+		{
+			if (FaderPtr.Get() == this)
+			{
+				FaderList->DeselectFaders();
+			}
+		}
+		FaderList->RemoveFader(SharedThis(this));
+	}
+
+	return FReply::Handled();
+}
+
+FReply SDMXFader::OnMouseButtonDown(const FGeometry& MyGeometry, const FPointerEvent& MouseEvent)
+{
+	if (MouseEvent.GetEffectingButton() == EKeys::LeftMouseButton)
+	{
+		SelectThisFader();
+	}
+
+	return FReply::Handled();
+}
+
+void SDMXFader::SetFaderEntity(const TWeakObjectPtr<UDMXEntityFader> InFaderEntity)
+{
+	if (!InFaderEntity.IsValid())
+	{
+		UE_LOG_DMXEDITOR(Warning, TEXT("InFaderWidget is invalid pointer"));
+		return;
+	}
+
+	WeakFaderEntity = InFaderEntity;
+}
+
+void SDMXFader::SetParentFaderList(const TSharedPtr<SDMXOutputFaderList>& InFaderList)
+{
+	WeakFaderList = InFaderList;
+}
+
+void SDMXFader::SetFaderLabel(const FString& InLabel)
+{
+	CustomFaderLabel->SetText(FText::FromString(InLabel));
+}
+
+void SDMXFader::AddChannelWidget(const FString& InUniverse, const FString& InChannel, uint16 InUniverseNumber, uint32 InChannelNumber, const TSharedPtr<IDMXProtocol>& InDMXProtocol)
+{
+	TSharedPtr<SDMXFaderChannel> NewChannel = SNew(SDMXFaderChannel)
+		.DMXEditor(WeakDMXEditor)
+		.UniverseNumber(InUniverseNumber)
+		.ChannelNumber(InChannelNumber)
+		.DMXProtocol(InDMXProtocol);
+
+	FaderChannels.Add(NewChannel);
+
+	FaderChannelSlots->AddSlot()
+		[
+			NewChannel.ToSharedRef()
+		];
+
+	if (NewChannel->GetUniverseValue().IsValid())
+	{
+		NewChannel->GetUniverseValue()->SetText(FText::FromString(InUniverse));
+	}
+	if (NewChannel->GetChannelValue().IsValid())
+	{
+		NewChannel->GetChannelValue()->SetText(FText::FromString(InChannel));
+	}
+}
+
+void SDMXFader::RemoveAllChannelWidgets()
+{
+	FaderChannelSlots->ClearChildren();
+	FaderChannels.Empty();
+}
+
+void SDMXFader::UpdateFaderTemplateProperties()
+{
+	if (TSharedPtr<SDMXOutputFaderList> FaderList = WeakFaderList.Pin())
+	{
+		FaderList->UpdateFaderTemplateObject(WeakFaderEntity);
+	}
+}
+
+void SDMXFader::SelectThisFader()
+{
+	if (TSharedPtr<SDMXOutputFaderList> FaderList = WeakFaderList.Pin())
+	{
+		FaderList->ResetFaderBackgrounds();
+		FaderList->WeakSelectedFaderWidget = SharedThis(this);
+		UpdateFaderTemplateProperties();
+		BackgroundBorder->SetBorderImage(FEditorStyle::GetBrush("DetailsView.CategoryMiddle_Highlighted"));
+	}
+}
+
+const FSlateBrush* SDMXFader::GetBorderImage() const
+{
+	if (IsHovered())
+	{
+		return FEditorStyle::GetBrush("DetailsView.CategoryMiddle_Hovered");
+	}
+	else
+	{
+		return FEditorStyle::GetBrush("DetailsView.CategoryMiddle");
+	}
+}
+
+void SDMXFader::HandleFaderChanged(uint8 NewValue)
+{
+	CurrentFaderValue = NewValue;
+
+	if (SendDMXCheckBox.IsValid() && SendDMXCheckBox->IsChecked())
+	{
+		for (const TSharedPtr<SDMXFaderChannel>& Channel : FaderChannels)
+		{
+			if (WeakFaderEntity.Get())
+			{
+				TSharedPtr<IDMXProtocol> DMXProtocol = WeakFaderEntity->DeviceProtocol;
+				if (DMXProtocol.IsValid())
+				{
+					IDMXFragmentMap DMXFragmentMap;
+					DMXFragmentMap.Add(Channel->GetChannelNumber(), CurrentFaderValue);
+					DMXProtocol->SendDMXFragment(Channel->GetUniverseNumber(), DMXFragmentMap);
+				}
+			}
+		}
+	}
+}
+
+FText SDMXFader::GetProtocolText() const
+{
+	if (WeakFaderEntity.IsValid())
+	{
+		return FText::FromName(WeakFaderEntity->DeviceProtocol);
+	}
+
+	return FText::GetEmpty();
+}
+
+#undef LOCTEXT_NAMESPACE
diff --git a/Engine/Plugins/Experimental/DigitalMultiplex/DMXEngine/Source/DMXEditor/Private/Widgets/OutputFader/SDMXFader.h b/Engine/Plugins/Experimental/DigitalMultiplex/DMXEngine/Source/DMXEditor/Private/Widgets/OutputFader/SDMXFader.h
new file mode 100644
--- /dev/null
+++ b/Engine/Plugins/Experimental/DigitalMultiplex/DMXEngine/Source/DMXEditor/Private/Widgets/OutputFader/SDMXFader.h
@@ -0,0 +1,105 @@
+// Copyright Epic Games, Inc. All Rights Reserved.
+
+#pragma once
+
+#include "CoreMinimal.h"
+#include "Widgets/DeclarativeSyntaxSupport.h"
+#include "Widgets/SCompoundWidget.h"
+#include "Styling/SlateTypes.h"
+
+template<typename NumericType>
+class SSpinBoxVertical;
+class SDMXOutputFaderList;
+class SScrollBox;
+class FDMXEditor;
+class UDMXEntityFader;
+class SBorder;
+class STextBlock;
+class SDMXFaderChannel;
+class SCheckBox;
+class IDMXProtocol;
+
+#define LOCTEXT_NAMESPACE "SDMXFader"
+
+/** Individual fader UI class */
+class SDMXFader
+	: public SCompoundWidget
+{
+public:
+	SLATE_BEGIN_ARGS(SDMXFader)
+		: _InText(LOCTEXT("FaderLabel", "Fader"))
+	{}
+
+	SLATE_ARGUMENT(TWeakPtr<FDMXEditor>, DMXEditor)
+	SLATE_ATTRIBUTE(FText, InText)
+
+	SLATE_END_ARGS()
+
+	/** Constructs the widget */
+	void Construct(const FArguments& InArgs);
+
+	void SetFaderEntity(const TWeakObjectPtr<UDMXEntityFader> InFaderEntity);
+
+	void SetParentFaderList(const TSharedPtr<SDMXOutputFaderList>& InFaderList);
+		
+	void SetFaderLabel(const FString& InLabel);
+
+	void AddChannelWidget(const FString& InUniverse, const FString& InChannel, uint16 InUniverseNumber, uint32 InChannelNumber, const TSharedPtr<IDMXProtocol>& DMXProtocol);
+
+	void RemoveAllChannelWidgets();
+
+	void SelectThisFader();
+
+	const TWeakObjectPtr<UDMXEntityFader>& GetWeakFaderEntity() const { return WeakFaderEntity; }
+
+	const TSharedPtr<SBorder>& GetBackgroundBorder() const { return BackgroundBorder; }
+
+	const TSharedPtr<SSpinBoxVertical<uint8>>& GetFaderBoxVertical() const { return FaderBoxVertical; }
+
+private:
+	//~ Begin SWidget implementation
+	virtual FReply OnMouseButtonDown(const FGeometry& MyGeometry, const FPointerEvent& MouseEvent) override;
+	//~ END SWidget implementation
+
+	FReply HandleRemoveFaderClicked();
+
+	void UpdateFaderTemplateProperties();
+
+	/**Change fader background color on hover */
+	const FSlateBrush* GetBorderImage() const;
+
+	/** Handles when the user changes the Fader value */
+	void HandleFaderChanged(uint8 NewValue);
+
+	FText GetProtocolText() const;
+
+private:
+	TSharedPtr<SBorder> BackgroundBorder;
+
+	TSharedPtr<STextBlock> CustomFaderLabel;
+
+	TSharedPtr<SScrollBox> FaderChannelSlots;
+
+	TSharedPtr<SSpinBoxVertical<uint8>> FaderBoxVertical;
+
+	/** The user-selected Fader Label */
+	TAttribute<FText> CurrentFaderLabel;
+
+	/** The user-selected Fader Value */
+	uint8 CurrentFaderValue;
+
+	/** Pointer back to the DMXEditor tool that owns us */
+	TWeakPtr<FDMXEditor> WeakDMXEditor;
+
+	/** Pointer back to assosiated fader entity UObject */
+	TWeakObjectPtr<UDMXEntityFader> WeakFaderEntity;
+
+	/** Pointer back to fader list widget */
+	TWeakPtr<SDMXOutputFaderList> WeakFaderList;
+
+	TArray<TSharedPtr<SDMXFaderChannel>> FaderChannels;
+
+	TSharedPtr<SCheckBox> SendDMXCheckBox;
+};
+
+#undef LOCTEXT_NAMESPACE
\ No newline at end of file
diff --git a/Engine/Plugins/Experimental/DigitalMultiplex/DMXEngine/Source/DMXEditor/Private/Widgets/OutputFader/SDMXFaderChannel.cpp b/Engine/Plugins/Experimental/DigitalMultiplex/DMXEngine/Source/DMXEditor/Private/Widgets/OutputFader/SDMXFaderChannel.cpp
new file mode 100644
--- /dev/null
+++ b/Engine/Plugins/Experimental/DigitalMultiplex/DMXEngine/Source/DMXEditor/Private/Widgets/OutputFader/SDMXFaderChannel.cpp
@@ -0,0 +1,92 @@
+// Copyright Epic Games, Inc. All Rights Reserved.
+
+#include "Widgets/OutputFader/SDMXFaderChannel.h"
+
+#include "DMXEditor.h"
+#include "Library/DMXLibrary.h"
+#include "Library/DMXEntityFader.h"
+#include "DMXEditorLog.h"
+
+#include "Widgets/Common/SSpinBoxVertical.h"
+#include "Widgets/OutputFader/SDMXOutputFaderList.h"
+
+#include "Widgets/Input/SCheckBox.h"
+#include "Widgets/Layout/SScrollBox.h"
+#include "Widgets/Layout/SWrapBox.h"
+#include "Widgets/Layout/SSeparator.h"
+
+#define LOCTEXT_NAMESPACE "SDMXFaderChannel"
+
+void SDMXFaderChannel::Construct(const FArguments& InArgs)
+{
+	WeakDMXEditor = InArgs._DMXEditor;
+	UniverseNumber = InArgs._UniverseNumber;
+	ChannelNumber = InArgs._ChannelNumber;
+	WeakDMXProtocol = InArgs._DMXProtocol;
+
+	ChildSlot
+	.Padding(0.f)
+	[
+		SNew(SVerticalBox)
+		+ SVerticalBox::Slot()
+		.VAlign(VAlign_Top)
+		.HAlign(HAlign_Fill)
+		.Padding(0.f, 1.f, 0.f, 0.f)
+		.AutoHeight()
+		[
+			SNew(SVerticalBox)
+			+ SVerticalBox::Slot()
+			.VAlign(VAlign_Fill)
+			.HAlign(HAlign_Fill)
+			.FillHeight(10)
+			[
+				SNew(SHorizontalBox)
+				+ SHorizontalBox::Slot()
+				.VAlign(VAlign_Center)
+				.HAlign(HAlign_Fill)
+				.FillWidth(1)
+				[
+					SNew(SVerticalBox)
+					+ SVerticalBox::Slot()
+					.VAlign(VAlign_Fill)
+					.HAlign(HAlign_Fill)
+					.AutoHeight()
+					[
+						SAssignNew(UniverseValue, STextBlock)
+						.Justification(ETextJustify::Center)
+						.Text(FText::FromString("0"))
+						.Font(FCoreStyle::GetDefaultFontStyle("Regular", 7))
+					]
+				]
+				+ SHorizontalBox::Slot()
+				.VAlign(VAlign_Center)
+				.HAlign(HAlign_Fill)
+				.FillWidth(1)
+				[
+					SNew(SVerticalBox)
+
+					+ SVerticalBox::Slot()
+					.VAlign(VAlign_Fill)
+					.HAlign(HAlign_Fill)
+					.AutoHeight()
+					[
+						SAssignNew(ChannelValue, STextBlock)
+						.Justification(ETextJustify::Center)
+						.Text(FText::FromString("0"))
+					]
+				]
+			]
+			+ SVerticalBox::Slot()
+			.VAlign(VAlign_Fill)
+			.HAlign(HAlign_Fill)
+			.FillHeight(1)
+			[
+				SNew(SBorder)
+				.BorderImage(FEditorStyle::GetBrush("DetailsView.CategoryMiddle_Hovered"))
+				.Padding(FMargin(0, 5, 0, 5))
+			]
+		]
+	];
+}
+
+#undef LOCTEXT_NAMESPACE
diff --git a/Engine/Plugins/Experimental/DigitalMultiplex/DMXEngine/Source/DMXEditor/Private/Widgets/OutputFader/SDMXFaderChannel.h b/Engine/Plugins/Experimental/DigitalMultiplex/DMXEngine/Source/DMXEditor/Private/Widgets/OutputFader/SDMXFaderChannel.h
new file mode 100644
--- /dev/null
+++ b/Engine/Plugins/Experimental/DigitalMultiplex/DMXEngine/Source/DMXEditor/Private/Widgets/OutputFader/SDMXFaderChannel.h
@@ -0,0 +1,61 @@
+// Copyright Epic Games, Inc. All Rights Reserved.
+
+#pragma once
+
+#include "CoreMinimal.h"
+#include "Widgets/DeclarativeSyntaxSupport.h"
+#include "Widgets/SCompoundWidget.h"
+
+class STextBlock;
+class FDMXEditor;
+class IDMXProtocol;
+
+#define LOCTEXT_NAMESPACE "SDMXFaderChannel"
+
+class SDMXFaderChannel
+	: public SCompoundWidget
+{
+public:
+	SLATE_BEGIN_ARGS(SDMXFaderChannel)
+		: _InText(LOCTEXT("FaderLabel", "Fader"))
+		, _UniverseNumber(0)
+		, _ChannelNumber(0)
+	{}
+
+	SLATE_ARGUMENT(TWeakPtr<FDMXEditor>, DMXEditor)
+	SLATE_ARGUMENT(FText, InText)
+	SLATE_ARGUMENT(uint16, UniverseNumber)
+	SLATE_ARGUMENT(uint32, ChannelNumber)
+	SLATE_ARGUMENT(TWeakPtr<IDMXProtocol>, DMXProtocol)
+
+	SLATE_END_ARGS()
+
+	/** Constructs the widget */
+	void Construct(const FArguments& InArgs);
+
+	const TSharedPtr<STextBlock> GetUniverseValue() const { return UniverseValue; }
+
+	const TSharedPtr<STextBlock> GetChannelValue() const { return ChannelValue; }
+
+	uint16 GetUniverseNumber() const { return UniverseNumber; }
+
+	uint32 GetChannelNumber() const { return ChannelNumber; }
+
+	const TSharedPtr<IDMXProtocol> GetProtocol() const { return WeakDMXProtocol.Pin(); }
+
+private:
+	/** Pointer back to the DMXEditor tool that owns us */
+	TWeakPtr<FDMXEditor> WeakDMXEditor;
+
+	TSharedPtr<STextBlock> UniverseValue;
+
+	TSharedPtr<STextBlock> ChannelValue;
+
+	uint16 UniverseNumber;
+
+	uint32 ChannelNumber;
+
+	TWeakPtr<IDMXProtocol> WeakDMXProtocol;
+};
+
+#undef LOCTEXT_NAMESPACE
\ No newline at end of file
diff --git a/Engine/Plugins/Experimental/DigitalMultiplex/DMXEngine/Source/DMXEditor/Private/Widgets/OutputFader/SDMXOutputFaderList.cpp b/Engine/Plugins/Experimental/DigitalMultiplex/DMXEngine/Source/DMXEditor/Private/Widgets/OutputFader/SDMXOutputFaderList.cpp
new file mode 100644
--- /dev/null
+++ b/Engine/Plugins/Experimental/DigitalMultiplex/DMXEngine/Source/DMXEditor/Private/Widgets/OutputFader/SDMXOutputFaderList.cpp
@@ -0,0 +1,277 @@
+// Copyright Epic Games, Inc. All Rights Reserved.
+
+#include "Widgets/OutputFader/SDMXOutputFaderList.h"
+
+#include "Interfaces/IDMXProtocol.h"
+
+#include "DMXEditor.h"
+#include "Library/DMXLibrary.h"
+#include "Library/DMXEntityFader.h"
+#include "DMXEditorLog.h"
+#include "DMXEditorUtils.h"
+
+#include "Widgets/OutputFader/SDMXFader.h"
+
+#include "Widgets/Input/SButton.h"
+#include "Widgets/Layout/SScrollBox.h"
+
+#define LOCTEXT_NAMESPACE "SDMXOutputFaderList"
+
+void SDMXOutputFaderList::Construct(const FArguments& InArgs)
+{
+	WeakDMXEditor = InArgs._DMXEditor;
+	WeakFaderTemplate = InArgs._FaderTemplate;
+
+	if (TSharedPtr<FDMXEditor> DMXEditorPtr = WeakDMXEditor.Pin())
+	{
+		WeakDMXLibrary = DMXEditorPtr->GetDMXLibrary();
+	}
+
+	static const float Padding_Borders = 5.0f;
+	static const float Padding_KeyVal = 10.0f;
+	static const float Padding_NewInput = 35.0f;
+
+	ChildSlot
+		.VAlign(VAlign_Fill)
+		.HAlign(HAlign_Fill)
+		[
+			SNew(SVerticalBox)
+			+ SVerticalBox::Slot()
+			.VAlign(VAlign_Fill)
+			.HAlign(HAlign_Left)
+			.AutoHeight()
+			.Padding(5.0f)
+			[
+				SNew(SHorizontalBox)
+				+ SHorizontalBox::Slot()
+				.VAlign(VAlign_Fill)
+				.HAlign(HAlign_Left)
+				.AutoWidth()
+				[
+					SNew(SBox)
+					.HeightOverride(23)
+					[
+						SNew(SButton)
+						.Text(LOCTEXT("AddFader", "Add Fader"))
+						.HAlign(HAlign_Center)
+						.OnClicked(this, &SDMXOutputFaderList::HandleAddFaderClicked)
+					]
+				]
+				+ SHorizontalBox::Slot()
+				.VAlign(VAlign_Fill)
+				.HAlign(HAlign_Left)
+				.AutoWidth()
+				.Padding(10.0f, 0.f, 0.f, 0.f)
+				[
+					SNew(SBox)
+					.HeightOverride(23)
+					[
+						SNew(SButton)
+						.Text(LOCTEXT("UpdateFader", "Update Selected Fader"))
+						.HAlign(HAlign_Center)
+						.OnClicked(this, &SDMXOutputFaderList::HandleUpdateFaderClicked)
+					]
+				]
+			]
+			+ SVerticalBox::Slot()
+			.AutoHeight()
+			[
+				SAssignNew(FaderSlots, SScrollBox)
+				.Orientation(EOrientation::Orient_Horizontal)
+			]
+		];
+
+	// Reconstruct All Widgets
+	if (TSharedPtr<FDMXEditor> DMXEditorPtr = WeakDMXEditor.Pin())
+	{
+		if (UDMXLibrary* DMXLibrary = WeakDMXLibrary.Get())
+		{
+			DMXLibrary->ForEachEntityOfType<UDMXEntityFader>([this](UDMXEntityFader* Fader)
+				{
+					AddFader(Fader->GetDisplayName());
+				});
+		}
+	}
+}
+
+FReply SDMXOutputFaderList::OnMouseButtonDown(const FGeometry& MyGeometry, const FPointerEvent& MouseEvent)
+{
+	if (MouseEvent.GetEffectingButton() == EKeys::LeftMouseButton)
+	{
+		// TODO. implement
+	}
+	return FReply::Handled();
+}
+
+FReply SDMXOutputFaderList::HandleAddFaderClicked()
+{
+	AddFader();
+	//DeselectFaders();
+	return FReply::Handled();
+}
+
+FReply SDMXOutputFaderList::HandleUpdateFaderClicked()
+{
+	if (TSharedPtr<SDMXFader> SelectedPtr = WeakSelectedFaderWidget.Pin())
+	{
+		SelectedPtr->RemoveAllChannelWidgets();
+		SetFaderProperties(SelectedPtr, SelectedPtr->GetWeakFaderEntity(), true);
+	}
+	return FReply::Handled();
+}
+
+void SDMXOutputFaderList::ResetFaderBackgrounds()
+{
+	for (TSharedPtr<SDMXFader> Fader : FaderWidgets)
+	{
+		if (Fader.IsValid())
+		{
+			Fader->GetBackgroundBorder()->SetBorderImage(FEditorStyle::GetBrush("DetailsView.CategoryMiddle"));
+		}
+	}
+}
+
+void SDMXOutputFaderList::AddFader(const FString& InName /*= TEXT("")*/)
+{
+	TSharedPtr<SDMXFader> NewFader = SNew(SDMXFader)
+		.DMXEditor(WeakDMXEditor)
+		.InText(LOCTEXT("FaderLabel", "Fader"));
+		FaderSlots->AddSlot()
+		[
+			NewFader.ToSharedRef()
+		];
+
+	if (TSharedPtr<FDMXEditor> DMXEditorPtr = WeakDMXEditor.Pin())
+	{
+		if (UDMXLibrary* DMXLibrary = DMXEditorPtr->GetDMXLibrary())
+		{
+			if (UDMXEntityFader* FaderEntity = Cast<UDMXEntityFader>(DMXLibrary->GetOrCreateEntityObject(InName, UDMXEntityFader::StaticClass())))
+			{
+				
+
+				NewFader->SetFaderEntity(FaderEntity);
+				NewFader->SetParentFaderList(SharedThis(this));
+				FaderWidgets.Add(NewFader);
+
+				SetFaderProperties(NewFader, FaderEntity, InName.IsEmpty());
+				DeselectFaders();
+				NewFader->SelectThisFader();
+			}
+		}
+	}
+}
+
+void SDMXOutputFaderList::RemoveFader(TSharedPtr<SDMXFader> FaderToRemove)
+{
+	if (!FaderToRemove.IsValid())
+	{
+		UE_LOG_DMXEDITOR(Warning, TEXT("FaderToRemove is invalid pointer"));
+		return;
+	}
+
+	FaderWidgets.Remove(FaderToRemove);
+	if (TSharedPtr<FDMXEditor> DMXEditorPtr = WeakDMXEditor.Pin())
+	{
+		if (UDMXLibrary* DMXLibrary = DMXEditorPtr->GetDMXLibrary())
+		{
+			DMXLibrary->RemoveEntity(FaderToRemove->GetWeakFaderEntity()->GetDisplayName());
+		}
+	}
+	FaderSlots->RemoveSlot(FaderToRemove.ToSharedRef());
+}
+
+void SDMXOutputFaderList::SetFaderProperties(const TSharedPtr<SDMXFader>& InFaderWidget, const TWeakObjectPtr<UDMXEntityFader>& InFaderObject, bool bIsTransferObject)
+{
+	if (!InFaderWidget.IsValid())
+	{
+		UE_LOG_DMXEDITOR(Warning, TEXT("FaderWidget is invalid pointer"));
+		return;
+	}
+
+	if (!InFaderObject.IsValid())
+	{
+		UE_LOG_DMXEDITOR(Warning, TEXT("FaderObject is invalid pointer"));
+		return;
+	}
+
+	// Transfer Object Properties if object is new
+	if (bIsTransferObject)
+	{
+		TransferCreatedFaderObjectProperties(InFaderObject, WeakFaderTemplate);
+	}
+
+	TSharedPtr<IDMXProtocol> DMXProtocol = InFaderObject->DeviceProtocol;
+
+	InFaderWidget->SetFaderLabel(InFaderObject->GetDisplayName());
+	for (const FDMXUniverse& Universe : InFaderObject->Universes)
+	{
+		// Add protocol universe
+		FString TmpFinalUniverse = FString::FromInt(Universe.UniverseNumber);
+		FString TmpChannel = FString::FromInt(Universe.Channel);
+		InFaderWidget->AddChannelWidget(TmpFinalUniverse, TmpChannel, Universe.UniverseNumber, Universe.Channel, DMXProtocol);
+	}
+}
+
+void SDMXOutputFaderList::TransferCreatedFaderObjectProperties(const TWeakObjectPtr<UDMXEntityFader>& TransferTo, const TWeakObjectPtr<UDMXEntityFader>& TransferFrom)
+{
+	if (TransferTo.IsValid() && TransferFrom.IsValid())
+	{
+		if (TSharedPtr<FDMXEditor> DMXEditorPtr = WeakDMXEditor.Pin())
+		{
+			if (UDMXLibrary* DMXLibrary = DMXEditorPtr->GetDMXLibrary())
+			{
+				// Check if name is different or empty
+				if (TransferTo->GetDisplayName().IsEmpty() ||
+					(TransferFrom->GetDisplayName().Equals(TransferTo->GetDisplayName()) == false))
+				{
+					FText Reason;
+					bool bIsNameIsValid = FDMXEditorUtils::ValidateEntityName(TransferFrom->GetDisplayName(), DMXLibrary, UDMXEntityFader::StaticClass(), Reason);
+					if (!bIsNameIsValid)
+					{
+						FString UniqueEntityName = FDMXEditorUtils::FindUniqueEntityName(DMXLibrary, UDMXEntityFader::StaticClass());
+						TransferTo->SetName(UniqueEntityName);
+						TransferFrom->SetName(UniqueEntityName);
+					}
+					else
+					{
+						TransferTo->SetName(TransferFrom->GetDisplayName());
+					}
+				}
+
+				TransferTo->Universes = TransferFrom->Universes;
+				TransferTo->DeviceProtocol = TransferFrom->DeviceProtocol;
+			}
+		}
+	}
+	else
+	{
+		UE_LOG_DMXEDITOR(Warning, TEXT("No valid Template or Entity"));
+	}
+}
+
+void SDMXOutputFaderList::TransferSelectedFaderObjectProperties(const TWeakObjectPtr<UDMXEntityFader>& TransferTo, const TWeakObjectPtr<UDMXEntityFader>& TransferFrom)
+{
+	if (TransferTo.IsValid() && TransferFrom.IsValid())
+	{
+		TransferTo->SetName(TransferFrom->GetDisplayName());
+		TransferTo->Universes = TransferFrom->Universes;
+		TransferTo->DeviceProtocol = TransferFrom->DeviceProtocol;
+	}
+	else
+	{
+		UE_LOG_DMXEDITOR(Warning, TEXT("No valid Template or Entity"));
+	}
+}
+
+void SDMXOutputFaderList::UpdateFaderTemplateObject(const TWeakObjectPtr<UDMXEntityFader>& InFaderObject)
+{
+	TransferSelectedFaderObjectProperties(WeakFaderTemplate, InFaderObject);
+}
+
+void SDMXOutputFaderList::DeselectFaders()
+{
+	ResetFaderBackgrounds();
+	WeakSelectedFaderWidget = nullptr;
+}
+
+#undef LOCTEXT_NAMESPACE
diff --git a/Engine/Plugins/Experimental/DigitalMultiplex/DMXEngine/Source/DMXEditor/Private/Widgets/OutputFader/SDMXOutputFaderList.h b/Engine/Plugins/Experimental/DigitalMultiplex/DMXEngine/Source/DMXEditor/Private/Widgets/OutputFader/SDMXOutputFaderList.h
new file mode 100644
--- /dev/null
+++ b/Engine/Plugins/Experimental/DigitalMultiplex/DMXEngine/Source/DMXEditor/Private/Widgets/OutputFader/SDMXOutputFaderList.h
@@ -0,0 +1,72 @@
+// Copyright Epic Games, Inc. All Rights Reserved.
+
+#pragma once
+
+#include "CoreMinimal.h"
+#include "Widgets/DeclarativeSyntaxSupport.h"
+#include "Widgets/SCompoundWidget.h"
+
+class SDMXFader;
+class FDMXEditor;
+class UDMXEntityFader;
+class SScrollBox;
+class UDMXLibrary;
+
+class SDMXOutputFaderList
+	: public SCompoundWidget
+{
+public:
+	SLATE_BEGIN_ARGS(SDMXOutputFaderList)
+	{}
+
+	SLATE_ARGUMENT(TWeakPtr<FDMXEditor>, DMXEditor)
+	SLATE_ARGUMENT(TWeakObjectPtr<UDMXEntityFader>, FaderTemplate)
+
+	SLATE_END_ARGS()
+
+	/** Constructs the widget */
+	void Construct(const FArguments& InArgs);
+
+	void DeselectFaders();
+
+	void UpdateFaderTemplateObject(const TWeakObjectPtr<UDMXEntityFader>& InFaderObject);
+
+	void ResetFaderBackgrounds();
+
+	void RemoveFader(TSharedPtr<SDMXFader> FaderToRemove);
+
+	const TWeakPtr<SDMXFader>& GetWeakSelectedFaderWidget() const { return WeakSelectedFaderWidget; }
+
+	const TArray<TSharedPtr<SDMXFader>>& GetFaderWidgets() const { return FaderWidgets; }
+
+private:
+	//~ Begin SWidget implementation
+	virtual FReply OnMouseButtonDown(const FGeometry& MyGeometry, const FPointerEvent& MouseEvent) override;
+	//~ Public SWidget implementation
+
+	FReply HandleAddFaderClicked();
+
+	FReply HandleUpdateFaderClicked();
+
+	void SetFaderProperties(const TSharedPtr<SDMXFader>& InFaderWidget, const TWeakObjectPtr<UDMXEntityFader>& InFaderObject, bool bIsTransferObject);
+
+	void TransferCreatedFaderObjectProperties(const TWeakObjectPtr<UDMXEntityFader>& TransferTo, const TWeakObjectPtr<UDMXEntityFader>& TransferFrom);
+	void TransferSelectedFaderObjectProperties(const TWeakObjectPtr<UDMXEntityFader>& TransferTo, const TWeakObjectPtr<UDMXEntityFader>& TransferFrom);
+
+	void AddFader(const FString& InName = TEXT(""));
+
+public:
+	TWeakPtr<SDMXFader> WeakSelectedFaderWidget;
+
+private:
+	/** Pointer back to the DMXEditor tool that owns us */
+	TWeakPtr<FDMXEditor> WeakDMXEditor;
+
+	TSharedPtr<SScrollBox> FaderSlots;
+
+	TArray<TSharedPtr<SDMXFader>> FaderWidgets;
+
+	TWeakObjectPtr<UDMXEntityFader> WeakFaderTemplate;
+
+	TWeakObjectPtr<UDMXLibrary> WeakDMXLibrary;
+};
