- DMX Editor

diff --git a/Engine/Plugins/Experimental/DigitalMultiplex/DMXEngine/Source/DMXEditor/Private/Widgets/SDMXEntityEditor.cpp b/Engine/Plugins/Experimental/DigitalMultiplex/DMXEngine/Source/DMXEditor/Private/Widgets/SDMXEntityEditor.cpp
new file mode 100644
--- /dev/null
+++ b/Engine/Plugins/Experimental/DigitalMultiplex/DMXEngine/Source/DMXEditor/Private/Widgets/SDMXEntityEditor.cpp
@@ -0,0 +1,209 @@
+// Copyright Epic Games, Inc. All Rights Reserved.
+
+#include "Widgets/SDMXEntityEditor.h"
+
+#include "DMXEditor.h"
+#include "DMXEditorLog.h"
+#include "DMXProtocolConstants.h"
+#include "DMXEditorMacros.h"
+
+#include "Widgets/SDMXEntityInspector.h"
+#include "Widgets/SDMXEntityList.h"
+
+#include "Library/DMXLibrary.h"
+#include "Library/DMXEntityFixturePatch.h"
+#include "Library/DMXEntityFixtureType.h"
+#include "Library/DMXEntityController.h"
+#include "Library/DMXImport.h"
+
+#include "Widgets/Layout/SSplitter.h"
+#include "Widgets/Layout/SBox.h"
+#include "Widgets/SNullWidget.h"
+#include "Widgets/Text/STextBlock.h"
+
+#define LOCTEXT_NAMESPACE "SDMXEntityEditor"
+
+void SDMXEntityEditor::Construct(const FArguments &InArgs)
+{
+	DMXEditor = InArgs._DMXEditor;
+	EditorEntityType = InArgs._EditorEntityType;
+
+	if (EditorEntityType->IsChildOf(UDMXEntityController::StaticClass()))
+	{
+		InspectorWidget = SNew(SDMXEntityInspectorControllers)
+					.DMXEditor(DMXEditor)
+					.OnFinishedChangingProperties(this, &SDMXEntityEditor::OnFinishedChangingProperties);
+	}
+	else if (EditorEntityType->IsChildOf(UDMXEntityFixtureType::StaticClass()))
+	{
+		InspectorWidget = SNew(SDMXEntityInspectorFixtureTypes)
+					.DMXEditor(DMXEditor)
+					.OnFinishedChangingProperties(this, &SDMXEntityEditor::OnFinishedChangingProperties);
+	}
+	else if (EditorEntityType->IsChildOf(UDMXEntityFixturePatch::StaticClass()))
+	{
+		InspectorWidget = SNew(SDMXEntityInspectorFixturePatches)
+					.DMXEditor(DMXEditor)
+					.OnFinishedChangingProperties(this, &SDMXEntityEditor::OnFinishedChangingProperties);
+	}
+
+	check(InspectorWidget.IsValid());
+
+	ChildSlot
+	.VAlign(VAlign_Fill)
+	.HAlign(HAlign_Fill)
+	[
+		SAssignNew(SidesSplitter, SSplitter)
+		.Orientation(EOrientation::Orient_Horizontal)
+		.ResizeMode(ESplitterResizeMode::FixedPosition)
+
+		// left side container
+		+SSplitter::Slot()
+		.Value(0.35f)
+		[
+			SAssignNew(ListWidget, SDMXEntityList, EditorEntityType)
+			.DMXEditor(InArgs._DMXEditor)
+			.OnSelectionUpdated(this, &SDMXEntityEditor::OnSelectionUpdated)
+		]
+
+		// right side container
+		+SSplitter::Slot()
+		.Value(0.65f)
+		[
+			InspectorWidget.ToSharedRef()
+		]
+	];
+}
+
+void SDMXEntityEditor::RequestRenameOnNewEntity(const UDMXEntity* InEntity, ESelectInfo::Type SelectionType)
+{
+	check(ListWidget.IsValid());
+
+	ListWidget->UpdateTree();
+	ListWidget->SelectItemByEntity(InEntity, SelectionType);
+	ListWidget->OnRenameNode();
+}
+
+void SDMXEntityEditor::SelectEntity(UDMXEntity* InEntity, ESelectInfo::Type InSelectionType /*= ESelectInfo::Type::Direct*/)
+{
+	check(ListWidget.IsValid());
+
+	ListWidget->SelectItemByEntity(InEntity, InSelectionType);
+}
+
+void SDMXEntityEditor::SelectEntities(const TArray<UDMXEntity*>& InEntities, ESelectInfo::Type InSelectionType /*= ESelectInfo::Type::Direct*/)
+{
+	check(ListWidget.IsValid());
+
+	ListWidget->SelectItemsByEntity(InEntities, InSelectionType);
+}
+
+TArray<UDMXEntity*> SDMXEntityEditor::GetSelectedEntities() const
+{
+	return ListWidget->GetSelectedEntities();
+}
+
+void SDMXEntityEditor::OnSelectionUpdated(TArray<TSharedPtr<FDMXTreeNodeBase>> InSelectedNodes)
+{
+	check(GetInspectorWidget().IsValid());
+
+	if (TSharedPtr<FDMXEditor> DMXEditorPtr = DMXEditor.Pin())
+	{
+		UDMXLibrary* DMXLibrary = DMXEditorPtr->GetDMXLibrary();
+
+		TArray<UObject*> SelectedObjects;
+		for (TSharedPtr<FDMXTreeNodeBase> Node : InSelectedNodes)
+		{
+			if (UDMXEntity* Entity = Node->GetEntity())
+			{
+				SelectedObjects.Add(Entity);
+			}
+		}
+
+		// Update property inspector if we changed the selection or if there are no entities in the list.
+		// So, if the user de-selects all entities, the last one's properties is still shown. But if they
+		// delete all entities, we empty the inspector.
+		if (SelectedObjects.Num() > 0 || ListWidget->IsListEmpty())
+		{
+			GetInspectorWidget()->ShowDetailsForEntities(SelectedObjects);
+		}
+	}
+}
+
+void SDMXEntityEditor::OnFinishedChangingProperties(const FPropertyChangedEvent& PropertyChangedEvent)
+{
+	if (ListWidget.IsValid())
+	{
+		ListWidget->UpdateTree();
+	}
+}
+
+void SDMXControllers::Construct(const FArguments& InArgs)
+{
+	SDMXEntityEditor::Construct(SDMXEntityEditor::FArguments()
+		.DMXEditor(InArgs._DMXEditor)
+		.EditorEntityType(UDMXEntityController::StaticClass())
+	);
+}
+
+void SDMXFixtureTypes::Construct(const FArguments& InArgs)
+{
+	SDMXEntityEditor::Construct(SDMXEntityEditor::FArguments()
+		.DMXEditor(InArgs._DMXEditor)
+		.EditorEntityType(UDMXEntityFixtureType::StaticClass())
+	);
+}
+
+void SDMXFixtureTypes::OnFinishedChangingProperties(const FPropertyChangedEvent& PropertyChangedEvent)
+{
+	if (PropertyChangedEvent.ChangeType == EPropertyChangeType::ArrayAdd)
+	{
+		if (PropertyChangedEvent.GetPropertyName() == GET_MEMBER_NAME_CHECKED(UDMXEntityFixtureType, Modes)
+			|| PropertyChangedEvent.GetPropertyName() == GET_MEMBER_NAME_CHECKED(FDMXFixtureMode, Functions)
+			|| PropertyChangedEvent.GetPropertyName() == GET_MEMBER_NAME_CHECKED(FDMXFixtureFunction, SubFunctions))
+		{
+			// When the user adds a Mode, Function or Sub Function, their names can't be empty
+
+			const TArray<UDMXEntity*>&& SelectedEntities = ListWidget->GetSelectedEntities();
+
+			for (UDMXEntity* Entity : SelectedEntities)
+			{
+				if (UDMXEntityFixtureType* FixtureType = Cast<UDMXEntityFixtureType>(Entity))
+				{
+					FDMXEditorUtils::SetNewFixtureFunctionsNames(FixtureType);
+				}
+			}
+		}
+	}
+	else if (PropertyChangedEvent.ChangeType == EPropertyChangeType::ValueSet)
+	{
+		if (PropertyChangedEvent.GetPropertyName() == GET_MEMBER_NAME_CHECKED(UDMXEntityFixtureType, DMXImport))
+		{
+			const TArray<UDMXEntity*>&& SelectedEntities = ListWidget->GetSelectedEntities();
+
+			for (UDMXEntity* Entity : SelectedEntities)
+			{
+				if (UDMXEntityFixtureType* FixtureType = Cast<UDMXEntityFixtureType>(Entity))
+				{
+					if (FixtureType->DMXImport != nullptr && FixtureType->DMXImport->IsValidLowLevelFast())
+					{
+						FixtureType->Modify();
+						FixtureType->SetModesFromDMXImport(FixtureType->DMXImport);
+					}
+				}
+			}
+		}
+	}
+
+	SDMXEntityEditor::OnFinishedChangingProperties(PropertyChangedEvent);
+}
+
+void SDMXFixturePatch::Construct(const FArguments& InArgs)
+{
+	SDMXEntityEditor::Construct(SDMXEntityEditor::FArguments()
+		.DMXEditor(InArgs._DMXEditor)
+		.EditorEntityType(UDMXEntityFixturePatch::StaticClass())
+	);
+}
+
+#undef LOCTEXT_NAMESPACE
\ No newline at end of file
diff --git a/Engine/Plugins/Experimental/DigitalMultiplex/DMXEngine/Source/DMXEditor/Private/Widgets/SDMXEntityEditor.h b/Engine/Plugins/Experimental/DigitalMultiplex/DMXEngine/Source/DMXEditor/Private/Widgets/SDMXEntityEditor.h
new file mode 100644
--- /dev/null
+++ b/Engine/Plugins/Experimental/DigitalMultiplex/DMXEngine/Source/DMXEditor/Private/Widgets/SDMXEntityEditor.h
@@ -0,0 +1,140 @@
+// Copyright Epic Games, Inc. All Rights Reserved.
+
+#pragma once
+
+#include "CoreMinimal.h"
+#include "Widgets/SCompoundWidget.h"
+#include "Widgets/SDMXEntityList.h"
+
+#include "Types/SlateEnums.h"
+
+class FDMXEditor;
+class SSplitter;
+class SBox;
+class SDMXEntityInspector;
+struct FPropertyChangedEvent;
+
+/** DMX Controllers tab root widget */
+class SDMXEntityEditor
+	: public SCompoundWidget
+{
+public:
+	SLATE_BEGIN_ARGS(SDMXEntityEditor)
+		: _DMXEditor(nullptr)
+	{}
+	
+	SLATE_ARGUMENT(TWeakPtr<FDMXEditor>, DMXEditor)
+	SLATE_ARGUMENT(TSubclassOf<UDMXEntity>, EditorEntityType)
+
+	SLATE_END_ARGS()
+
+public:
+	SDMXEntityEditor()
+	{
+		SetCanTick(false);
+		bCanSupportFocus = false;
+	}
+
+	/** Constructs the sides container widgets. The child classes should use
+	* SetLeftSideWidget and SetRightSideWidget to add the child widgets into the containers. */
+	void Construct(const FArguments& InArgs);
+
+	/** Get the current left side column's widget */
+	const TSharedPtr<SDMXEntityList> GetListWidget() const { return ListWidget; }
+	/** Get the current right side column's widget */
+	const TSharedPtr<SDMXEntityInspector> GetInspectorWidget() const { return InspectorWidget; }
+
+	/**
+	 * Refreshes the list of entities to display any added entities, select the new entity by name
+	 * and initiates a rename on the selected Entity node
+	 */
+	void RequestRenameOnNewEntity(const UDMXEntity* InEntity, ESelectInfo::Type SelectionType);
+
+	/** Selects an entity in this editor tab's list */
+	void SelectEntity(UDMXEntity* InEntity, ESelectInfo::Type InSelectionType = ESelectInfo::Type::Direct);
+	/** Selects Entities in this editor tab's list */
+	void SelectEntities(const TArray<UDMXEntity*>& InEntities, ESelectInfo::Type SelectionType = ESelectInfo::Type::Direct);
+	/** Returns the selected entities on this editor tab */
+	TArray<UDMXEntity*> GetSelectedEntities() const;
+
+protected:
+	/** Callback for when entities list selection changes */
+	virtual void OnSelectionUpdated(TArray<TSharedPtr<FDMXTreeNodeBase>> InSelectedNodes);
+
+	/** Callback for when some property has changed in the inspector */
+	virtual void OnFinishedChangingProperties(const FPropertyChangedEvent& PropertyChangedEvent);
+
+protected:
+	/** Pointer back to the DMXEditor tool that owns us */
+	TWeakPtr<FDMXEditor> DMXEditor;
+
+	/** Contains both sides' columns */
+	TSharedPtr<SSplitter> SidesSplitter;
+
+	/** Left widget parent */
+	TSharedPtr<SBox> ListContainer;
+
+	/** Right widget parent */
+	TSharedPtr<SBox> InspectorContainer;
+
+	/** Left child widget */
+	TSharedPtr<SDMXEntityList> ListWidget;
+
+	/** Right child widget */
+	TSharedPtr<SDMXEntityInspector> InspectorWidget;
+
+	/** Indicates which tab we are */
+	TSubclassOf<UDMXEntity> EditorEntityType;
+};
+
+/** DMX Controllers tab root widget */
+class SDMXControllers
+	: public SDMXEntityEditor
+{
+public:
+	SLATE_BEGIN_ARGS(SDMXControllers)
+	{}
+
+	SLATE_ARGUMENT(TWeakPtr<FDMXEditor>, DMXEditor)
+
+	SLATE_END_ARGS()
+
+	void Construct(const FArguments& InArgs);
+};
+
+
+/** DMX Controllers tab root widget */
+class SDMXFixtureTypes
+	: public SDMXEntityEditor
+{
+public:
+	SLATE_BEGIN_ARGS(SDMXFixtureTypes)
+	{}
+
+	SLATE_ARGUMENT(TWeakPtr<FDMXEditor>, DMXEditor)
+
+	SLATE_END_ARGS()
+
+	/** Constructs the widget */
+	void Construct(const FArguments& InArgs);
+
+	//~ SDMXEntityEditor interface
+	virtual void OnFinishedChangingProperties(const FPropertyChangedEvent& PropertyChangedEvent) override;
+}; 
+
+
+/** DMX Controllers tab root widget */
+class SDMXFixturePatch
+	: public SDMXEntityEditor
+{
+public:
+	SLATE_BEGIN_ARGS(SDMXFixturePatch)
+	{}
+
+	SLATE_ARGUMENT(TWeakPtr<FDMXEditor>, DMXEditor)
+
+	SLATE_END_ARGS()
+
+	/** Constructs the widget */
+	void Construct(const FArguments& InArgs);
+};
diff --git a/Engine/Plugins/Experimental/DigitalMultiplex/DMXEngine/Source/DMXEditor/Private/Widgets/SDMXEntityInspector.cpp b/Engine/Plugins/Experimental/DigitalMultiplex/DMXEngine/Source/DMXEditor/Private/Widgets/SDMXEntityInspector.cpp
new file mode 100644
--- /dev/null
+++ b/Engine/Plugins/Experimental/DigitalMultiplex/DMXEngine/Source/DMXEditor/Private/Widgets/SDMXEntityInspector.cpp
@@ -0,0 +1,198 @@
+// Copyright Epic Games, Inc. All Rights Reserved.
+
+#include "Widgets/SDMXEntityInspector.h"
+
+#include "DMXEditor.h"
+#include "DMXEditorLog.h"
+#include "DMXProtocolConstants.h"
+
+#include "Customizations/DMXEditorPropertyEditorCustomization.h"
+
+#include "Library/DMXEntityFixtureType.h"
+#include "Library/DMXEntityFixturePatch.h"
+#include "Library/DMXEntityController.h"
+#include "Library/DMXEntityFader.h"
+
+#include "Game/DMXComponent.h"
+
+#include "Widgets/Text/STextBlock.h"
+#include "Widgets/Layout/SBorder.h"
+
+#include "Modules/ModuleManager.h"
+#include "PropertyEditorModule.h"
+
+#define LOCTEXT_NAMESPACE "SDMXEntityInspector"
+
+void SDMXEntityInspector::Construct(const FArguments& InArgs)
+{
+	// Initialize input arguments
+	DMXEditor = InArgs._DMXEditor;
+	UserOnFinishedChangingProperties = InArgs._OnFinishedChangingProperties;
+	bIsShowSearch = InArgs._ShowSearch;
+	bIsShowTitleArea = InArgs._ShowTitleArea;
+	bIsHideFilterArea = InArgs._HideFilterArea;
+
+	// Do not update by default
+	bRefreshOnTick = false;
+
+	// Initialize property view widget
+	FPropertyEditorModule& PropertyEditorModule = FModuleManager::GetModuleChecked<FPropertyEditorModule>("PropertyEditor");
+
+	FNotifyHook* NotifyHook = nullptr;
+	if (InArgs._SetNotifyHook)
+	{
+		NotifyHook = DMXEditor.Pin().Get();
+	}
+	FDetailsViewArgs::ENameAreaSettings NameAreaSettings = InArgs._HideNameArea ? FDetailsViewArgs::HideNameArea : FDetailsViewArgs::ObjectsUseNameArea;
+	FDetailsViewArgs DetailsViewArgs( 
+		/*bUpdateFromSelection=*/ false, 
+		/*bLockable=*/ false, 
+		/*bAllowSearch=*/ bIsShowSearch,
+		/*NameAreaSettings=*/ NameAreaSettings,
+		/*bHideSelectionTip=*/ true, 
+		/*InNotifyHook=*/ NotifyHook, 
+		/*InSearchInitialKeyFocus=*/ false, 
+		/*InViewIdentifier=*/ InArgs._ViewIdentifier);
+	PropertyView = PropertyEditorModule.CreateDetailView(DetailsViewArgs);
+
+	ChildSlot
+	[	SNew(SVerticalBox)
+		+ SVerticalBox::Slot()
+		.FillHeight(1.0f)
+		[
+			SAssignNew(EditingWidget, SBorder)
+		]
+	];
+
+	// Update based on the current (empty) selection set
+	TArray<UObject*> InitialSelectedObjects;
+	UpdateFromObjects(InitialSelectedObjects);
+}
+
+void SDMXEntityInspector::Tick(const FGeometry& AllottedGeometry, const double InCurrentTime, const float InDeltaTime)
+{
+	if (bRefreshOnTick)
+	{
+		UpdateFromObjects(RefreshPropertyObjects);
+		RefreshPropertyObjects.Empty();
+		bRefreshOnTick = false;
+	}
+}
+
+void SDMXEntityInspector::ShowDetailsForSingleEntity(UObject* Object)
+{
+	TArray<UObject*> PropertyObjects;
+
+	if (Object != NULL)
+	{
+		PropertyObjects.Add(Object);
+	}
+
+	ShowDetailsForEntities(PropertyObjects);
+}
+
+void SDMXEntityInspector::ShowDetailsForEntities(const TArray<UObject*>& PropertyObjects)
+{
+	// Refresh is being deferred until the next tick, this prevents batch operations from bombarding the details view with calls to refresh
+	RefreshPropertyObjects = PropertyObjects;
+	bRefreshOnTick = true;
+}
+
+void SDMXEntityInspector::UpdateFromObjects(const TArray<UObject*>& PropertyObjects)
+{
+	// Register IDetailsView callbacks
+	PropertyView->OnFinishedChangingProperties().Clear();
+	PropertyView->OnFinishedChangingProperties().Add(UserOnFinishedChangingProperties);
+
+	// Update our context-sensitive editing widget
+	EditingWidget->SetContent(MakeEditingWidget(PropertyObjects));
+}
+
+TSharedRef<SWidget> SDMXEntityInspector::MakeEditingWidget(const TArray<UObject*>& Objects)
+{
+	TSharedRef< SVerticalBox > InnerEditingWidget = SNew( SVerticalBox );
+
+	// Show the property editor
+	PropertyView->HideFilterArea(bIsHideFilterArea);
+	PropertyView->SetObjects(Objects, true);
+
+	InnerEditingWidget->AddSlot()
+	.FillHeight( 0.9f )
+	.VAlign( VAlign_Top )
+	[
+		SNew( SBox )
+		[
+			SNew( SVerticalBox )
+			+SVerticalBox::Slot()
+			.AutoHeight()
+			+SVerticalBox::Slot()
+			[
+				PropertyView.ToSharedRef()
+			]
+		]
+	];
+
+	return InnerEditingWidget;
+}
+
+void SDMXEntityInspectorControllers::Construct(const FArguments& InArgs)
+{
+	SDMXEntityInspector::Construct(SDMXEntityInspector::FArguments()
+		.DMXEditor(InArgs._DMXEditor)
+		.OnFinishedChangingProperties(InArgs._OnFinishedChangingProperties)
+	);
+
+	// Register customization for UOBJECT
+	FOnGetDetailCustomizationInstance ControllersDetails = FOnGetDetailCustomizationInstance::CreateStatic(&FDMXCustomizationFactory::MakeInstance<FDMXControllersDetails>, InArgs._DMXEditor);
+	GetPropertyView()->RegisterInstancedCustomPropertyLayout(UDMXEntityController::StaticClass(), ControllersDetails);
+}
+
+void SDMXEntityInspectorFixturePatches::Construct(const FArguments& InArgs)
+{
+	SDMXEntityInspector::Construct(SDMXEntityInspector::FArguments()
+		.DMXEditor(InArgs._DMXEditor)
+		.OnFinishedChangingProperties(InArgs._OnFinishedChangingProperties)
+	);
+
+	// Register customization for UOBJECT
+	FOnGetDetailCustomizationInstance FixtureTypesDetails = FOnGetDetailCustomizationInstance::CreateStatic(&FDMXCustomizationFactory::MakeInstance<FDMXFixturePatchesDetails>, InArgs._DMXEditor);
+	GetPropertyView()->RegisterInstancedCustomPropertyLayout(UDMXEntityFixturePatch::StaticClass(), FixtureTypesDetails);
+}
+
+void SDMXEntityInspectorFixtureTypes::Construct(const FArguments& InArgs)
+{
+	SDMXEntityInspector::Construct(SDMXEntityInspector::FArguments()
+		.DMXEditor(InArgs._DMXEditor)
+		.OnFinishedChangingProperties(InArgs._OnFinishedChangingProperties)
+	);
+
+	// Register generic customization for Fixture Type UOBJECT, just to keep its categories in order
+	FOnGetDetailCustomizationInstance FixtureTypesDetails = FOnGetDetailCustomizationInstance::CreateStatic(&FDMXCustomizationFactory::MakeInstance<FDMXCustomization>, InArgs._DMXEditor);
+	GetPropertyView()->RegisterInstancedCustomPropertyLayout(UDMXEntityFixtureType::StaticClass(), FixtureTypesDetails);
+	
+	// Register customization for Fixture Mode USTRUCT
+	FOnGetPropertyTypeCustomizationInstance FixtureModeDetails = FOnGetPropertyTypeCustomizationInstance::CreateStatic(&FDMXCustomizationFactory::MakeInstance<FDMXFixtureModeDetails>, InArgs._DMXEditor);
+	GetPropertyView()->RegisterInstancedCustomPropertyTypeLayout(FDMXFixtureMode::StaticStruct()->GetFName(), FixtureModeDetails);
+
+	// Register customization for Fixture Function USTRUCT
+	FOnGetPropertyTypeCustomizationInstance FixtureFunctionDetails = FOnGetPropertyTypeCustomizationInstance::CreateStatic(&FDMXCustomizationFactory::MakeInstance<FDMXFixtureFunctionDetails>, InArgs._DMXEditor);
+	GetPropertyView()->RegisterInstancedCustomPropertyTypeLayout(FDMXFixtureFunction::StaticStruct()->GetFName(), FixtureFunctionDetails);
+	
+	// Register customization for Fixture Sub Function USTRUCT
+	FOnGetPropertyTypeCustomizationInstance FixtureSubFunctionDetails = FOnGetPropertyTypeCustomizationInstance::CreateStatic(&FDMXCustomizationFactory::MakeInstance<FDMXFixtureSubFunctionDetails>, InArgs._DMXEditor);
+	GetPropertyView()->RegisterInstancedCustomPropertyTypeLayout(FDMXFixtureSubFunction::StaticStruct()->GetFName(), FixtureSubFunctionDetails);
+}
+
+void SDMXEntityInspectorFaders::Construct(const FArguments& InArgs)
+{
+	SDMXEntityInspector::Construct(SDMXEntityInspector::FArguments()
+		.DMXEditor(InArgs._DMXEditor)
+		.OnFinishedChangingProperties(InArgs._OnFinishedChangingProperties)
+	);
+
+	// Register customization for UOBJECT
+	FOnGetDetailCustomizationInstance UniverseManagerDetails = FOnGetDetailCustomizationInstance::CreateStatic(&FDMXCustomizationFactory::MakeInstance<FDMXCustomization>, InArgs._DMXEditor);
+	GetPropertyView()->RegisterInstancedCustomPropertyLayout(UDMXEntityFader::StaticClass(), UniverseManagerDetails);
+}
+
+#undef LOCTEXT_NAMESPACE
diff --git a/Engine/Plugins/Experimental/DigitalMultiplex/DMXEngine/Source/DMXEditor/Private/Widgets/SDMXEntityInspector.h b/Engine/Plugins/Experimental/DigitalMultiplex/DMXEngine/Source/DMXEditor/Private/Widgets/SDMXEntityInspector.h
new file mode 100644
--- /dev/null
+++ b/Engine/Plugins/Experimental/DigitalMultiplex/DMXEngine/Source/DMXEditor/Private/Widgets/SDMXEntityInspector.h
@@ -0,0 +1,180 @@
+// Copyright Epic Games, Inc. All Rights Reserved.
+
+#pragma once
+
+#include "CoreMinimal.h"
+#include "PropertyEditorDelegates.h"
+#include "Widgets/SCompoundWidget.h"
+#include "Widgets/DeclarativeSyntaxSupport.h"
+
+class FDMXEditor;
+class SBorder;
+
+/** DMX inspector details view base class */
+class SDMXEntityInspector
+	: public SCompoundWidget
+{
+public:
+	SLATE_BEGIN_ARGS(SDMXEntityInspector)
+		: _SetNotifyHook(true)
+		, _ShowSearch(false)
+		, _ShowTitleArea(false)
+		, _HideFilterArea(false)
+		, _HideNameArea(true)
+		, _ViewIdentifier(FName(TEXT("DMXEntityInspector")))
+		, _OnFinishedChangingProperties()
+	{}
+
+	SLATE_ARGUMENT(TWeakPtr<FDMXEditor>, DMXEditor)
+	SLATE_ARGUMENT(bool, SetNotifyHook)
+	SLATE_ARGUMENT(bool, ShowSearch)
+	SLATE_ARGUMENT(bool, ShowTitleArea)
+	SLATE_ARGUMENT(bool, HideFilterArea)
+
+	SLATE_ARGUMENT(bool, HideNameArea)
+	SLATE_ARGUMENT(FName, ViewIdentifier)
+
+	SLATE_EVENT(FOnFinishedChangingProperties::FDelegate, OnFinishedChangingProperties)
+
+	SLATE_END_ARGS()
+
+	/** Constructs the widget */
+	void Construct(const FArguments& InArgs);
+
+	//~ Begin SWidget interface
+	virtual void Tick(const FGeometry& AllottedGeometry, const double InCurrentTime, const float InDeltaTime) override;
+	//~End of SWidget interface
+
+	/** Update the inspector window to show information on the supplied Entity */
+	void ShowDetailsForSingleEntity(UObject* Object);
+
+	/** Update the inspector window to show information on the supplied objects */
+	void ShowDetailsForEntities(const TArray<UObject*>& PropertyObjects);
+
+	/** Get current property view widget */
+	TSharedPtr<class IDetailsView> GetPropertyView() const { return PropertyView; }
+
+protected:
+	/** Update the inspector window to show information on the supplied objects */
+	void UpdateFromObjects(const TArray<UObject*>& PropertyObjects);
+
+	/**
+	 * Generates a widget that is used to edit the specified object array.  This widget
+	 * will be displayed along with a property view in the level editor
+	 */
+	TSharedRef<SWidget> MakeEditingWidget(const TArray<UObject*>& Objects);
+
+protected:
+	/** Pointer back to the DMXEditor tool that owns us */
+	TWeakPtr<FDMXEditor> DMXEditor;
+
+	/** Property viewing widget */
+	TSharedPtr<IDetailsView> PropertyView;
+
+	/** Show Search for property view */
+	bool bIsShowSearch;
+
+	/** Show title area for property view */
+	bool bIsShowTitleArea;	
+	
+	/** Hide filter area for property view */
+	bool bIsHideFilterArea;
+
+	/** User defined delegate for OnFinishedChangingProperties */
+	FOnFinishedChangingProperties::FDelegate UserOnFinishedChangingProperties;
+
+	/** When TRUE, the inspector needs to refresh the details view on Tick */
+	bool bRefreshOnTick;
+
+	/** Holds the property objects that need to be displayed by the inspector starting on the next tick */
+	TArray<UObject*> RefreshPropertyObjects;
+
+	/** Border widget that wraps a dynamic context-sensitive widget for editing objects that the property window is displaying */
+	TSharedPtr<SBorder> EditingWidget;
+};
+
+/** DMX controllers inspector */
+class SDMXEntityInspectorControllers
+	: public SDMXEntityInspector
+{
+public:
+	SLATE_BEGIN_ARGS(SDMXEntityInspectorControllers)
+	{}
+
+	SLATE_ARGUMENT(TWeakPtr<FDMXEditor>, DMXEditor)
+	SLATE_EVENT(FOnFinishedChangingProperties::FDelegate, OnFinishedChangingProperties)
+
+
+	SLATE_END_ARGS()
+
+	/** Constructs the widget */
+	void Construct(const FArguments& InArgs);
+};
+
+/** DMX fixture confiturations inspector */
+class SDMXEntityInspectorFixturePatches
+	: public SDMXEntityInspector
+{
+public:
+	SLATE_BEGIN_ARGS(SDMXEntityInspectorFixturePatches)
+	{}
+
+	SLATE_ARGUMENT(TWeakPtr<FDMXEditor>, DMXEditor)
+	SLATE_EVENT(FOnFinishedChangingProperties::FDelegate, OnFinishedChangingProperties)
+
+	SLATE_END_ARGS()
+
+	/** Constructs the widget */
+	void Construct(const FArguments& InArgs);
+};
+
+/** DMX fixture types inspector */
+class SDMXEntityInspectorFixtureTypes
+	: public SDMXEntityInspector
+{
+public:
+	SLATE_BEGIN_ARGS(SDMXEntityInspectorFixtureTypes)
+	{}
+
+	SLATE_ARGUMENT(TWeakPtr<FDMXEditor>, DMXEditor)
+	SLATE_EVENT(FOnFinishedChangingProperties::FDelegate, OnFinishedChangingProperties)
+
+	SLATE_END_ARGS()
+
+	/** Constructs the widget */
+	void Construct(const FArguments& InArgs);
+};
+
+/** DMX fixture types inspector */
+class SDMXEntityInspectorFaders
+	: public SDMXEntityInspector
+{
+public:
+	SLATE_BEGIN_ARGS(SDMXEntityInspectorFaders)
+	{}
+
+		SLATE_ARGUMENT(TWeakPtr<FDMXEditor>, DMXEditor)
+		SLATE_EVENT(FOnFinishedChangingProperties::FDelegate, OnFinishedChangingProperties)
+
+	SLATE_END_ARGS()
+
+	/** Constructs the widget */
+	void Construct(const FArguments& InArgs);
+};
+
+/** DMX Component inspector */
+class SDMXComponentInspector
+	: public SDMXEntityInspector
+{
+public:
+	SLATE_BEGIN_ARGS(SDMXComponentInspector)
+	{}
+
+	SLATE_ARGUMENT(TWeakPtr<FDMXEditor>, DMXEditor)
+		SLATE_EVENT(FOnFinishedChangingProperties::FDelegate, OnFinishedChangingProperties)
+
+		SLATE_END_ARGS()
+
+		/** Constructs the widget */
+		void Construct(const FArguments& InArgs);
+};
