- DMX Editor

diff --git a/Engine/Plugins/Experimental/DigitalMultiplex/DMXEngine/Source/DMXEditor/Private/Widgets/SDMXEntityList.cpp b/Engine/Plugins/Experimental/DigitalMultiplex/DMXEngine/Source/DMXEditor/Private/Widgets/SDMXEntityList.cpp
new file mode 100644
--- /dev/null
+++ b/Engine/Plugins/Experimental/DigitalMultiplex/DMXEngine/Source/DMXEditor/Private/Widgets/SDMXEntityList.cpp
@@ -0,0 +1,2727 @@
+// Copyright Epic Games, Inc. All Rights Reserved.
+
+#include "Widgets/SDMXEntityList.h"
+
+#include "DMXEditor.h"
+#include "DMXEditorUtils.h"
+#include "DMXEditorLog.h"
+#include "DMXEditorStyle.h"
+#include "DMXProtocolConstants.h"
+#include "DMXProtocolTypes.h"
+#include "Library/DMXEntityController.h"
+#include "Library/DMXEntityFixturePatch.h"
+#include "Library/DMXEntityFixtureType.h"
+#include "Library/DMXLibrary.h"
+#include "Commands/DMXEditorCommands.h"
+#include "Widgets/SDMXEntityEditor.h"
+
+#include "Framework/Commands/UICommandList.h"
+#include "Framework/Commands/GenericCommands.h"
+#include "Framework/MultiBox/MultiBoxBuilder.h"
+#include "ScopedTransaction.h"
+#include "Dialogs/Dialogs.h"
+
+#include "EditorStyleSet.h"
+#include "Widgets/SBoxPanel.h"
+#include "Widgets/SToolTip.h"
+#include "Widgets/Docking/SDockTab.h"
+#include "Widgets/Input/SComboButton.h"
+#include "Widgets/Input/SButton.h"
+#include "Widgets/Input/SCheckBox.h"
+#include "Widgets/Input/SComboButton.h"
+#include "Widgets/Input/SSearchBox.h"
+#include "Widgets/Images/SImage.h"
+#include "Widgets/Layout/SScaleBox.h"
+#include "Widgets/Text/STextBlock.h"
+#include "Widgets/Text/SInlineEditableTextBlock.h"
+#include "Widgets/Views/SExpanderArrow.h"
+
+
+#define LOCTEXT_NAMESPACE "SDMXEntityListBase"
+
+///////////////////////////////////////////////////////////////////////////////
+// FDMXEntityListTreeNode
+
+FDMXTreeNodeBase::FDMXTreeNodeBase(FDMXTreeNodeBase::ENodeType InNodeType)
+	: DMXEntity(nullptr)
+	, WarningToolTip(FText::GetEmpty())
+	, ErrorToolTip(FText::GetEmpty())
+	, NodeType(InNodeType)
+	, bPendingRenameRequest(false)
+	, bShouldBeExpanded(false)
+	, FilterFlags((uint8)EFilteredState::Unknown)
+{
+}
+
+FString FDMXTreeNodeBase::GetDisplayString() const
+{
+	if (UDMXEntity* Entity = GetEntity())
+	{
+		return Entity->GetDisplayName();
+	}
+	return TEXT("null entity");
+}
+
+FText FDMXTreeNodeBase::GetDisplayName() const
+{
+	if (UDMXEntity* Entity = GetEntity())
+	{
+		return FText::FromString(Entity->GetDisplayName());
+	}
+
+	return LOCTEXT("NullEntityError", "Entity is null");
+}
+
+UDMXEntity* FDMXTreeNodeBase::GetEntity() const
+{
+	return DMXEntity.Get();
+}
+
+FDMXTreeNodeBase::ENodeType FDMXTreeNodeBase::GetNodeType() const
+{
+	return NodeType;
+}
+
+void FDMXTreeNodeBase::AddChild(TSharedPtr<FDMXTreeNodeBase> InChildPtr)
+{
+	if (InChildPtr.IsValid())
+	{
+		// unparent from previous parent
+		if (InChildPtr->GetParent().IsValid())
+		{
+			InChildPtr->RemoveFromParent();
+		}
+
+		InChildPtr->ParentNodePtr = SharedThis(this);
+		Children.Add(InChildPtr);
+	}
+}
+
+void FDMXTreeNodeBase::RemoveChild(TSharedPtr<FDMXTreeNodeBase> InChildPtr)
+{
+	if (InChildPtr.IsValid())
+	{
+		if (Children.Contains(InChildPtr))
+		{
+			Children.Remove(InChildPtr);
+			InChildPtr->ParentNodePtr = nullptr;
+		}
+	}
+}
+
+void FDMXTreeNodeBase::RemoveFromParent()
+{
+	if (GetParent().IsValid())
+	{
+		GetParent().Pin()->RemoveChild(SharedThis(this));
+	}
+}
+
+const TArray<TSharedPtr<FDMXTreeNodeBase>>& FDMXTreeNodeBase::GetChildren() const
+{
+	return Children;
+}
+
+void FDMXTreeNodeBase::ClearChildren()
+{
+	Children.Empty(0);
+}
+
+void FDMXTreeNodeBase::SortChildren()
+{
+	Children.HeapSort([](const TSharedPtr<FDMXTreeNodeBase>& Left, const TSharedPtr<FDMXTreeNodeBase>& Right)->bool
+		{
+			if (Left.IsValid() && Right.IsValid())
+			{
+				return *Left < *Right;
+			}
+			return false;
+		});
+}
+
+void FDMXTreeNodeBase::SortChildren(TFunction<bool (const TSharedPtr<FDMXTreeNodeBase>&, const TSharedPtr<FDMXTreeNodeBase>&)> Predicate)
+{
+	Children.Sort(Predicate);
+}
+
+bool FDMXTreeNodeBase::BroadcastRenameRequest()
+{
+	if (RenameRequestEvent.IsBound())
+	{
+		RenameRequestEvent.Execute();
+		bPendingRenameRequest = false;
+	}
+	else
+	{
+		bPendingRenameRequest = true;
+	}
+	return !bPendingRenameRequest;
+}
+
+bool FDMXTreeNodeBase::IsRenameRequestPending() const
+{
+	return bPendingRenameRequest;
+}
+
+TSharedPtr<FDMXTreeNodeBase> FDMXTreeNodeBase::FindChild(const UDMXEntity* InEntity, bool bRecursiveSearch /*= false*/, uint32* OutDepth /*= NULL*/) const
+{
+	TSharedPtr<FDMXTreeNodeBase> Result;
+
+	// Ensure that the given entity is valid
+	if (InEntity != nullptr)
+	{
+		// Look for a match in our set of child nodes
+		for (int32 ChildIndex = 0; ChildIndex < Children.Num() && !Result.IsValid(); ++ChildIndex)
+		{
+			if (InEntity == Children[ChildIndex]->GetEntity())
+			{
+				Result = Children[ChildIndex];
+			}
+			else if (bRecursiveSearch)
+			{
+				Result = Children[ChildIndex]->FindChild(InEntity, true, OutDepth);
+			}
+		}
+	}
+
+	if (OutDepth && Result.IsValid())
+	{
+		*OutDepth += 1;
+	}
+
+	return Result;
+}
+
+void FDMXTreeNodeBase::UpdateCachedFilterState(bool bMatchesFilter, bool bUpdateParent)
+{
+	bool bFlagsChanged = false;
+	if ((FilterFlags & EFilteredState::Unknown) == EFilteredState::Unknown)
+	{
+		FilterFlags = EFilteredState::FilteredOut;
+		bFlagsChanged = true;
+	}
+
+	if (bMatchesFilter)
+	{
+		bFlagsChanged |= (FilterFlags & EFilteredState::MatchesFilter) == 0;
+		FilterFlags |= EFilteredState::MatchesFilter;
+	}
+	else
+	{
+		bFlagsChanged |= (FilterFlags & EFilteredState::MatchesFilter) != 0;
+		FilterFlags &= ~EFilteredState::MatchesFilter;
+	}
+
+	const bool bHadChildMatch = (FilterFlags & EFilteredState::ChildMatches) != 0;
+	// refresh the cached child state (don't update the parent, we'll do that below if it's needed)
+	RefreshCachedChildFilterState(/*bUpdateParent =*/false);
+
+	bFlagsChanged |= bHadChildMatch != ((FilterFlags & EFilteredState::ChildMatches) != 0);
+	if (bUpdateParent && bFlagsChanged)
+	{
+		ApplyFilteredStateToParent();
+	}
+}
+
+void FDMXTreeNodeBase::SetExpansionState(bool bNewExpansionState)
+{
+	bShouldBeExpanded = bNewExpansionState;
+}
+
+void FDMXTreeNodeBase::SetWarningStatus(const FText& InWarningToolTip)
+{
+	WarningToolTip = InWarningToolTip;
+}
+
+void FDMXTreeNodeBase::SetErrorStatus(const FText& InErrorToolTip)
+{
+	ErrorToolTip = InErrorToolTip;
+}
+
+bool FDMXTreeNodeBase::operator<(const FDMXTreeNodeBase& Other) const
+{
+	const FString&& ThisName = this->GetDisplayString();
+	const FString&& OtherName = Other.GetDisplayString();
+	if (ThisName.IsNumeric() && OtherName.IsNumeric())
+	{
+		return FCString::Atoi(*ThisName) < FCString::Atoi(*OtherName);
+	}
+	
+	// If the names are strings with numbers at the end, separate them to compare name then number
+	FString NameOnlyThis;
+	FString NameOnlyOther;
+	int32 NumberThis = 0;
+	int32 NumberOther = 0;
+	if (FDMXEditorUtils::GetNameAndIndexFromString(ThisName, NameOnlyThis, NumberThis)
+		&& FDMXEditorUtils::GetNameAndIndexFromString(OtherName, NameOnlyOther, NumberOther)
+		&& NameOnlyThis.Equals(NameOnlyOther))
+	{
+		return NumberThis < NumberOther;
+	}
+
+	return ThisName < OtherName;
+}
+
+void FDMXTreeNodeBase::RefreshCachedChildFilterState(bool bUpdateParent)
+{
+	const bool bContainedMatch = !IsFlaggedForFiltration();
+
+	FilterFlags &= ~EFilteredState::ChildMatches;
+	for (TSharedPtr<FDMXTreeNodeBase> Child : Children)
+	{
+		if (!Child->IsFlaggedForFiltration())
+		{
+			FilterFlags |= EFilteredState::ChildMatches;
+			break;
+		}
+	}
+	const bool bContainsMatch = !IsFlaggedForFiltration();
+
+	const bool bStateChange = bContainedMatch != bContainsMatch;
+	if (bUpdateParent && bStateChange)
+	{
+		ApplyFilteredStateToParent();
+	}
+}
+
+void FDMXTreeNodeBase::ApplyFilteredStateToParent()
+{
+	FDMXTreeNodeBase* Child = this;
+
+	while (Child->ParentNodePtr.IsValid())
+	{
+		FDMXTreeNodeBase* Parent = Child->ParentNodePtr.Pin().Get();
+
+		if (!IsFlaggedForFiltration())
+		{
+			if ((Parent->FilterFlags & EFilteredState::ChildMatches) == 0)
+			{
+				Parent->FilterFlags |= EFilteredState::ChildMatches;
+			}
+			else
+			{
+				// all parents from here on up should have the flag
+				break;
+			}
+		}
+		// have to see if this was the only child contributing to this flag
+		else if (Parent->FilterFlags & EFilteredState::ChildMatches)
+		{
+			Parent->FilterFlags &= ~EFilteredState::ChildMatches;
+			for (const TSharedPtr<FDMXTreeNodeBase>& Sibling : Parent->Children)
+			{
+				if (Sibling.Get() == Child)
+				{
+					continue;
+				}
+
+				if (Sibling->FilterFlags & EFilteredState::FilteredInMask)
+				{
+					Parent->FilterFlags |= EFilteredState::ChildMatches;
+					break;
+				}
+			}
+
+			if (Parent->FilterFlags & EFilteredState::ChildMatches)
+			{
+				// another child added the flag back
+				break;
+			}
+		}
+		Child = Parent;
+	}
+}
+
+FDMXCategoryTreeNode::FDMXCategoryTreeNode(ECategoryType InCategoryType, FText InCategoryName, const void* InCategoryValuePtr, const FText& InToolTip /*= FText::GetEmpty()*/)
+	: FDMXTreeNodeBase(FDMXTreeNodeBase::ENodeType::CategoryNode)
+	, CategoryType(InCategoryType)
+	, CategoryValuePtr(InCategoryValuePtr)
+	, CategoryName(InCategoryName)
+	, ToolTip(InToolTip)
+{
+}
+
+FString FDMXCategoryTreeNode::GetDisplayString() const
+{
+	return CategoryName.ToString();
+}
+
+FText FDMXCategoryTreeNode::GetDisplayName() const
+{
+	return CategoryName;
+}
+
+FDMXTreeNodeBase::ECategoryType FDMXCategoryTreeNode::GetCategoryType() const
+{
+	return CategoryType;
+}
+
+FDMXEntityBaseTreeNode::FDMXEntityBaseTreeNode(UDMXEntity* InEntity)
+	: FDMXTreeNodeBase(FDMXTreeNodeBase::ENodeType::EntityNode)
+{
+	DMXEntity = InEntity;
+}
+
+bool FDMXEntityBaseTreeNode::IsEntityNode() const
+{
+	return true;
+}
+
+///////////////////////////////////////////////////////////////////////////////
+// FDMXEntityDragDropOperation
+
+FDMXEntityDragDropOperation::FDMXEntityDragDropOperation(UDMXLibrary* InLibrary, TWeakPtr<SDMXEntityList> InEntityList, TArray<TSharedPtr<FDMXEntityBaseTreeNode>>&& InEntities)
+	: DraggedFromLibrary(InLibrary)
+	, DraggedEntities(MoveTemp(InEntities))
+	, EntityList(InEntityList)
+	, bValidDropTarget(false)
+{
+	DraggedLabel = DraggedEntities.Num() == 1
+		? FText::FromString(TEXT("'") + DraggedEntities[0]->GetEntity()->GetDisplayName() + TEXT("'"))
+		: FDMXEditorUtils::GetEntityTypeNameText(DraggedEntities[0]->GetEntity()->GetClass(), true);
+	
+	SetDraggingFromMultipleCategories();
+
+	Construct();
+}
+
+void FDMXEntityDragDropOperation::SetHoveredEntity(TSharedPtr<FDMXEntityBaseTreeNode> InEntityNode, UDMXLibrary* InLibrary, TSubclassOf<UDMXEntity> TabEntitiesType)
+{
+	HoveredEntity = InEntityNode;
+	HoveredLibrary = InLibrary;
+	HoveredTabType = TabEntitiesType;
+	HoverTargetChanged();
+}
+
+void FDMXEntityDragDropOperation::SetHoveredCategory(TSharedPtr<FDMXCategoryTreeNode> InCategory, UDMXLibrary* InLibrary, TSubclassOf<UDMXEntity> TabEntitiesType)
+{
+	HoveredCategory = InCategory;
+	HoveredLibrary = InLibrary;
+	HoveredTabType = TabEntitiesType;
+	HoverTargetChanged();
+}
+
+void FDMXEntityDragDropOperation::DroppedOnEntity(TSharedRef<FDMXEntityBaseTreeNode> InEntity, UDMXLibrary* InLibrary, TSubclassOf<UDMXEntity> TabEntitiesType)
+{
+	if (!bValidDropTarget) return;
+	check(EntityList.IsValid() && DraggedFromLibrary != nullptr && HoveredEntity.IsValid());
+
+	// Register transaction and current DMX library state for Undo
+	const FScopedTransaction ReorderTransaction = FScopedTransaction(
+		FText::Format(LOCTEXT("ReorderEntities", "Reorder {0}|plural(one=Entity, other=Entities)"), DraggedEntities.Num())
+	);
+	DraggedFromLibrary->Modify();
+
+	// The index of the Entity we're about to insert the dragged ones before
+	const int32 InsertBeforeIndex = DraggedFromLibrary->FindEntityIndex(HoveredEntity->GetEntity());
+	check(InsertBeforeIndex != INDEX_NONE);
+	ReorderEntities(InsertBeforeIndex);
+
+	if (IsDraggingToDifferentCategory())
+	{
+		SetPropertyForNewCategory();
+	}
+
+	// Display the changes in the Entities list
+	EntityList.Pin()->UpdateTree();
+}
+
+void FDMXEntityDragDropOperation::DroppedOnCategory(TSharedRef<FDMXCategoryTreeNode> InCategory, UDMXLibrary* InLibrary, TSubclassOf<UDMXEntity> TabEntitiesType)
+{
+	if (!bValidDropTarget)
+	{
+		return; 
+	}
+	check(EntityList.IsValid() && DraggedFromLibrary != nullptr && HoveredCategory.IsValid());
+
+	// Register transaction and current DMX library state for Undo
+	const FScopedTransaction ChangeCategoryTransaction = FScopedTransaction(
+		FText::Format(LOCTEXT("ChangeEntitiesCategory", "Change {0}|plural(one=Entity, other=Entities) category"), DraggedEntities.Num())
+	);
+	DraggedFromLibrary->Modify();
+
+	SetPropertyForNewCategory();
+
+	if (HoveredCategory->GetChildren().Num() > 0)
+	{
+		// Index after last entity in hovered category
+		UDMXEntity* LastEntityInCategory = HoveredCategory->GetChildren().Last(0)->GetEntity();
+		const int32 LastEntityIndex = DraggedFromLibrary->FindEntityIndex(LastEntityInCategory);
+		check(LastEntityIndex != INDEX_NONE);
+		// Move dragged entities after the last ones in the category
+		ReorderEntities(LastEntityIndex + 1);
+	}
+
+	// Display the changes in the Entities list
+	EntityList.Pin()->UpdateTree();
+}
+
+void FDMXEntityDragDropOperation::ReorderEntities(int32 NewIndex)
+{
+	// Reverse for to keep dragged entities order
+	for (int32 EntityIndex = DraggedEntities.Num() - 1; EntityIndex > -1; --EntityIndex)
+	{
+		if (DraggedEntities[EntityIndex].IsValid() && DraggedEntities[EntityIndex]->GetEntity() != nullptr)
+		{
+			UDMXEntity* Entity = DraggedEntities[EntityIndex]->GetEntity();
+			DraggedFromLibrary->SetEntityIndex(Entity, NewIndex);
+		}
+	}
+}
+
+void FDMXEntityDragDropOperation::SetPropertyForNewCategory()
+{
+	const TSharedPtr<FDMXCategoryTreeNode>&& NewCategory = GetTargetCategory();
+	check(NewCategory.IsValid());
+
+	if (!NewCategory->IsCategoryValueValid())
+	{
+		return; 
+	}
+
+	switch (NewCategory->GetCategoryType())
+	{
+	case FDMXTreeNodeBase::ECategoryType::DeviceProtocol:
+		{
+			const FDMXProtocolName& DeviceProtocol = *NewCategory->GetCategoryValue<FDMXProtocolName>();
+			for (TSharedPtr<FDMXEntityBaseTreeNode> Entity : DraggedEntities)
+			{
+				if (UDMXEntityController* Controller = Cast<UDMXEntityController>(Entity->GetEntity()))
+				{
+					Controller->Modify();
+					Controller->DeviceProtocol = DeviceProtocol;
+				}
+			}
+		}
+		break;
+
+	case FDMXTreeNodeBase::ECategoryType::DMXCategory:
+		{
+			const FDMXFixtureCategory& FixtureCategory = *NewCategory->GetCategoryValue<FDMXFixtureCategory>();
+			for (TSharedPtr<FDMXEntityBaseTreeNode> Entity : DraggedEntities)
+			{
+				if (UDMXEntityFixtureType* FixtureType = Cast<UDMXEntityFixtureType>(Entity->GetEntity()))
+				{
+					FixtureType->Modify();
+					FixtureType->DMXCategory = FixtureCategory;
+				}
+			}
+		}
+		break;
+
+	case FDMXTreeNodeBase::ECategoryType::UniverseID:
+	case FDMXTreeNodeBase::ECategoryType::FixtureAssignmentState:
+		{
+			const uint32& UniverseID = *NewCategory->GetCategoryValue<uint32>();
+			for (TSharedPtr<FDMXEntityBaseTreeNode> Entity : DraggedEntities)
+			{
+				if (UDMXEntityFixturePatch* FixturePatch = Cast<UDMXEntityFixturePatch>(Entity->GetEntity()))
+				{
+					FixturePatch->Modify();
+					FixturePatch->UniverseID = UniverseID;
+				}
+			}
+		}
+
+	default:
+		// The other category types don't change properties
+		break;
+	}
+}
+
+void FDMXEntityDragDropOperation::Construct()
+{
+	// Create the drag-drop decorator window
+	CursorDecoratorWindow = SWindow::MakeCursorDecorator();
+	const bool bShowImmediately = false;
+	FSlateApplication::Get().AddWindow(CursorDecoratorWindow.ToSharedRef(), bShowImmediately);
+
+	HoverTargetChanged();
+}
+
+void FDMXEntityDragDropOperation::HoverTargetChanged()
+{
+	check(EntityList.IsValid());
+	TSharedPtr<SDMXEntityList> PinnedList = EntityList.Pin();
+
+	if (HoveredLibrary != nullptr && DraggedFromLibrary != HoveredLibrary)
+	{
+		// For now, we don't allow dragging entities from one library to the other
+		SetFeedbackMessageError(FText::Format(
+			LOCTEXT("CantDragToDifferentLibrary", "Cannot move {0} outside {1}|plural(one=its, other=their) library"),
+			DraggedLabel,
+			DraggedEntities.Num()
+		));
+		bValidDropTarget = false;
+		return;
+	}
+	else if (HoveredTabType != UDMXEntity::StaticClass() && PinnedList->GetListType() != HoveredTabType)
+	{
+		// Don't allow dragging entities from a type onto a different type tab
+		SetFeedbackMessageError(FText::Format(
+			LOCTEXT("CantDragToDifferentType", "Cannot move {0} to {1} tab"),
+			DraggedLabel,
+			FDMXEditorUtils::GetEntityTypeNameText(HoveredTabType, /*bPlural =*/true)
+		));
+		bValidDropTarget = false;
+		return;
+	}
+	else if (HoveredEntity.IsValid() && HoveredEntity->GetEntity() != nullptr)
+	{
+		if (IsDraggingToDifferentCategory())
+		{
+			// If dragging into a different category, some property will have to be changed
+
+			check(HoveredEntity->GetParent().IsValid());
+			FText PropertyChangeName;
+			FText PropertyNewValue;
+			bValidDropTarget = GetCategoryPropertyNameFromTabType(PropertyChangeName, PropertyNewValue);
+			if (bValidDropTarget)
+			{
+				SetFeedbackMessageOK(FText::Format(
+					LOCTEXT("ReorderAndSetProperty", "Reorder {0} before '{1}'\nSet {2} = '{3}'"),
+					DraggedLabel,
+					HoveredEntity->GetDisplayName(),
+					PropertyChangeName,
+					PropertyNewValue
+				));
+			}
+			return;
+		}
+		else if (DraggedEntities.Num() == 1 && DraggedEntities[0] == HoveredEntity)
+		{
+			SetFeedbackMessageError(FText::Format(
+				LOCTEXT("ReorderBeforeItself", "Cannot reorder {0} before itself"),
+				DraggedLabel
+			));
+			bValidDropTarget = false;
+			return;
+		}
+		else
+		{
+			// Reordering between entities of same category is fine
+			SetFeedbackMessageOK(FText::Format(
+				LOCTEXT("ReorderEntities", "Reorder {0} before '{1}'"),
+				DraggedLabel,
+				HoveredEntity->GetDisplayName()
+			));
+			bValidDropTarget = true;
+			return;
+		}
+	}
+	else if (HoveredCategory.IsValid())
+	{
+		if (!IsDraggingToDifferentCategory())
+		{
+			// Good visual feedback, but we register as invalid drop target. There wouldn't be
+			// any change by dragging the items into their own category.
+			SetFeedbackMessageOK(FText::Format(
+				LOCTEXT("DragIntoSelfCategory", "The selected {0} {1}|plural(one=is, other=are) already in this category"),
+				FDMXEditorUtils::GetEntityTypeNameText(PinnedList->GetListType(), DraggedEntities.Num() > 1),
+				DraggedEntities.Num()
+			));
+			bValidDropTarget = false;
+			return;
+		}
+		else
+		{
+			// Some (or all) items will have a property changed because they come from another category
+			FText PropertyChangeName;
+			FText PropertyNewValue;
+			bValidDropTarget = GetCategoryPropertyNameFromTabType(PropertyChangeName, PropertyNewValue);
+			if (bValidDropTarget)
+			{
+				SetFeedbackMessageOK(FText::Format(
+					LOCTEXT("ReorderAndSetProperty", "{0}\nSet {1} = '{2}'"),
+					DraggedLabel,
+					PropertyChangeName,
+					PropertyNewValue
+				));
+			}
+			return;
+		}
+	}
+
+	SetFeedbackMessageError(DraggedLabel);
+	bValidDropTarget = false;
+}
+
+void FDMXEntityDragDropOperation::SetDraggingFromMultipleCategories()
+{
+	bDraggingFromMultipleCategories = false;
+
+	TWeakPtr<FDMXTreeNodeBase> FirstCategory = DraggedEntities[0]->GetParent();
+	for (TSharedPtr<FDMXEntityBaseTreeNode> DraggedEntity : DraggedEntities)
+	{
+		if (DraggedEntity->GetParent() != FirstCategory)
+		{
+			bDraggingFromMultipleCategories = true;
+			break;
+		}
+	}
+}
+
+bool FDMXEntityDragDropOperation::GetCategoryPropertyNameFromTabType(FText& PropertyName, FText& CategoryPropertyValue)
+{
+	check(EntityList.IsValid());
+	const TSubclassOf<UDMXEntity> DraggedEntitiesType = EntityList.Pin()->GetListType();
+	TSharedPtr<FDMXCategoryTreeNode>&& TargetCategory = GetTargetCategory();
+	check(TargetCategory.IsValid());
+
+	if (DraggedEntitiesType->IsChildOf(UDMXEntityController::StaticClass()))
+	{
+		PropertyName = LOCTEXT("Property_DeviceProtocol", "Device Protocol");
+		CategoryPropertyValue = TargetCategory->GetDisplayName();
+	}
+	else if (DraggedEntitiesType->IsChildOf(UDMXEntityFixtureType::StaticClass()))
+	{
+		PropertyName = LOCTEXT("Property_DMXCategory", "DMX Category");
+		CategoryPropertyValue = TargetCategory->GetDisplayName();
+	}
+	else if (DraggedEntitiesType->IsChildOf(UDMXEntityFixturePatch::StaticClass()))
+	{
+		if (TargetCategory->IsCategoryValueValid())
+		{
+			PropertyName = LOCTEXT("Property_Universe", "Universe");
+			const uint32 UniverseID = *TargetCategory->GetCategoryValue<uint32>();
+			CategoryPropertyValue = UniverseID == MAX_uint32
+				? LOCTEXT("UnassignedUniverseIDValue", "Unassigned")
+				: FText::AsNumber(UniverseID);
+		}
+		else
+		{
+			// Can't assign universe by simply dragging into "Assigned Fixtures"
+			SetFeedbackMessageError(FText::Format(
+				LOCTEXT("DragCantChangeUniverse", "Drag onto a Universe to assign it"),
+				FDMXEditorUtils::GetEntityTypeNameText(DraggedEntitiesType, DraggedEntities.Num() > 1)
+			));
+			return false;
+		}
+	}
+	else
+	{
+		// Dragged Entities are of unimplemented type!
+		SetFeedbackMessageError(FText::Format(
+			LOCTEXT("DragUnimplementedCategoryChange", "Cannot move {0} to another category"),
+			FDMXEditorUtils::GetEntityTypeNameText(DraggedEntitiesType, DraggedEntities.Num() > 1)
+		));
+		return false;
+	}
+
+	return true;
+}
+
+bool FDMXEntityDragDropOperation::IsDraggingToDifferentCategory() const
+{
+	if (HoveredEntity.IsValid())
+	{
+		return bDraggingFromMultipleCategories || DraggedEntities[0]->GetParent() != HoveredEntity->GetParent();
+	}
+	else if (HoveredCategory.IsValid())
+	{
+		return bDraggingFromMultipleCategories || DraggedEntities[0]->GetParent() != HoveredCategory;
+	}
+
+	UE_LOG_DMXEDITOR(Fatal, TEXT("%S was called and there was no hovered Entity or Category"), __FUNCTION__);
+	return false;
+}
+
+TSharedPtr<FDMXCategoryTreeNode> FDMXEntityDragDropOperation::GetTargetCategory() const
+{
+	if (HoveredCategory.IsValid())
+	{
+		return HoveredCategory;
+	}
+	else if (HoveredEntity.IsValid())
+	{
+		return StaticCastSharedPtr<FDMXCategoryTreeNode>(HoveredEntity->GetParent().Pin());
+	}
+
+	return nullptr;
+}
+
+void FDMXEntityDragDropOperation::SetFeedbackMessageError(const FText& Message)
+{
+	const FSlateBrush* StatusSymbol = FEditorStyle::GetBrush(TEXT("Graph.ConnectorFeedback.Error"));
+	SetFeedbackMessage(StatusSymbol, Message);
+}
+
+void FDMXEntityDragDropOperation::SetFeedbackMessageOK(const FText& Message)
+{
+	const FSlateBrush* StatusSymbol = FEditorStyle::GetBrush(TEXT("Graph.ConnectorFeedback.OK"));
+	SetFeedbackMessage(StatusSymbol, Message);
+}
+
+void FDMXEntityDragDropOperation::SetFeedbackMessage(const FSlateBrush* Icon, const FText& Message)
+{
+	if (!Message.IsEmpty())
+	{
+		CursorDecoratorWindow->ShowWindow();
+		CursorDecoratorWindow->SetContent
+		(
+			SNew(SBorder)
+			.BorderImage(FEditorStyle::GetBrush("Graph.ConnectorFeedback.Border"))
+			[
+				SNew(SHorizontalBox)
+				+SHorizontalBox::Slot()
+				.AutoWidth()
+				.Padding(3.0f)
+				[
+					SNew(SScaleBox)
+					.Stretch(EStretch::ScaleToFit)
+					[
+						SNew(SImage)
+						.Image(Icon)
+					]
+				]
+				+SHorizontalBox::Slot()
+				.AutoWidth()
+				.Padding(3.0f)
+				.MaxWidth(500)
+				.VAlign(VAlign_Center)
+				[
+					SNew(STextBlock)
+					.WrapTextAt(480.0f)
+					.Text(Message)
+				]
+			]
+		);
+	}
+	else
+	{
+		CursorDecoratorWindow->HideWindow();
+		CursorDecoratorWindow->SetContent(SNullWidget::NullWidget);
+	}
+}
+
+///////////////////////////////////////////////////////////////////////////////
+// SDMXCategoryRow
+
+void SDMXCategoryRow::Construct(const FArguments& InArgs, const TSharedRef<STableViewBase>& InOwnerTableView, TSharedPtr<FDMXTreeNodeBase> InNodePtr, bool bInIsRootCategory, TWeakPtr<SDMXEntityList> InEditorList)
+{
+	EditorListPtr = InEditorList;
+	TreeNodePtr = StaticCastSharedPtr<FDMXCategoryTreeNode>(InNodePtr);
+	check(TreeNodePtr.IsValid());
+
+	// background color tint
+	const FLinearColor BackgroundTint(0.6f, 0.6f, 0.6f, bInIsRootCategory ? 1.0f : 0.3f);
+
+	// rebuilds the whole table row from scratch
+	SDMXTableRowType::ChildSlot
+	.Padding(0.0f, 2.0f, 0.0f, 0.0f)
+	[
+		SAssignNew(ContentBorder, SBorder)
+		.BorderImage(this, &SDMXCategoryRow::GetBackgroundImage)
+		.Padding(FMargin(0.0f, 3.0f))
+		.BorderBackgroundColor(BackgroundTint)
+		.ToolTipText(TreeNodePtr.Pin()->GetToolTip())
+		[
+			SNew(SHorizontalBox)
+
+			+ SHorizontalBox::Slot()
+			.VAlign(VAlign_Center)
+			.Padding(2.0f, 2.0f, 2.0f, 2.0f)
+			.AutoWidth()
+			[
+				SNew(SExpanderArrow, SDMXTableRowType::SharedThis(this))
+			]
+
+			+ SHorizontalBox::Slot()
+			.VAlign(VAlign_Center)
+			[
+				InArgs._Content.Widget
+			]
+		]
+	];
+
+	SDMXTableRowType::ConstructInternal(
+		SDMXTableRowType::FArguments()
+			.Style(FEditorStyle::Get(), "DetailsView.TreeView.TableRow")
+			.ShowSelection(false),
+		InOwnerTableView
+	);
+}
+
+void SDMXCategoryRow::SetContent(TSharedRef< SWidget > InContent)
+{
+	ContentBorder->SetContent(InContent);
+}
+
+void SDMXCategoryRow::SetRowContent(TSharedRef< SWidget > InContent)
+{
+	ContentBorder->SetContent(InContent);
+}
+
+void SDMXCategoryRow::OnDragEnter(const FGeometry& MyGeometry, const FDragDropEvent& DragDropEvent)
+{
+	check(GetNode().IsValid());
+
+	if (TSharedPtr<FDMXEntityDragDropOperation> EntityDrag = DragDropEvent.GetOperationAs<FDMXEntityDragDropOperation>())
+	{
+		if (TSharedPtr<SDMXEntityList> EditorListPinned = EditorListPtr.Pin())
+		{
+			EntityDrag->SetHoveredCategory(GetNode(), EditorListPinned->GetDMXLibrary(), EditorListPinned->GetListType());
+		}
+	}
+}
+
+void SDMXCategoryRow::OnDragLeave(const FDragDropEvent& DragDropEvent)
+{
+	if (TSharedPtr<FDMXEntityDragDropOperation> EntityDrag = DragDropEvent.GetOperationAs<FDMXEntityDragDropOperation>())
+	{
+		EntityDrag->SetHoveredCategory(nullptr, nullptr, UDMXEntity::StaticClass());
+	}
+}
+
+FReply SDMXCategoryRow::OnDrop(const FGeometry& MyGeometry, const FDragDropEvent& DragDropEvent)
+{
+	check(GetNode().IsValid());
+
+	if (TSharedPtr<FDMXEntityDragDropOperation> EntityDrag = DragDropEvent.GetOperationAs<FDMXEntityDragDropOperation>())
+	{
+		if (TSharedPtr<SDMXEntityList> EditorListPinned = EditorListPtr.Pin())
+		{
+			EntityDrag->DroppedOnCategory(GetNode().ToSharedRef(), EditorListPinned->GetDMXLibrary(), EditorListPinned->GetListType());
+		}
+		return FReply::Handled();
+	}
+
+	return FReply::Unhandled();
+}
+
+const FSlateBrush* SDMXCategoryRow::GetBackgroundImage() const
+{
+	if (SDMXTableRowType::IsHovered())
+	{
+		return SDMXTableRowType::IsItemExpanded()
+			? FEditorStyle::GetBrush("DetailsView.CategoryTop_Hovered")
+			: FEditorStyle::GetBrush("DetailsView.CollapsedCategory_Hovered");
+	}
+	else
+	{
+		return SDMXTableRowType::IsItemExpanded()
+			? FEditorStyle::GetBrush("DetailsView.CategoryTop")
+			: FEditorStyle::GetBrush("DetailsView.CollapsedCategory");
+	}
+}
+
+void SDMXEntityRow::Construct(const FArguments& InArgs, TSharedPtr<FDMXTreeNodeBase> InNodePtr, TSharedPtr<STableViewBase> InOwnerTableView, TWeakPtr<SDMXEntityList> InEditorList)
+{
+	TreeNodePtr = StaticCastSharedPtr<FDMXEntityBaseTreeNode>(InNodePtr);
+	EditorListPtr = InEditorList;
+
+	OnEntityDragged = InArgs._OnEntityDragged;
+	OnGetFilterText = InArgs._OnGetFilterText;
+
+	const FSlateFontInfo&& NameFont = FCoreStyle::GetDefaultFontStyle("Regular", 10);
+
+	// Without ETableRowSignalSelectionMode::Instantaneous, when the user is editing a property in
+	// the inspector panel and then clicks on a different row on the list panel, the selection event
+	// is deferred. But because we update the tree right after a property change and that triggers
+	// selection changes too, the selection change event is triggered only from UpdateTree, with
+	// Direct selection mode, which doesn't trigger the SDMXEntityList::OnSelectionUpdated event.
+	// This setting forces the event with OnMouseClick selection type to be fired as soon as the
+	// row is clicked.
+	SDMXTableRowType::FArguments BaseArguments = SDMXTableRowType::FArguments()
+		.SignalSelectionMode(ETableRowSignalSelectionMode::Instantaneous)
+		.OnDragDetected(this, &SDMXEntityRow::HandleOnDragDetected);
+
+	SDMXTableRowType::Construct(BaseArguments, InOwnerTableView.ToSharedRef());
+
+	// Horizontal box to add content conditionally later 
+	TSharedPtr<SHorizontalBox> ContentBox;
+
+	SetContent
+	(
+		SAssignNew(ContentBox, SHorizontalBox)
+		// Status icon to show the user if there's an error with the Entity's usability
+		+ SHorizontalBox::Slot()
+		.AutoWidth()
+		.HAlign(HAlign_Left)
+		.VAlign(VAlign_Center)
+		[
+			SNew(SImage)
+			.Image(this, &SDMXEntityRow::GetStatusIcon)
+			.ToolTipText(this, &SDMXEntityRow::GetStatusToolTip)
+		]
+		
+		// Entity's name
+		+ SHorizontalBox::Slot()
+		.FillWidth(1.0f)
+		.VAlign(VAlign_Center)
+		.Padding(3.0f, 0.0f)
+		[
+			SAssignNew(InlineRenameWidget, SInlineEditableTextBlock)
+			.Text(this, &SDMXEntityRow::GetDisplayText)
+			.Font(NameFont)
+			.HighlightText(this, &SDMXEntityRow::GetFilterText)
+			.ToolTipText(this, &SDMXEntityRow::GetToolTipText)
+			.OnTextCommitted(this, &SDMXEntityRow::OnNameTextCommit)
+			.OnVerifyTextChanged(this, &SDMXEntityRow::OnNameTextVerifyChanged)
+			.IsSelected(this, &SDMXTableRowType::IsSelected)
+			.IsReadOnly(false)
+		]
+	);
+
+	// Per entity type customizations
+	if (UDMXEntityFixturePatch* EntityAsPatch = Cast<UDMXEntityFixturePatch>(InNodePtr->GetEntity()))
+	{
+		// For Fixture Patch we display a channel-auto-assignment box and the channel range occupied by it
+
+		// Auto channel assignment check box
+		OnAutoAssignChannelStateChanged = InArgs._OnAutoAssignChannelStateChanged;
+
+		ContentBox->AddSlot()
+			.AutoWidth()
+			.VAlign(VAlign_Center)
+			[
+				SNew(SCheckBox)
+				.IsChecked(EntityAsPatch->bAutoAssignAddress ? ECheckBoxState::Checked : ECheckBoxState::Unchecked)
+				.OnCheckStateChanged(this, &SDMXEntityRow::OnAutoAssignChannelBoxStateChanged)
+				.ToolTipText(LOCTEXT("AutoAssignChannelToolTip", "Auto-assign channel from drag/drop list order"))
+			];
+
+		// Used channels range labels
+		const FSlateFontInfo&& ChannelFont = FCoreStyle::GetDefaultFontStyle("Bold", 8);
+		const FLinearColor ChannelLabelColor(1.0f, 1.0f, 1.0f, 0.8f);
+		const float MinChannelTextWidth = 20.0f;
+
+		// Starting channel number
+		ContentBox->AddSlot()
+			.AutoWidth()
+			.VAlign(VAlign_Center)
+			[
+				SNew(SBorder)
+				.VAlign(VAlign_Center)
+				.Padding(FMargin(0.0f, 3.0f))
+				.BorderImage(FDMXEditorStyle::Get().GetBrush("DMXEditor.BlackBrush"))
+				.BorderBackgroundColor(FLinearColor(1.0f, 1.0f, 1.0f, 0.15f))
+				.ToolTipText(LOCTEXT("ChannelStartToolTip", "Channels range: start"))
+				[
+					SNew(STextBlock)
+					.Text(this, &SDMXEntityRow::GetStartingChannelLabel)
+					.Font(ChannelFont)
+					.ColorAndOpacity(ChannelLabelColor)
+					.MinDesiredWidth(MinChannelTextWidth)
+					.Justification(ETextJustify::Center)
+				]
+			];
+
+		// Ending channel number
+		ContentBox->AddSlot()
+			.AutoWidth()
+			.VAlign(VAlign_Center)
+			[
+				SNew(SBorder)
+				.VAlign(VAlign_Center)
+				.Padding(FMargin(0.0f, 3.0f))
+				.BorderImage(FDMXEditorStyle::Get().GetBrush("DMXEditor.BlackBrush"))
+				.BorderBackgroundColor(FLinearColor(1.0f, 1.0f, 1.0f, 0.25f)) // darker background
+				.ColorAndOpacity(FLinearColor(1.0f, 1.0f, 1.0f, 0.5f)) // darker text
+				.ToolTipText(LOCTEXT("ChannelEndToolTip", "Channels range: end"))
+				[
+					SNew(STextBlock)
+					.Text(this, &SDMXEntityRow::GetEndingChannelLabel)
+					.Font(ChannelFont)
+					.ColorAndOpacity(ChannelLabelColor)
+					.MinDesiredWidth(MinChannelTextWidth)
+					.Justification(ETextJustify::Center)
+				]
+			];
+	}
+
+	InNodePtr->OnRenameRequest().BindSP(InlineRenameWidget.Get(), &SInlineEditableTextBlock::EnterEditingMode);
+}
+
+FText SDMXEntityRow::GetDisplayText() const
+{
+	if (TreeNodePtr.IsValid())
+	{
+		return TreeNodePtr.Pin()->GetDisplayName();
+	}
+	return LOCTEXT("InvalidNodeLabel", "Invalid Node");
+}
+
+FText SDMXEntityRow::GetStartingChannelLabel() const
+{
+	check(TreeNodePtr.IsValid());
+	
+	if (TSharedPtr<FDMXTreeNodeBase> PinnedNode = TreeNodePtr.Pin())
+	{
+		if (UDMXEntityFixturePatch* Patch = Cast<UDMXEntityFixturePatch>(PinnedNode->GetEntity()))
+		{
+			return FText::AsNumber(Patch->GetStartingChannel());
+		}
+	}
+
+	return FText::GetEmpty();
+}
+
+FText SDMXEntityRow::GetEndingChannelLabel() const
+{
+	check(TreeNodePtr.IsValid());
+
+	if (TSharedPtr<FDMXTreeNodeBase> PinnedNode = TreeNodePtr.Pin())
+	{
+		if (UDMXEntityFixturePatch* Patch = Cast<UDMXEntityFixturePatch>(PinnedNode->GetEntity()))
+		{
+			return FText::AsNumber(Patch->GetStartingChannel() + Patch->GetChannelSpan() - 1);
+		}
+	}
+
+	return FText::GetEmpty();
+}
+
+void SDMXEntityRow::OnAutoAssignChannelBoxStateChanged(ECheckBoxState NewState)
+{
+	if (OnAutoAssignChannelStateChanged.IsBound())
+	{
+		switch (NewState)
+		{
+		case ECheckBoxState::Unchecked:
+			OnAutoAssignChannelStateChanged.Execute(false);
+			break;
+		case ECheckBoxState::Checked:
+			OnAutoAssignChannelStateChanged.Execute(true);
+			break;
+		case ECheckBoxState::Undetermined:
+		default:
+			break;
+		}
+	}
+}
+
+void SDMXEntityRow::OnDragEnter(const FGeometry& MyGeometry, const FDragDropEvent& DragDropEvent)
+{
+	check(GetNode().IsValid() && GetNode()->GetEntity() != nullptr);
+
+	if (TSharedPtr<FDMXEntityDragDropOperation> EntityDrag = DragDropEvent.GetOperationAs<FDMXEntityDragDropOperation>())
+	{
+		if (TSharedPtr<SDMXEntityList> EditorListPinned = EditorListPtr.Pin())
+		{
+			EntityDrag->SetHoveredEntity(GetNode(), EditorListPinned->GetDMXLibrary(), EditorListPinned->GetListType());
+		}
+	}
+}
+
+void SDMXEntityRow::OnDragLeave(const FDragDropEvent& DragDropEvent)
+{
+	if (TSharedPtr<FDMXEntityDragDropOperation> EntityDrag = DragDropEvent.GetOperationAs<FDMXEntityDragDropOperation>())
+	{
+		EntityDrag->SetHoveredEntity(nullptr, nullptr, UDMXEntity::StaticClass());
+	}
+}
+
+FReply SDMXEntityRow::OnDrop(const FGeometry& MyGeometry, const FDragDropEvent& DragDropEvent)
+{
+	check(GetNode().IsValid() && GetNode()->GetEntity() != nullptr);
+
+	if (TSharedPtr<FDMXEntityDragDropOperation> EntityDrag = DragDropEvent.GetOperationAs<FDMXEntityDragDropOperation>())
+	{
+		if (TSharedPtr<SDMXEntityList> EditorListPinned = EditorListPtr.Pin())
+		{
+			EntityDrag->DroppedOnEntity(GetNode().ToSharedRef(), EditorListPinned->GetDMXLibrary(), EditorListPinned->GetListType());
+		}
+		return FReply::Handled();
+	}
+
+	return FReply::Unhandled();
+}
+
+bool SDMXEntityRow::OnNameTextVerifyChanged(const FText& InNewText, FText& OutErrorMessage)
+{
+	FString TextAsString = InNewText.ToString();
+	if (TextAsString.Equals(TreeNodePtr.Pin()->GetDisplayString()))
+	{
+		return true;
+	}
+
+	return FDMXEditorUtils::ValidateEntityName(
+		TextAsString,
+		EditorListPtr.Pin()->GetDMXLibrary(),
+		TreeNodePtr.Pin()->GetEntity()->GetClass(),
+		OutErrorMessage
+	);
+}
+
+void SDMXEntityRow::OnNameTextCommit(const FText& InNewName, ETextCommit::Type InTextCommit)
+{
+	const FString NewNameString = InNewName.ToString();
+	TSharedPtr<FDMXTreeNodeBase> Node = TreeNodePtr.Pin();
+
+	// Check if the name is unchanged
+	if (NewNameString.Equals(Node->GetDisplayString()))
+	{
+		return;
+	}
+
+	const FScopedTransaction Transaction(LOCTEXT("RenameEntity", "Rename Entity"));
+	EditorListPtr.Pin()->GetDMXLibrary()->Modify();
+
+	FDMXEditorUtils::RenameEntity(EditorListPtr.Pin()->GetDMXLibrary(), Node->GetEntity(), NewNameString);
+
+	EditorListPtr.Pin()->SelectItemByName(NewNameString, ESelectInfo::OnMouseClick);
+}
+
+FText SDMXEntityRow::GetToolTipText() const
+{
+	return (LOCTEXT("EntityRowToolTip", ""));
+}
+
+FReply SDMXEntityRow::HandleOnDragDetected(const FGeometry& MyGeometry, const FPointerEvent& MouseEvent)
+{
+	if (MouseEvent.IsMouseButtonDown(EKeys::LeftMouseButton) && OnEntityDragged.IsBound())
+	{
+		if (GetNode().IsValid())
+		{
+			return OnEntityDragged.Execute(GetNode(), MouseEvent);
+		}
+	}
+
+	return FReply::Unhandled();
+}
+
+FText SDMXEntityRow::GetFilterText() const
+{
+	if (OnGetFilterText.IsBound())
+	{
+		return OnGetFilterText.Execute();
+	}
+	return FText();
+}
+
+const FSlateBrush* SDMXEntityRow::GetStatusIcon() const
+{
+	static FSlateNoResource EmptyBrush = FSlateNoResource();
+
+	if (TSharedPtr<FDMXTreeNodeBase> TreeNode = TreeNodePtr.Pin())
+	{
+		if (!TreeNode->GetErrorStatus().IsEmpty())
+		{
+			return FEditorStyle::GetBrush("Icons.Error");
+		}
+
+		if (!TreeNode->GetWarningStatus().IsEmpty())
+		{
+			return FEditorStyle::GetBrush("Icons.Warning");
+		}
+	}
+
+	return &EmptyBrush;
+}
+
+FText SDMXEntityRow::GetStatusToolTip() const
+{
+	if (TSharedPtr<FDMXTreeNodeBase> TreeNode = TreeNodePtr.Pin())
+	{
+		const FText& ErrorStatus = TreeNode->GetErrorStatus();
+		if (!ErrorStatus.IsEmpty())
+		{
+			return ErrorStatus;
+		}
+
+		const FText& WarningStatus = TreeNode->GetWarningStatus();
+		if (!WarningStatus.IsEmpty())
+		{
+			return WarningStatus;
+		}
+	}
+
+	return FText::GetEmpty();
+}
+
+///////////////////////////////////////////////////////////////////////////////
+// SDMXEntityListBase
+
+void SDMXEntityList::Construct(const FArguments& InArgs, const TSubclassOf<UDMXEntity> InListType)
+{
+	// Initialize Widget input variables
+	DMXEditor = InArgs._DMXEditor;
+	ListType = InListType;
+	OnSelectionUpdated = InArgs._OnSelectionUpdated;
+
+	// listen to common editor shortcuts for copy/paste etc
+	CommandList = MakeShared<FUICommandList>();
+	CommandList->MapAction(FGenericCommands::Get().Cut,
+		FUIAction(FExecuteAction::CreateSP(this, &SDMXEntityList::OnCutSelectedNodes),
+			FCanExecuteAction::CreateSP(this, &SDMXEntityList::CanCutNodes))
+	);
+	CommandList->MapAction(FGenericCommands::Get().Copy,
+		FUIAction(FExecuteAction::CreateSP(this, &SDMXEntityList::OnCopySelectedNodes),
+			FCanExecuteAction::CreateSP(this, &SDMXEntityList::CanCopyNodes))
+	);
+	CommandList->MapAction(FGenericCommands::Get().Paste,
+		FUIAction(FExecuteAction::CreateSP(this, &SDMXEntityList::OnPasteNodes),
+			FCanExecuteAction::CreateSP(this, &SDMXEntityList::CanPasteNodes))
+	);
+	CommandList->MapAction(FGenericCommands::Get().Duplicate,
+		FUIAction(FExecuteAction::CreateSP(this, &SDMXEntityList::OnDuplicateNodes),
+			FCanExecuteAction::CreateSP(this, &SDMXEntityList::CanDuplicateNodes))
+	);
+
+	CommandList->MapAction(FGenericCommands::Get().Delete,
+		FUIAction(FExecuteAction::CreateSP(this, &SDMXEntityList::OnDeleteNodes),
+			FCanExecuteAction::CreateSP(this, &SDMXEntityList::CanDeleteNodes))
+	);
+
+	CommandList->MapAction(FGenericCommands::Get().Rename,
+		FUIAction(FExecuteAction::CreateSP(this, &SDMXEntityList::OnRenameNode),
+			FCanExecuteAction::CreateSP(this, &SDMXEntityList::CanRenameNode))
+	);
+
+	GEditor->RegisterForUndo(this);
+
+	// Top part, with the  [+ Add New] button and the filter box
+	FText AddButtonLabel;
+	FText AddButtonToolTip;
+	if (ListType->IsChildOf(UDMXEntityController::StaticClass()))
+	{
+		AddButtonLabel = FDMXEditorCommands::Get().AddNewEntityController->GetLabel();
+		AddButtonToolTip = FDMXEditorCommands::Get().AddNewEntityController->GetDescription();
+	}
+	else if (ListType->IsChildOf(UDMXEntityFixtureType::StaticClass()))
+	{
+		AddButtonLabel = FDMXEditorCommands::Get().AddNewEntityFixtureType->GetLabel();
+		AddButtonToolTip = FDMXEditorCommands::Get().AddNewEntityFixtureType->GetDescription();
+	}
+	else if (ListType->IsChildOf(UDMXEntityFixturePatch::StaticClass()))
+	{
+		AddButtonLabel = FDMXEditorCommands::Get().AddNewEntityFixturePatch->GetLabel();
+		AddButtonToolTip = FDMXEditorCommands::Get().AddNewEntityFixturePatch->GetDescription();
+	}
+	else
+	{
+		AddButtonLabel = LOCTEXT("AddButtonDefaultLabel", "Add New");
+		AddButtonToolTip = LOCTEXT("AddButtonDefaultToolTip", "Add a new Entity");
+	}
+
+	TSharedPtr<SWidget> AddButtonContent = SNew(SHorizontalBox)
+		+ SHorizontalBox::Slot()
+		.AutoWidth()
+		.Padding(FMargin(0, 1))
+		[
+			SNew(SImage)
+			.Image(FEditorStyle::GetBrush("Plus"))
+		]
+		+ SHorizontalBox::Slot()
+		.VAlign(VAlign_Center)
+		.AutoWidth()
+		.Padding(FMargin(2, 0, 2, 0))
+		[
+			SNew(STextBlock)
+			.Text(AddButtonLabel)
+		];
+
+	TSharedPtr<SWidget> AddButton;
+
+	// The Fixture Patch tab is a special case because the Add Button is a menu
+	if (ListType->IsChildOf(UDMXEntityFixturePatch::StaticClass()))
+	{
+		AddButtonDropdownList = SNew(SDMXEntityDropdownMenu<UDMXEntityFixtureType>)
+			.DMXEditor(DMXEditor)
+			.OnEntitySelected(this, &SDMXEntityList::OnFixtureTypeSelected);
+
+		SAssignNew(AddComboButton, SComboButton)
+			.ButtonContent()
+			[
+				AddButtonContent.ToSharedRef()
+			]
+			.MenuContent()
+			[
+				AddButtonDropdownList.ToSharedRef()
+			]
+			.IsFocusable(true)
+			.ContentPadding(FMargin(5.0f, 1.0f))
+			.ComboButtonStyle(FEditorStyle::Get(), "ToolbarComboButton")
+			.ButtonStyle(FEditorStyle::Get(), "FlatButton.Success")
+			.ForegroundColor(FLinearColor::White)
+			.ToolTipText(AddButtonToolTip)
+			.OnComboBoxOpened(FOnComboBoxOpened::CreateLambda([this]() { AddButtonDropdownList->ClearSelection(); }));
+
+		AddButtonDropdownList->SetComboButton(AddComboButton);
+		AddButton = AddComboButton;
+	}
+	else
+	{
+		SAssignNew(AddButton, SButton)
+			.ButtonStyle(FEditorStyle::Get(), "FlatButton.Success")
+			.ForegroundColor(FLinearColor::White)
+			.ToolTipText(AddButtonToolTip)
+			.ContentPadding(FMargin(5.0f, 1.0f))
+			.OnClicked(this, &SDMXEntityList::OnAddNewClicked)
+			.Content()
+			[
+				AddButtonContent.ToSharedRef()
+			];
+	}
+
+	TSharedPtr<SBorder> HeaderBox = SNew(SBorder)
+		.Padding(0)
+		.BorderImage(FEditorStyle::GetBrush("DetailsView.CategoryTop"))
+		.BorderBackgroundColor(FLinearColor(0.6f, 0.6f, 0.6f, 1.0f))
+		[
+			SNew(SHorizontalBox)
+
+			// [+ Add New] button
+			+ SHorizontalBox::Slot()
+			.Padding(3.0f)
+			.AutoWidth()
+			.HAlign(HAlign_Left)
+			[
+				AddButton.ToSharedRef()
+			] // Add New button
+
+			// Filter box
+			+ SHorizontalBox::Slot()
+			.FillWidth(1.0f)
+			.VAlign(VAlign_Center)
+			.Padding(3.0f)
+			[
+				SAssignNew(FilterBox, SSearchBox)
+				.HintText(LOCTEXT("SearchEntitiesHint", "Search entities"))
+				.OnTextChanged(this, &SDMXEntityList::OnFilterTextChanged)
+			] // Filter box
+
+		]; // HeaderBox
+	
+
+	// Tree widget which displays the entities in their categories (e.g. protocol),
+	// and also controls selection and drag/drop
+	RootNode = MakeShared<FDMXTreeNodeBase>(FDMXTreeNodeBase::ENodeType::CategoryNode);
+	EntitiesTreeWidget = SNew(STreeView<TSharedPtr<FDMXTreeNodeBase>>)
+		.ItemHeight(24)
+		.TreeItemsSource(&RootNode->GetChildren())
+		.SelectionMode(ESelectionMode::Multi)
+		.OnGenerateRow(this, &SDMXEntityList::MakeNodeWidget)
+		.OnGetChildren(this, &SDMXEntityList::OnGetChildrenForTree)
+		.OnExpansionChanged(this, &SDMXEntityList::OnItemExpansionChanged)
+		.OnSelectionChanged(this, &SDMXEntityList::OnTreeSelectionChanged)
+		.OnContextMenuOpening(this, &SDMXEntityList::OnContextMenuOpen)
+		.OnItemScrolledIntoView(this, &SDMXEntityList::OnItemScrolledIntoView)
+		.HighlightParentNodesForSelection(true);
+
+	ChildSlot
+	[
+		SNew(SVerticalBox)
+		
+		+ SVerticalBox::Slot()
+		.Padding(0)
+		.AutoHeight()
+		.HAlign(HAlign_Fill)
+		[
+			HeaderBox.ToSharedRef()
+		]
+
+		+ SVerticalBox::Slot()
+		.Padding(0)
+		[
+			SNew(SBorder)
+			.Padding(2.0f)
+			.BorderImage(FEditorStyle::GetBrush("SCSEditor.TreePanel"))
+			[
+				EntitiesTreeWidget.ToSharedRef()
+			]
+		]
+	];
+
+	UpdateTree();
+
+	// Make sure we know when tabs become active to update details tab
+	OnActiveTabChangedDelegateHandle = FGlobalTabmanager::Get()->OnActiveTabChanged_Subscribe(FOnActiveTabChanged::FDelegate::CreateRaw(this, &SDMXEntityList::OnActiveTabChanged));
+}
+
+SDMXEntityList::~SDMXEntityList()
+{
+	FGlobalTabmanager::Get()->OnActiveTabChanged_Unsubscribe(OnActiveTabChangedDelegateHandle);
+	GEditor->UnregisterForUndo(this);
+}
+
+FReply SDMXEntityList::OnKeyDown(const FGeometry& MyGeometry, const FKeyEvent& InKeyEvent)
+{
+	if (CommandList->ProcessCommandBindings(InKeyEvent))
+	{
+		return FReply::Handled();
+	}
+	return FReply::Unhandled();
+}
+
+bool SDMXEntityList::CanCutNodes() const
+{
+	return CanCopyNodes() && CanDeleteNodes();
+}
+
+void SDMXEntityList::OnCutSelectedNodes()
+{
+	TArray<UDMXEntity*>&& SelectedNodes = GetSelectedEntities();
+	const FScopedTransaction Transaction(SelectedNodes.Num() > 1 ? LOCTEXT("CutEntities", "Cut Entities") : LOCTEXT("CutEntity", "Cut Entity"));
+
+	OnCopySelectedNodes();
+	OnDeleteNodes();
+}
+
+bool SDMXEntityList::CanCopyNodes() const
+{
+	TArray<UDMXEntity*>&& EntitiesToCopy = GetSelectedEntities();
+	return EntitiesToCopy.Num() > 0;
+}
+
+void SDMXEntityList::OnCopySelectedNodes()
+{
+	TArray<UDMXEntity*>&& EntitiesToCopy = GetSelectedEntities();
+
+	// Copy the entities to the clipboard
+	FDMXEditorUtils::CopyEntities(MoveTemp(EntitiesToCopy));
+}
+
+bool SDMXEntityList::CanPasteNodes() const
+{
+	return FDMXEditorUtils::CanPasteEntities();
+}
+
+void SDMXEntityList::OnPasteNodes()
+{
+	// Get the Entities to paste from the clipboard
+	TArray<UDMXEntity*> NewObjects;
+	FDMXEditorUtils::GetEntitiesFromClipboard(NewObjects);
+	check(NewObjects.Num() != 0)
+
+	// Get the library that's being edited
+	UDMXLibrary* Library = GetDMXLibrary();
+	check(Library != nullptr);
+
+	// Start transaction for Undo and take a snapshot of the current Library state
+	const FScopedTransaction PasteEntities(NewObjects.Num() > 1 ? LOCTEXT("PasteEntities", "Paste Entities") : LOCTEXT("PasteEntity", "Paste Entity"));
+	Library->Modify();
+
+	// If pasting Patches we'll need to check against existing Fixture Types
+	TArray<UDMXEntityFixtureType*> ExistingFixtureTypes;
+	if (NewObjects[0]->GetClass()->IsChildOf<UDMXEntityFixturePatch>())
+	{
+		ExistingFixtureTypes = Library->GetEntitiesTypeCast<UDMXEntityFixtureType>();
+	}
+	// Caches suitable replacements for pasted FixtureTypes (Pasted -> Existing replacement)
+	TMap<UDMXEntityFixtureType*, UDMXEntityFixtureType*> PatchTemplateReplacements;
+
+	// Add each pasted Entity to the Library
+	for (UDMXEntity* NewEntity : NewObjects)
+	{
+		check(NewEntity);
+
+		// If we're pasting Fixture Patches, we'll need to check for existing similar
+		// Fixture Type templates in this editor's Library to replace the temp one from copy
+		// or add the temp one if there's no suitable replacement
+		if (UDMXEntityFixturePatch* AsPatch = Cast<UDMXEntityFixturePatch>(NewEntity))
+		{
+			// Do we need to replace the template?
+			if (UDMXEntityFixtureType* CopiedPatchTemplate = AsPatch->ParentFixtureTypeTemplate)
+			{
+				// Did it come from this editor's DMX Library and does the original still exists?
+				if (UDMXEntityFixtureType* OriginalTemplate = Cast<UDMXEntityFixtureType>(Library->FindEntity(CopiedPatchTemplate->GetID())))
+				{
+					AsPatch->ParentFixtureTypeTemplate = OriginalTemplate;
+				}
+				else
+				{
+					check(CopiedPatchTemplate != nullptr);
+
+					// Is there already a suitable replacement registered for this template?
+					if (PatchTemplateReplacements.Contains(CopiedPatchTemplate))
+					{
+						// Replace the Patch's template with the replacement
+						AsPatch->ParentFixtureTypeTemplate = *PatchTemplateReplacements.Find(CopiedPatchTemplate);
+					}
+					else
+					{
+						// Search for a suitable replacement for the pasted Fixture Type, with identical
+						// properties, except for the Name, ID and Parent Library
+						bool bFoundReplacement = false;
+						for (UDMXEntityFixtureType* ExistingFixtureType : ExistingFixtureTypes)
+						{
+							if (FDMXEditorUtils::AreFixtureTypesIdentical(CopiedPatchTemplate, ExistingFixtureType))
+							{
+								AsPatch->ParentFixtureTypeTemplate = ExistingFixtureType;
+								PatchTemplateReplacements.Add(CopiedPatchTemplate, ExistingFixtureType);
+								bFoundReplacement = true;
+								break;
+							}
+						}
+
+						if (!bFoundReplacement)
+						{
+							// Move the Fixture Type template from the transient package into the Library package
+							NewEntity->Rename(*MakeUniqueObjectName(Library, UDMXEntityFixtureType::StaticClass()).ToString(), Library, REN_DoNotDirty | REN_DontCreateRedirectors);
+							// Make sure the Template's name and ID won't collide with existing Fixture Types
+							CopiedPatchTemplate->SetName(FDMXEditorUtils::FindUniqueEntityName(Library, UDMXEntityFixtureType::StaticClass(), CopiedPatchTemplate->GetDisplayName()));
+							CopiedPatchTemplate->RefreshID();
+							// Add to the Library
+							Library->AddEntity(CopiedPatchTemplate);
+						}
+					}
+				}
+			}
+		}
+
+		// Move the Entity from the transient package into the Library package
+		NewEntity->Rename(*MakeUniqueObjectName(Library, NewEntity->GetClass()).ToString(), Library, REN_DoNotDirty | REN_DontCreateRedirectors);
+		// Make sure the Entity's name won't collide with existing ones
+		NewEntity->SetName(FDMXEditorUtils::FindUniqueEntityName(Library, NewEntity->GetClass(), NewEntity->GetDisplayName()));
+
+		Library->AddEntity(NewEntity);
+	}
+
+	// Select the new Entities in their type tab
+	if (NewObjects[0]->GetClass() == ListType)
+	{
+		UpdateTree();
+		SelectItemsByEntity(NewObjects, ESelectInfo::OnMouseClick);
+	}
+	else
+	{
+		// Navigate to the correct tab for the pasted entities type and selected them
+		if (TSharedPtr<FDMXEditor> PinnedEditor = DMXEditor.Pin())
+		{
+			// Switching tabs will already trigger a UpdateTree, so we don't need to call it
+			PinnedEditor->SelectEntitiesInTypeTab(NewObjects, ESelectInfo::OnMouseClick);
+		}
+	}
+}
+
+bool SDMXEntityList::CanDuplicateNodes() const
+{
+	return CanCopyNodes();
+}
+
+void SDMXEntityList::OnDuplicateNodes()
+{
+	const TArray<UDMXEntity*>&& SelectedEntities = GetSelectedEntities();
+	
+	UDMXLibrary* Library = GetDMXLibrary();
+	if(SelectedEntities.Num() > 0 && Library != nullptr)
+	{
+		// Force the text box being edited (if any) to commit its text. The duplicate operation may trigger a regeneration of the tree view,
+		// releasing all row widgets. If one row was in edit mode (rename/rename on create), it was released before losing the focus and
+		// this would prevent the completion of the 'rename' or 'create + give initial name' transaction (occurring on focus lost).
+		FSlateApplication::Get().ClearKeyboardFocus();
+
+		const FScopedTransaction Transaction(SelectedEntities.Num() > 1 ? LOCTEXT("DuplicateEntities", "Duplicate Entities") : LOCTEXT("DuplicateEntity", "Duplicate Entity"));
+		Library->Modify();
+
+		// Store new entities to select them after updating the tree
+		TArray<UDMXEntity*> NewEntities;
+		NewEntities.Reserve(SelectedEntities.Num());
+
+		// We'll have the duplicates be placed right after their original counterparts
+		int32 NewEntityIndex = Library->FindEntityIndex(SelectedEntities.Last(0));
+		// Duplicate each selected entity
+		for (UDMXEntity* Entity : SelectedEntities)
+		{
+			FObjectDuplicationParameters DuplicationParams(Entity, GetDMXLibrary());
+			if (UDMXEntity* EntityCopy = CastChecked<UDMXEntity>(StaticDuplicateObjectEx(DuplicationParams)))
+			{
+				EntityCopy->SetName(FDMXEditorUtils::FindUniqueEntityName(Library, EntityCopy->GetClass(), EntityCopy->GetDisplayName()));
+				Library->AddEntity(EntityCopy);
+				NewEntities.Add(EntityCopy);
+				Library->SetEntityIndex(EntityCopy, ++NewEntityIndex);
+			}
+		}
+
+		// Refresh entities tree to contain nodes with the new entities and select them
+		UpdateTree();
+		SelectItemsByEntity(NewEntities, ESelectInfo::OnMouseClick); // OnMouseClick triggers selection updated event
+	}
+}
+
+bool SDMXEntityList::CanDeleteNodes() const
+{
+	return GetSelectedEntities().Num() > 0;
+}
+
+void SDMXEntityList::OnDeleteNodes()
+{
+	TArray<UDMXEntity*>&& EntitiesToDelete = GetSelectedEntities();
+	
+	// Check for entities being used by other objects
+	TArray<UDMXEntity*> EntitiesInUse;
+	for (UDMXEntity* Entity : EntitiesToDelete)
+	{
+		if (FDMXEditorUtils::IsEntityUsed(GetDMXLibrary(), Entity))
+		{
+			EntitiesInUse.Add(Entity);
+		}
+	}
+
+	// Confirm deletion of Entities in use, if any
+	if (EntitiesInUse.Num() > 0)
+	{
+		FText ConfirmDelete;
+		
+		// Confirmation text for a single entity in use
+		if (EntitiesInUse.Num() == 1)
+		{
+			ConfirmDelete = FText::Format(LOCTEXT("ConfirmDeleteEntityInUse",
+				"Entity \"{0}\" is in use! Do you really want to delete it?"),
+				FText::FromString(EntitiesInUse[0]->GetDisplayName()));
+		}
+		// Confirmation text for when all of the selected entities are in use
+		else if (EntitiesInUse.Num() == EntitiesToDelete.Num())
+		{
+			ConfirmDelete = LOCTEXT("ConfirmDeleteAllEntitiesInUse",
+				"All selected entities are in use! Do you really want to delete them?");
+		}
+		// Confirmation text for multiple entities, but not so much that would make the dialog huge
+		else if (EntitiesInUse.Num() > 1 && EntitiesInUse.Num() <= 10)
+		{
+			FString EntitiesNames;
+			for (UDMXEntity* Entity : EntitiesInUse)
+			{
+				EntitiesNames += TEXT("\t") + Entity->GetDisplayName() + TEXT("\n");
+			}
+
+			ConfirmDelete = FText::Format(LOCTEXT("ConfirmDeleteSomeEntitiesInUse",
+				"The Entities below are in use!\n"
+				"{0}" // no line break here because EntitiesNames will have one at the end already
+				"\nDo you really want to delete them?"),
+				FText::FromString(EntitiesNames));
+		}
+		// Confirmation text for several entities. Displaying each of their names would make a huge dialog
+		else
+		{
+			ConfirmDelete = FText::Format(LOCTEXT("ConfirmDeleteManyEntitiesInUse",
+				"{0} of the selected entities are in use!\nDo you really want to delete them?"),
+				FText::AsNumber(EntitiesInUse.Num()));
+		}
+
+		// Warn the user that this may result in data loss
+		FSuppressableWarningDialog::FSetupInfo Info(ConfirmDelete, LOCTEXT("DeleteEntities", "Delete Entities"), "DeleteEntitiesInUse_Warning");
+		Info.ConfirmText = LOCTEXT("DeleteEntities_Yes", "Yes");
+		Info.CancelText = LOCTEXT("DeleteEntities_No", "No");
+
+		FSuppressableWarningDialog DeleteEntitiesInUse(Info);
+		if (DeleteEntitiesInUse.ShowModal() == FSuppressableWarningDialog::Cancel)
+		{
+			return;
+		}
+	}
+
+	{
+		// Clears references to the Entities and delete them
+		const FScopedTransaction Transaction(EntitiesToDelete.Num() > 1 ? LOCTEXT("RemoveEntities", "Remove Entities") : LOCTEXT("RemoveEntity", "Remove Entity"));
+		FDMXEditorUtils::RemoveEntities(GetDMXLibrary(), MoveTemp(EntitiesToDelete));
+	}
+
+	UpdateTree();
+}
+
+bool SDMXEntityList::CanRenameNode() const
+{
+	return EntitiesTreeWidget->GetSelectedItems().Num() == 1 && EntitiesTreeWidget->GetSelectedItems()[0]->CanRename();
+}
+
+void SDMXEntityList::OnRenameNode()
+{
+	TArray< TSharedPtr<FDMXTreeNodeBase> > SelectedItems = EntitiesTreeWidget->GetSelectedItems();
+
+	// Should already be prevented from making it here.
+	check(SelectedItems.Num() == 1);
+
+	if (!SelectedItems[0]->BroadcastRenameRequest())
+	{
+		EntitiesTreeWidget->RequestScrollIntoView(SelectedItems[0]);
+	}
+}
+
+TArray<UDMXEntity*> SDMXEntityList::GetSelectedEntities() const
+{
+	TArray<UDMXEntity*> SelectedEntities;
+
+	if (EntitiesTreeWidget.IsValid())
+	{
+		const TArray<TSharedPtr<FDMXTreeNodeBase>>&& SelectedItems = EntitiesTreeWidget->GetSelectedItems();
+		for (const TSharedPtr<FDMXTreeNodeBase>& Item : SelectedItems)
+		{
+			if (Item.IsValid() && Item->IsEntityNode() && Item->GetEntity() != nullptr)
+			{
+				SelectedEntities.Add(Item->GetEntity());
+			}
+		}
+	}
+
+	return SelectedEntities;
+}
+
+void SDMXEntityList::SelectItemByName(const FString& ItemName, ESelectInfo::Type SelectInfo /*= ESelectInfo::Direct*/)
+{
+	// Check if the tree is being told to clear
+	if (ItemName.IsEmpty())
+	{
+		EntitiesTreeWidget->ClearSelection();
+	}
+	else
+	{
+		const TSharedPtr<FDMXTreeNodeBase>& ItemNode = FindTreeNode(FText::FromString(ItemName));
+		if (ItemNode.IsValid())
+		{
+			// If ItemNode is filtered out, we won't be able to select it
+			if (ItemNode->IsFlaggedForFiltration())
+			{
+				FilterBox->SetText(FText::GetEmpty());
+			}
+
+			// Expand the parent nodes
+			for (TSharedPtr<FDMXTreeNodeBase> ParentNode = ItemNode->GetParent().Pin(); ParentNode.IsValid(); ParentNode = ParentNode->GetParent().Pin())
+			{
+				EntitiesTreeWidget->SetItemExpansion(ParentNode, true);
+			}
+
+			EntitiesTreeWidget->SetSelection(ItemNode, SelectInfo);
+			EntitiesTreeWidget->RequestScrollIntoView(ItemNode);
+			FSlateApplication::Get().SetKeyboardFocus(EntitiesTreeWidget, EFocusCause::SetDirectly);
+		}
+	}
+}
+
+void SDMXEntityList::SelectItemByEntity(const UDMXEntity* InEntity, ESelectInfo::Type SelectInfo /*= ESelectInfo::Direct*/)
+{
+	// Check if the tree is being told to clear
+	if (InEntity == nullptr)
+	{
+		EntitiesTreeWidget->ClearSelection();
+	}
+	else
+	{
+		const TSharedPtr<FDMXTreeNodeBase>& ItemNode = FindTreeNode(InEntity);
+		if (ItemNode.IsValid())
+		{
+			// If ItemNode is filtered out, we won't be able to select it
+			if (ItemNode->IsFlaggedForFiltration())
+			{
+				FilterBox->SetText(FText::GetEmpty());
+			}
+
+			// Expand the parent nodes
+			for (TSharedPtr<FDMXTreeNodeBase> ParentNode = ItemNode->GetParent().Pin(); ParentNode.IsValid(); ParentNode = ParentNode->GetParent().Pin())
+			{
+				EntitiesTreeWidget->SetItemExpansion(ParentNode, true);
+			}
+
+			EntitiesTreeWidget->SetSelection(ItemNode, SelectInfo);
+			EntitiesTreeWidget->RequestScrollIntoView(ItemNode);
+			FSlateApplication::Get().SetKeyboardFocus(EntitiesTreeWidget, EFocusCause::SetDirectly);
+		}
+	}
+}
+
+void SDMXEntityList::SelectItemsByEntity(const TArray<UDMXEntity*>& InEntities, ESelectInfo::Type SelectInfo /*= ESelectInfo::Direct*/)
+{
+	EntitiesTreeWidget->ClearSelection();
+
+	if (InEntities.Num() > 0)
+	{
+		TSharedPtr<FDMXTreeNodeBase> FirstNode;
+		for (UDMXEntity* Entity : InEntities)
+		{
+			if (Entity == nullptr)
+			{
+				continue; 
+			}
+
+			// Find the Entity node for this Entity
+			if (TSharedPtr<FDMXTreeNodeBase> EntityNode = FindTreeNode(Entity))
+			{
+				EntitiesTreeWidget->SetItemSelection(EntityNode, true);
+
+				if (!FirstNode.IsValid())
+				{
+					FirstNode = EntityNode;
+				}
+			}
+		}
+
+		// Scroll the first selected node into view
+		if (FirstNode.IsValid())
+		{
+			EntitiesTreeWidget->RequestScrollIntoView(FirstNode);
+		}
+
+		// Notify about the new selection
+		if (SelectInfo != ESelectInfo::Direct)
+		{
+			UpdateSelectionFromNodes(EntitiesTreeWidget->GetSelectedItems());
+		}
+
+		FSlateApplication::Get().SetKeyboardFocus(EntitiesTreeWidget, EFocusCause::SetDirectly);
+	}
+}
+
+void SDMXEntityList::InitializeNodes()
+{
+	UDMXLibrary* Library = GetDMXLibrary();
+	check(Library != nullptr);
+
+	RootNode->ClearChildren();
+	EntitiesCount = 0;
+	
+	// Category type for current tab (set below)
+	FDMXTreeNodeBase::ECategoryType CategoryType;
+
+	// Sort the nodes into categories
+	if (ListType->IsChildOf(UDMXEntityController::StaticClass()))
+	{
+		CategoryType = FDMXTreeNodeBase::ECategoryType::DeviceProtocol;
+
+		Library->ForEachEntityOfType<UDMXEntityController>([this, &CategoryType](UDMXEntityController* Controller)
+			{
+				// Create this entity's node
+				TSharedPtr<FDMXEntityBaseTreeNode> ControllerNode = CreateEntityTreeNode(Controller);
+
+				// For each Entity, we find or create a category node then add the entity as its child
+				const FDMXProtocolName& Protocol = Controller->DeviceProtocol;
+				// Get the category if already existent or create it
+				TSharedPtr<FDMXTreeNodeBase> CategoryNode = GetOrCreateCategoryNode(CategoryType, FText::FromName(Protocol), &Controller->DeviceProtocol);
+
+				CategoryNode->AddChild(ControllerNode);
+			});
+
+		RootNode->SortChildren();
+	}
+	else if (ListType->IsChildOf(UDMXEntityFixtureType::StaticClass()))
+	{
+		CategoryType = FDMXTreeNodeBase::ECategoryType::DMXCategory;
+		
+		Library->ForEachEntityOfType<UDMXEntityFixtureType>([this, &CategoryType] (UDMXEntityFixtureType* FixtureType)
+			{
+				// Create this entity's node
+				TSharedPtr<FDMXEntityBaseTreeNode> FixtureTypeNode = CreateEntityTreeNode(FixtureType);
+
+				// For each Entity, we find or create a category node then add the entity as its child
+				const FDMXFixtureCategory DMXCategory = FixtureType->DMXCategory;
+				const FText DMXCategoryName = FText::FromName(DMXCategory);
+				// Get the category if already existent or create it
+				TSharedPtr<FDMXTreeNodeBase> CategoryNode = GetOrCreateCategoryNode(CategoryType, DMXCategoryName, &FixtureType->DMXCategory);
+
+				CategoryNode->AddChild(FixtureTypeNode);
+			});
+
+		RootNode->SortChildren();
+	}
+	else if (ListType->IsChildOf(UDMXEntityFixturePatch::StaticClass()))
+	{
+		static const uint32 UnassignedUniverseValue = MAX_uint32;
+		// These nodes' categories are either Assigned or Unassigned
+		TSharedPtr<FDMXCategoryTreeNode> AssignedFixturesCategoryNode = MakeShared<FDMXCategoryTreeNode>(
+			FDMXTreeNodeBase::ECategoryType::FixtureAssignmentState,
+			LOCTEXT("AssignedFixturesCategory", "Assigned Fixtures"),
+			nullptr,
+			LOCTEXT("AssignedFixturesToolTip", "Patches which Universe IDs match one of the Controllers")
+		);
+		TSharedPtr<FDMXCategoryTreeNode> UnassignedFixturesCategoryNode = MakeShared<FDMXCategoryTreeNode>(
+			FDMXTreeNodeBase::ECategoryType::FixtureAssignmentState,
+			LOCTEXT("UnassignedFixturesCategory", "Unassigned Fixtures"),
+			&UnassignedUniverseValue,
+			LOCTEXT("UnassignedFixturesToolTip", "Patches which Universe IDs match no Controllers")
+		);
+		RootNode->AddChild(AssignedFixturesCategoryNode);
+		RootNode->AddChild(UnassignedFixturesCategoryNode);
+		EntitiesTreeWidget->SetItemExpansion(AssignedFixturesCategoryNode, true);
+		EntitiesTreeWidget->SetItemExpansion(UnassignedFixturesCategoryNode, true);
+		RefreshFilteredState(AssignedFixturesCategoryNode, false);
+		RefreshFilteredState(UnassignedFixturesCategoryNode, false);
+
+		// We need to know which Universe values are valid from the controllers using them
+		const TArray<UDMXEntityController*>&& Controllers = Library->GetEntitiesTypeCast<UDMXEntityController>();
+
+		// for the Universe sub-category nodes
+		CategoryType = FDMXTreeNodeBase::ECategoryType::UniverseID;
+
+		Library->ForEachEntityOfType<UDMXEntityFixturePatch>([&] (UDMXEntityFixturePatch* FixturePatch)
+			{
+				// Create this entity's node
+				TSharedPtr<FDMXEntityBaseTreeNode> FixturePatchNode = CreateEntityTreeNode(FixturePatch);
+
+				if (FixturePatch->IsInControllersRange(Controllers))
+				{
+					// create or get existing sub-category in Assigned Fixtures category
+					TSharedPtr<FDMXTreeNodeBase> UniverseCategoryNode = GetOrCreateCategoryNode(
+						CategoryType,
+						FText::Format(LOCTEXT("UniverseSubcategoryLabel", "Universe {0}"),
+							FText::AsNumber(FixturePatch->UniverseID)),
+						&FixturePatch->UniverseID,
+						AssignedFixturesCategoryNode);
+
+					UniverseCategoryNode->AddChild(FixturePatchNode);
+				}
+				else
+				{
+					UnassignedFixturesCategoryNode->AddChild(FixturePatchNode);
+				}
+			});
+		// Sort Universe ID sub-categories in crescent order
+		AssignedFixturesCategoryNode->SortChildren();
+
+		// Sort configurations by channel value within their Universes.
+		for (TSharedPtr<FDMXTreeNodeBase> UniverseIDCategory : AssignedFixturesCategoryNode->GetChildren())
+		{
+			// Check for Patches with Auto-Assign Channel on and set their AutoStartingAddress accordingly.
+			// We won't create a Transaction for this because auto starting addresses are a consequence of other
+			// property changes, like switch Auto-Assign on/off and changing the drag/drop order of entities.
+			// So we just change it with the nodes initialization, which happens whenever any property changes,
+			// keeping it always correct and cached to be saved with the DMXLibrary.
+			const TArray< TSharedPtr<FDMXTreeNodeBase> >& PatchNodes = UniverseIDCategory->GetChildren();
+			for (int32 NodeIndex = 0; NodeIndex < PatchNodes.Num(); ++NodeIndex)
+			{
+				UDMXEntityFixturePatch* Patch = CastChecked<UDMXEntityFixturePatch>(PatchNodes[NodeIndex]->GetEntity());
+				if (Patch->bAutoAssignAddress)
+				{
+					if (NodeIndex > 0)
+					{
+						UDMXEntityFixturePatch* PreviousPatch = CastChecked<UDMXEntityFixturePatch>(PatchNodes[NodeIndex - 1]->GetEntity());
+						Patch->AutoStartingAddress = PreviousPatch->GetStartingChannel() + PreviousPatch->GetChannelSpan();
+					}
+					else
+					{
+						// This is the first Patch in this Universe, so it gets channel 1
+						Patch->AutoStartingAddress = 1;
+					}
+				}
+			}
+
+			// Sort Patches by starting channel
+			UniverseIDCategory->SortChildren([](const TSharedPtr<FDMXTreeNodeBase>& A, const TSharedPtr<FDMXTreeNodeBase>& B)->bool
+				{
+					UDMXEntityFixturePatch* PatchA = Cast<UDMXEntityFixturePatch>(A->GetEntity());
+					UDMXEntityFixturePatch* PatchB = Cast<UDMXEntityFixturePatch>(B->GetEntity());
+					if (PatchA != nullptr && PatchB != nullptr)
+					{
+						const int32& ChannelA = PatchA->GetStartingChannel();
+						const int32& ChannelB = PatchB->GetStartingChannel();
+
+						if (ChannelA == ChannelB)
+						{
+							if (PatchA->bAutoAssignAddress != PatchB->bAutoAssignAddress)
+							{
+								// Draw is decided by setting the Auto-Assigned ones as first 
+								return PatchA->bAutoAssignAddress;
+							}
+							else
+							{
+								// If both are not auto-assigned, keep drag/drop order
+								return true;
+							}
+						}
+						return PatchA->GetStartingChannel() < PatchB->GetStartingChannel();
+					}
+					return false;
+				});
+		}
+
+		// Check for Patches' overlapping channels in their universes
+		for (TSharedPtr<FDMXTreeNodeBase> UniverseIDNode : AssignedFixturesCategoryNode->GetChildren())
+		{
+			// Store the latest occupied channel in this Universe
+			int32 AvailableChannel = 1;
+			const UDMXEntityFixturePatch* PreviousPatch = nullptr;
+
+			for (TSharedPtr<FDMXTreeNodeBase> Node : UniverseIDNode->GetChildren())
+			{
+				if (UDMXEntityFixturePatch* Patch = Cast<UDMXEntityFixturePatch>(Node->GetEntity()))
+				{
+					const int32 ChannelSpan = Patch->GetChannelSpan();
+
+					if (Patch->GetStartingChannel() < AvailableChannel && PreviousPatch != nullptr)
+					{
+						// This Patch is overlapping occupied channels
+						Node->SetWarningStatus(FText::Format(
+							LOCTEXT("PatchOverlapWarning", "Start channel overlaps channels from {0}"),
+							FText::FromString(PreviousPatch->GetDisplayName())
+						));
+					}
+
+					// Update error status because after auto-channel changes there could be validation errors
+					FText InvalidReason;
+					if (!Patch->IsValidEntity(InvalidReason))
+					{
+						Node->SetErrorStatus(InvalidReason);
+					}
+
+					// Update the next available channel from this Patch's functions
+					AvailableChannel = Patch->GetStartingChannel() + ChannelSpan;
+
+					PreviousPatch = Patch;
+				}
+			}
+		}
+	}
+	else
+	{
+		UE_LOG_DMXEDITOR(Error, TEXT("%S: Current editor mode not implemented!"), __FUNCTION__);
+	}
+}
+
+TSharedPtr<FDMXEntityBaseTreeNode> SDMXEntityList::CreateEntityTreeNode(UDMXEntity* Entity)
+{
+	check(Entity != nullptr);
+	TSharedPtr<FDMXEntityBaseTreeNode> NewNode = MakeShared<FDMXEntityBaseTreeNode>(Entity);
+	RefreshFilteredState(NewNode, false);
+	
+	// Error status
+	FText InvalidReason;
+	if (!Entity->IsValidEntity(InvalidReason))
+	{
+		NewNode->SetErrorStatus(InvalidReason);
+	}
+
+	++EntitiesCount;
+	return NewNode;
+}
+
+FReply SDMXEntityList::OnAddNewClicked()
+{
+	check(DMXEditor.IsValid());
+	TSharedPtr<FDMXEditor> PinnedEditor = DMXEditor.Pin();
+
+	if (ListType->IsChildOf(UDMXEntityController::StaticClass()))
+	{
+		PinnedEditor->GetToolkitCommands()->ExecuteAction(FDMXEditorCommands::Get().AddNewEntityController.ToSharedRef());
+		return FReply::Handled();
+	}
+	else if (ListType->IsChildOf(UDMXEntityFixtureType::StaticClass()))
+	{
+		PinnedEditor->GetToolkitCommands()->ExecuteAction(FDMXEditorCommands::Get().AddNewEntityFixtureType.ToSharedRef());
+		return FReply::Handled();
+	}
+	else // UDMXEntityFixturePatch AddNew button calls OnFixtureTypeSelected
+	{
+		UE_LOG_DMXEDITOR(Error, TEXT("%S: Editor Mode not implemented!"), __FUNCTION__);
+		return FReply::Unhandled();
+	}
+}
+
+UDMXLibrary* SDMXEntityList::GetDMXLibrary() const
+{
+	if (DMXEditor.IsValid())
+	{
+		TSharedPtr<FDMXEditor> PinnedEditor(DMXEditor.Pin());
+		if (PinnedEditor.IsValid())
+		{
+			return PinnedEditor->GetDMXLibrary();
+		}
+	}
+	return nullptr;
+}
+
+TSubclassOf<UDMXEntity> SDMXEntityList::GetListType() const
+{
+	return ListType;
+}
+
+void SDMXEntityList::PostUndo(bool bSuccess)
+{
+	UpdateTree();
+}
+
+void SDMXEntityList::OnFilterTextChanged(const FText& InFilterText)
+{
+	for (TSharedPtr<FDMXTreeNodeBase> Node : RootNode->GetChildren())
+	{
+		RefreshFilteredState(Node, true);
+	}
+
+	// Clears selection to make UpdateTree automatically select the first visible node
+	EntitiesTreeWidget->ClearSelection();
+	UpdateTree(/*bRegenerateTreeNodes =*/false);
+	// If we reset the filter, recover nodes expansion states
+	UpdateNodesExpansion(RootNode.ToSharedRef(), GetFilterText().IsEmpty());
+}
+
+bool SDMXEntityList::RefreshFilteredState(TSharedPtr<FDMXTreeNodeBase> TreeNode, bool bRecursive)
+{
+	FString FilterText = FText::TrimPrecedingAndTrailing(GetFilterText()).ToString();
+	TArray<FString> FilterTerms;
+	FilterText.ParseIntoArray(FilterTerms, TEXT(" "), /*CullEmpty =*/true);
+
+	struct RefreshFilteredState_Inner
+	{
+		static void RefreshFilteredState(TSharedPtr<FDMXTreeNodeBase> TreeNodeIn, const TArray<FString>& FilterTermsIn, bool bRecursiveIn)
+		{
+			if (bRecursiveIn)
+			{
+				for (TSharedPtr<FDMXTreeNodeBase> Child : TreeNodeIn->GetChildren())
+				{
+					RefreshFilteredState(Child, FilterTermsIn, bRecursiveIn);
+				}
+			}
+
+			FString DisplayStr = TreeNodeIn->GetDisplayString();
+
+			bool bIsFilteredOut = false;
+			for (const FString& FilterTerm : FilterTermsIn)
+			{
+				if (!DisplayStr.Contains(FilterTerm))
+				{
+					bIsFilteredOut = true;
+				}
+			}
+			// if we're not recursing, then assume this is for a new node and we need to update the parent
+			// otherwise, assume the parent was hit as part of the recursion
+			TreeNodeIn->UpdateCachedFilterState(!bIsFilteredOut, /*bUpdateParent =*/!bRecursiveIn);
+		}
+	};
+
+	RefreshFilteredState_Inner::RefreshFilteredState(TreeNode, FilterTerms, bRecursive);
+	return TreeNode->IsFlaggedForFiltration();
+}
+
+TSharedRef<ITableRow> SDMXEntityList::MakeNodeWidget(TSharedPtr<FDMXTreeNodeBase> InNodePtr, const TSharedRef<STableViewBase>& OwnerTable)
+{
+	// Create the node of the appropriate type
+	if (InNodePtr->GetNodeType() == FDMXTreeNodeBase::ENodeType::CategoryNode)
+	{
+		const bool bIsRootCategory = InNodePtr->GetCategoryType() != FDMXTreeNodeBase::ECategoryType::UniverseID;
+		return SNew(SDMXCategoryRow, OwnerTable, InNodePtr, bIsRootCategory, SharedThis(this))
+			[
+				SNew(STextBlock)
+				.Text(InNodePtr->GetDisplayName())
+				.TextStyle(FEditorStyle::Get(), "DetailsView.CategoryTextStyle")
+			];
+	}
+	else if (InNodePtr->GetNodeType() == FDMXTreeNodeBase::ENodeType::EntityNode)
+	{
+		TSharedRef<SDMXEntityRow> EntityRow = SNew(SDMXEntityRow, InNodePtr, OwnerTable, SharedThis(this))
+			.OnGetFilterText(this, &SDMXEntityList::GetFilterText)
+			.OnEntityDragged(this, &SDMXEntityList::OnEntityDragged);
+
+		if (ListType->IsChildOf(UDMXEntityFixturePatch::StaticClass()))
+		{
+			// Change selected entities Auto Assign Channel property when one is clicked
+			EntityRow->GetOnAutoAssignChannelStateChanged().BindSP(this, &SDMXEntityList::OnAutoAssignChannelStateChanged, InNodePtr);
+		}
+
+		return EntityRow;
+	}
+	else
+	{
+		UE_LOG_DMXEDITOR(Fatal, TEXT("%S: node type was unexpected!"), __FUNCTION__);
+		return SNew(SDMXEntityRow, nullptr, nullptr, SharedThis(this));
+	}
+}
+
+FText SDMXEntityList::GetFilterText() const
+{
+	return FilterBox->GetText();
+}
+
+TSharedPtr< SWidget > SDMXEntityList::OnContextMenuOpen()
+{
+	const bool bCloseWindowAfterMenuSelection = true;
+	FMenuBuilder MenuBuilder(bCloseWindowAfterMenuSelection, CommandList);
+
+	if (GetSelectedEntities().Num() > 0)
+	{
+		MenuBuilder.BeginSection("BasicOperations");
+		{
+			MenuBuilder.AddMenuEntry(FGenericCommands::Get().Rename);
+			MenuBuilder.AddMenuEntry(FGenericCommands::Get().Cut);
+			MenuBuilder.AddMenuEntry(FGenericCommands::Get().Copy);
+			MenuBuilder.AddMenuEntry(FGenericCommands::Get().Paste);
+			MenuBuilder.AddMenuEntry(FGenericCommands::Get().Duplicate);
+			MenuBuilder.AddMenuEntry(FGenericCommands::Get().Delete);
+		}
+		MenuBuilder.EndSection();
+	}
+	else if (ListType->IsChildOf(UDMXEntityFixturePatch::StaticClass()))
+	{
+		return SNew(SDMXEntityDropdownMenu<UDMXEntityFixtureType>)
+			.DMXEditor(DMXEditor)
+			.OnEntitySelected(this, &SDMXEntityList::OnFixtureTypeSelected);
+		// TODO add (somehow) Paste option to this menu
+	}
+	else
+	{
+		BuildAddNewMenu(MenuBuilder);
+		MenuBuilder.BeginSection("BasicOperations");
+		{
+			MenuBuilder.AddMenuEntry(FGenericCommands::Get().Paste);
+		}
+		MenuBuilder.EndSection();
+	}
+
+	return MenuBuilder.MakeWidget();
+}
+
+void SDMXEntityList::BuildAddNewMenu(FMenuBuilder& MenuBuilder)
+{
+	MenuBuilder.BeginSection("AddNewItem", LOCTEXT("AddOperations", "Add New"));
+
+	if (ListType->IsChildOf(UDMXEntityController::StaticClass()))
+	{
+		MenuBuilder.AddMenuEntry(FDMXEditorCommands::Get().AddNewEntityController);
+	}
+	else if (ListType->IsChildOf(UDMXEntityFixtureType::StaticClass()))
+	{
+		MenuBuilder.AddMenuEntry(FDMXEditorCommands::Get().AddNewEntityFixtureType);
+	}
+	else
+	{
+		UE_LOG_DMXEDITOR(Error, TEXT("%S: ListType of not implemented class!"), __FUNCTION__);
+	}
+
+	MenuBuilder.EndSection();
+}
+
+void SDMXEntityList::OnItemScrolledIntoView(TSharedPtr<FDMXTreeNodeBase> InItem, const TSharedPtr<ITableRow>& InWidget)
+{
+	if (InItem->IsRenameRequestPending())
+	{
+		InItem->BroadcastRenameRequest();
+	}
+}
+
+void SDMXEntityList::GetCollapsedNodes(TSet<TSharedPtr<FDMXTreeNodeBase>>& OutCollapsedNodes, TSharedPtr<FDMXTreeNodeBase> InParentNodePtr /*= nullptr*/) const
+{
+	if (!InParentNodePtr.IsValid())
+	{
+		InParentNodePtr = RootNode;
+	}
+
+	for (const TSharedPtr<FDMXTreeNodeBase>& Node : InParentNodePtr->GetChildren())
+	{
+		if (Node->GetChildren().Num() > 0)
+		{
+			if (!EntitiesTreeWidget->IsItemExpanded(Node))
+			{
+				OutCollapsedNodes.Add(Node);
+			}
+			else // not collapsed. Check children
+			{
+				GetCollapsedNodes(OutCollapsedNodes, Node);
+			}
+		}
+	}
+}
+
+TSharedPtr<FDMXTreeNodeBase> SDMXEntityList::FindTreeNode(const UDMXEntity* InEntity, TSharedPtr<FDMXTreeNodeBase> InStartNodePtr /*= nullptr*/) const
+{
+	TSharedPtr<FDMXTreeNodeBase> NodePtr;
+	if (InEntity != nullptr)
+	{
+		// Start at root node if none was given
+		if (!InStartNodePtr.IsValid())
+		{
+			InStartNodePtr = RootNode;
+		}
+
+		if (InStartNodePtr.IsValid())
+		{
+			// Check to see if the given entity matches the given tree node
+			if (InStartNodePtr->GetEntity() == InEntity)
+			{
+				NodePtr = InStartNodePtr;
+			}
+			else
+			{
+				for (int32 i = 0; i < InStartNodePtr->GetChildren().Num() && !NodePtr.IsValid(); ++i)
+				{
+					NodePtr = FindTreeNode(InEntity, InStartNodePtr->GetChildren()[i]);
+				}
+			}
+		}
+	}
+
+	return NodePtr;
+}
+
+TSharedPtr<FDMXTreeNodeBase> SDMXEntityList::FindTreeNode(const FText& InName, TSharedPtr<FDMXTreeNodeBase> InStartNodePtr /*= nullptr*/) const
+{
+	TSharedPtr<FDMXTreeNodeBase> NodePtr;
+	if (!InName.IsEmpty())
+	{
+		// Start at root node if none was given
+		if (!InStartNodePtr.IsValid())
+		{
+			InStartNodePtr = RootNode;
+		}
+
+		if (InStartNodePtr.IsValid())
+		{
+			// Check to see if the given entity matches the given tree node
+			if (InStartNodePtr->GetDisplayName().CompareTo(InName) == 0)
+			{
+				NodePtr = InStartNodePtr;
+			}
+			else
+			{
+				for (int32 i = 0; i < InStartNodePtr->GetChildren().Num() && !NodePtr.IsValid(); ++i)
+				{
+					NodePtr = FindTreeNode(InName, InStartNodePtr->GetChildren()[i]);
+				}
+			}
+		}
+	}
+
+	return NodePtr;
+}
+
+TSharedPtr<FDMXTreeNodeBase> SDMXEntityList::GetOrCreateCategoryNode(const FDMXTreeNodeBase::ECategoryType InCategoryType, const FText InCategoryName, void* CategoryValuePtr, TSharedPtr<FDMXTreeNodeBase> InParentNodePtr /*= nullptr*/, const FText& InToolTip /*= FText::GetEmpty()*/)
+{
+	for (const TSharedPtr<FDMXTreeNodeBase>& Node : InParentNodePtr.IsValid() ? InParentNodePtr->GetChildren() : RootNode->GetChildren())
+	{
+		if (Node.IsValid() && Node->GetNodeType() == FDMXTreeNodeBase::ENodeType::CategoryNode)
+		{
+			if (Node->GetCategoryType() == InCategoryType && Node->GetDisplayName().CompareTo(InCategoryName) == 0)
+			{
+				EntitiesTreeWidget->SetItemExpansion(Node, true);
+				return Node;
+			}
+		}
+	}
+
+	// Didn't find an existing node. Add one.
+	TSharedPtr<FDMXTreeNodeBase> NewNode = MakeShared<FDMXCategoryTreeNode>(InCategoryType, InCategoryName, CategoryValuePtr, InToolTip);
+	if (InParentNodePtr.IsValid())
+	{
+		InParentNodePtr->AddChild(NewNode);
+	}
+	else
+	{
+		RootNode->AddChild(NewNode);
+	}
+
+	RefreshFilteredState(NewNode, false);
+	EntitiesTreeWidget->SetItemExpansion(NewNode, true);
+	NewNode->SetExpansionState(true);
+
+	return NewNode;
+}
+
+void SDMXEntityList::OnActiveTabChanged(TSharedPtr<SDockTab> PreviouslyActive, TSharedPtr<SDockTab> NewlyActivated)
+{
+	if (IsInTab(NewlyActivated))
+	{
+		UpdateTree();
+
+		if (ListType->IsChildOf(UDMXEntityFixturePatch::StaticClass()))
+		{
+			// New Fixture Types might've been added on their tab
+			AddButtonDropdownList->RefreshEntitiesList();
+		}
+
+		// Refresh selected entities' properties on the inspector panel by issuing a selection update.
+		// Some properties might have been changed on a previously selected tab.
+		UpdateSelectionFromNodes(EntitiesTreeWidget->GetSelectedItems());
+	}
+}
+
+bool SDMXEntityList::IsInTab(TSharedPtr<SDockTab> InDockTab) const
+{
+	// Too many hierarchy levels to do it with a recursive function. Crashes with Stack Overflow.
+	// Using loop instead.
+	
+	if (InDockTab.IsValid())
+	{
+		// Tab content that should be a parent of this widget on some level
+		TSharedPtr<SWidget> TabContent = InDockTab->GetContent();
+		// Current parent being checked against
+		TSharedPtr<SWidget> CurrentParent = GetParentWidget();
+
+		while (CurrentParent.IsValid())
+		{
+			if (CurrentParent == TabContent)
+			{
+				return true;
+			}
+			CurrentParent = CurrentParent->GetParentWidget();
+		}
+
+		// reached top widget (parent is invalid) and none was the tab
+		return false;
+	}
+
+	return false;
+}
+
+void SDMXEntityList::OnGetChildrenForTree(TSharedPtr<FDMXTreeNodeBase> InNodePtr, TArray<TSharedPtr<FDMXTreeNodeBase>>& OutChildren)
+{
+	if (InNodePtr.IsValid())
+	{
+		const TArray<TSharedPtr<FDMXTreeNodeBase>>& Children = InNodePtr->GetChildren();
+
+		if (!GetFilterText().IsEmpty())
+		{
+			OutChildren.Reserve(Children.Num());
+
+			for (const TSharedPtr<FDMXTreeNodeBase>& Child : Children)
+			{
+				if (!Child->IsFlaggedForFiltration())
+				{
+					OutChildren.Add(Child);
+				}
+			}
+		}
+		else
+		{
+			OutChildren = Children;
+		}
+	}
+	else
+	{
+		OutChildren.Empty();
+	}
+}
+
+void SDMXEntityList::OnTreeSelectionChanged(TSharedPtr<FDMXTreeNodeBase> InSelectedNodePtr, ESelectInfo::Type SelectInfo)
+{
+	if (SelectInfo != ESelectInfo::Direct)
+	{
+		UpdateSelectionFromNodes(EntitiesTreeWidget->GetSelectedItems());
+	}
+}
+
+void SDMXEntityList::UpdateSelectionFromNodes(const TArray<TSharedPtr<FDMXTreeNodeBase>>& SelectedNodes)
+{
+	bUpdatingSelection = true;
+
+	// Notify that the selection has updated
+	OnSelectionUpdated.ExecuteIfBound(SelectedNodes);
+
+	bUpdatingSelection = false;
+}
+
+void SDMXEntityList::SetNodeExpansionState(TSharedPtr<FDMXTreeNodeBase> InNodeToChange, const bool bIsExpanded)
+{
+	if (EntitiesTreeWidget.IsValid() && InNodeToChange.IsValid())
+	{
+		EntitiesTreeWidget->SetItemExpansion(InNodeToChange, bIsExpanded);
+	}
+}
+
+void SDMXEntityList::UpdateTree(bool bRegenerateTreeNodes /*= true*/)
+{
+	check(EntitiesTreeWidget.IsValid());
+
+	if (bRegenerateTreeNodes)
+	{
+		// Obtain the set of expandable tree nodes that are currently collapsed
+		TSet<TSharedPtr<FDMXTreeNodeBase>> CollapsedTreeNodes;
+		GetCollapsedNodes(CollapsedTreeNodes);
+
+		// Obtain the list of selected items
+		TArray<TSharedPtr<FDMXTreeNodeBase>> SelectedTreeNodes = EntitiesTreeWidget->GetSelectedItems();
+
+		// Clear the current tree
+		if (SelectedTreeNodes.Num() != 0)
+		{
+			EntitiesTreeWidget->ClearSelection();
+		}
+
+		InitializeNodes();
+
+		// Restore the previous expansion state on the new tree nodes
+		TArray<TSharedPtr<FDMXTreeNodeBase>> CollapsedTreeNodeArray = CollapsedTreeNodes.Array();
+		for (int32 i = 0; i < CollapsedTreeNodeArray.Num(); ++i)
+		{
+			// Look for a category match in the new hierarchy; if found, mark it as collapsed to match the previous setting
+			TSharedPtr<FDMXTreeNodeBase> NodeToExpandPtr = FindTreeNode(CollapsedTreeNodeArray[i]->GetDisplayName());
+			if (NodeToExpandPtr.IsValid())
+			{
+				EntitiesTreeWidget->SetItemExpansion(NodeToExpandPtr, false);
+			}
+			else
+			{
+				EntitiesTreeWidget->SetItemExpansion(NodeToExpandPtr, true);
+			}
+		}
+
+		if (SelectedTreeNodes.Num() > 0)
+		{
+			// Restore the previous selection state on the new tree nodes
+			for (int32 i = 0; i < SelectedTreeNodes.Num(); ++i)
+			{
+				TSharedPtr<FDMXTreeNodeBase> NodeToSelectPtr = FindTreeNode(SelectedTreeNodes[i]->GetEntity());
+				if (NodeToSelectPtr.IsValid())
+				{
+					EntitiesTreeWidget->SetItemSelection(NodeToSelectPtr, true, ESelectInfo::Direct);
+				}
+			}
+		}
+	}
+
+	// refresh widget
+	EntitiesTreeWidget->RequestTreeRefresh();
+
+	// If no entity is selected, select first available one, if any
+	if (EntitiesTreeWidget->GetNumItemsSelected() == 0)
+	{
+		UDMXLibrary* Library = GetDMXLibrary();
+		check(Library != nullptr);
+
+		bool bSelectedAnEntity = false;
+		// Find the first non filtered out entity
+		Library->ForEachEntityOfTypeWithBreak(ListType, [this](UDMXEntity* Entity)
+			{
+				if (TSharedPtr<FDMXTreeNodeBase> Node = FindTreeNode(Entity))
+				{
+					if (!Node->IsFlaggedForFiltration())
+					{
+						EntitiesTreeWidget->SetSelection(Node, ESelectInfo::OnMouseClick);
+						return false;
+					}
+				}
+				return true;
+			});
+
+		if (!bSelectedAnEntity)
+		{
+			// There are no entities. Update the property inspector to empty it
+			UpdateSelectionFromNodes({});
+		}
+	}
+}
+
+void SDMXEntityList::UpdateNodesExpansion(TSharedRef<FDMXTreeNodeBase> InRootNode, bool bFilterIsEmpty)
+{
+	// Only category nodes have children and need expansion
+	if (!InRootNode->IsEntityNode())
+	{
+		// If the filter is not empty, all nodes should be expanded
+		EntitiesTreeWidget->SetItemExpansion(InRootNode, !bFilterIsEmpty || InRootNode->GetExpansionState());
+
+		for (const TSharedPtr<FDMXTreeNodeBase>& Child : InRootNode->GetChildren())
+		{
+			if (Child.IsValid() && !Child->IsEntityNode())
+			{
+				UpdateNodesExpansion(Child.ToSharedRef(), bFilterIsEmpty);
+			}
+		}
+	}
+}
+
+void SDMXEntityList::OnItemExpansionChanged(TSharedPtr<FDMXTreeNodeBase> InNodePtr, bool bInExpansionState)
+{
+	 // Only applies when there's no filtering
+	if (InNodePtr.IsValid() && GetFilterText().IsEmpty())
+	{
+		InNodePtr->SetExpansionState(bInExpansionState);
+	}
+}
+
+FReply SDMXEntityList::OnEntityDragged(TSharedPtr<FDMXTreeNodeBase> InNodePtr, const FPointerEvent& MouseEvent)
+{
+	if (InNodePtr.IsValid() && InNodePtr->GetEntity() != nullptr)
+	{
+		TArray<TSharedPtr<FDMXTreeNodeBase>>&& SelectedItems = EntitiesTreeWidget->GetSelectedItems();
+		TArray<TSharedPtr<FDMXEntityBaseTreeNode>> DraggedEntities;
+		DraggedEntities.Reserve(SelectedItems.Num());
+
+		for (TSharedPtr<FDMXTreeNodeBase> SelectedItem : SelectedItems)
+		{
+			if (SelectedItem->IsEntityNode())
+			{
+				if (TSharedPtr<FDMXEntityBaseTreeNode> AsEntityNode = StaticCastSharedPtr<FDMXEntityBaseTreeNode>(SelectedItem))
+				{
+					DraggedEntities.Add(AsEntityNode);
+				}
+			}
+		}
+
+		if (DraggedEntities.Num() == 0)
+		{
+			if (TSharedPtr<FDMXEntityBaseTreeNode> AsEntityNode = StaticCastSharedPtr<FDMXEntityBaseTreeNode>(InNodePtr))
+			{
+				DraggedEntities.Add(AsEntityNode);
+			}
+			else
+			{
+				return FReply::Unhandled();
+			}
+		}
+
+		TSharedRef<FDMXEntityDragDropOperation> DragOperation = MakeShared<FDMXEntityDragDropOperation>(GetDMXLibrary(), SharedThis(this), MoveTemp(DraggedEntities));
+
+		return FReply::Handled().BeginDragDrop(DragOperation);
+	}
+
+	return FReply::Unhandled();
+}
+
+void SDMXEntityList::OnFixtureTypeSelected(UDMXEntity* InSelectedFixtureType)
+{
+	check(DMXEditor.IsValid() && ListType->IsChildOf(UDMXEntityFixturePatch::StaticClass()));
+	TSharedPtr<FDMXEditor> PinnedEditor = DMXEditor.Pin();
+
+	// Editor will call these during the creation of the new Entity
+	UDMXEntityFixtureType* AsFixtureType = CastChecked<UDMXEntityFixtureType>(InSelectedFixtureType);
+	OnGetBaseNameForNewEntityHandle = PinnedEditor->GetOnGetBaseNameForNewEntity().AddSP(this, &SDMXEntityList::OnEditorGetBaseNameForNewFixturePatch, AsFixtureType);
+	OnSetupNewEntityHandle = PinnedEditor->GetOnSetupNewEntity().AddSP(this, &SDMXEntityList::OnEditorSetupNewFixturePatch, AsFixtureType);
+
+	PinnedEditor->GetToolkitCommands()->ExecuteAction(FDMXEditorCommands::Get().AddNewEntityFixturePatch.ToSharedRef());
+}
+
+void SDMXEntityList::OnEditorGetBaseNameForNewFixturePatch(TSubclassOf<UDMXEntity> InEntityClass, FString& OutBaseName, UDMXEntityFixtureType* InSelectedFixtureType)
+{
+	if (!InEntityClass->IsChildOf(UDMXEntityFixturePatch::StaticClass()))
+	{
+		return;
+	}
+
+	TSharedPtr<FDMXEditor> PinnedEditor = DMXEditor.Pin();
+	if (PinnedEditor.IsValid())
+	{
+		PinnedEditor->GetOnGetBaseNameForNewEntity().Remove(OnGetBaseNameForNewEntityHandle);
+
+		OutBaseName = InSelectedFixtureType->GetDisplayName() + TEXT("_Patch");
+	}
+}
+
+void SDMXEntityList::OnEditorSetupNewFixturePatch(UDMXEntity* InNewEntity, UDMXEntityFixtureType* InSelectedFixtureType)
+{
+	if (UDMXEntityFixturePatch* FixturePatch = Cast<UDMXEntityFixturePatch>(InNewEntity))
+	{
+		TSharedPtr<FDMXEditor> PinnedEditor = DMXEditor.Pin();
+		if (PinnedEditor.IsValid())
+		{
+			PinnedEditor->GetOnSetupNewEntity().Remove(OnSetupNewEntityHandle);
+
+			FixturePatch->ParentFixtureTypeTemplate = InSelectedFixtureType;
+			// Issue a selection to trigger a OnSelectionUpdate and make the inspector display the new values
+			SelectItemByEntity(FixturePatch);
+		}
+	}
+	else
+	{
+		UE_LOG_DMXEDITOR(Error, TEXT("%S: New Entity wasn't a FixturePatch!"), __FUNCTION__);
+	}
+}
+
+void SDMXEntityList::OnAutoAssignChannelStateChanged(bool NewState, TSharedPtr<FDMXTreeNodeBase> InNodePtr)
+{
+	const FScopedTransaction Transaction(LOCTEXT("SetAutoAssignChannelTransaction", "Set Auto Assign Channel"));
+
+	// Was the changed entity one of the selected ones?
+	if (EntitiesTreeWidget->IsItemSelected(InNodePtr))
+	{
+		const TArray<UDMXEntity*>&& SelectedEntities = GetSelectedEntities();
+		for (UDMXEntity* SelectedEntity : SelectedEntities)
+		{
+			if (UDMXEntityFixturePatch* Patch = Cast<UDMXEntityFixturePatch>(SelectedEntity))
+			{
+				if (Patch->bAutoAssignAddress != NewState)
+				{
+					Patch->Modify();
+					Patch->bAutoAssignAddress = NewState;
+				}
+			}
+		}
+	}
+	else
+	{
+		if (UDMXEntityFixturePatch* Patch = Cast<UDMXEntityFixturePatch>(InNodePtr->GetEntity()))
+		{
+			Patch->Modify();
+			Patch->bAutoAssignAddress = NewState;
+		}
+	}
+
+	UpdateTree();
+}
+
+#undef LOCTEXT_NAMESPACE
