- DMX Editor

diff --git a/Engine/Plugins/Experimental/DigitalMultiplex/DMXEngine/Source/DMXEditor/Private/Widgets/SDMXEntityList.h b/Engine/Plugins/Experimental/DigitalMultiplex/DMXEngine/Source/DMXEditor/Private/Widgets/SDMXEntityList.h
new file mode 100644
--- /dev/null
+++ b/Engine/Plugins/Experimental/DigitalMultiplex/DMXEngine/Source/DMXEditor/Private/Widgets/SDMXEntityList.h
@@ -0,0 +1,626 @@
+// Copyright Epic Games, Inc. All Rights Reserved.
+
+#pragma once
+
+#include "CoreMinimal.h"
+
+#include "Library/DMXEntity.h"
+#include "Widgets/SDMXEntityDropdownMenu.h"
+
+#include "Widgets/SCompoundWidget.h"
+#include "Widgets/Views/STableRow.h"
+#include "Widgets/Views/STreeView.h"
+
+#include "EditorUndoClient.h"
+
+class FDMXEditor;
+class SDMXEntityList;
+class UDMXLibrary;
+class SDMXEntityEditor;
+class UDMXEntity;
+class UDMXEntityFixtureType;
+
+class FScopedTransaction;
+class FUICommandList;
+class FMenuBuilder;
+
+class SSearchBox;
+class SInlineEditableTextBlock;
+class SDockTab;
+class SComboButton;
+
+///////////////////////////////////////////////////////////////////////////////
+// FDMXEntityListTreeNode
+
+class FDMXTreeNodeBase
+	: public TSharedFromThis<FDMXTreeNodeBase>
+{
+public:
+	/** Delegate for when the context menu requests a rename */
+	DECLARE_DELEGATE(FOnRenameRequested);
+
+	enum class ENodeType : uint8
+	{
+		CategoryNode,
+		EntityNode
+	};
+
+	enum class ECategoryType : uint8
+	{
+		DeviceProtocol,
+		DMXCategory,
+		FixtureAssignmentState,
+		UniverseID,
+		NONE
+	};
+
+	/**  Constructs an empty tree node. */
+	FDMXTreeNodeBase(FDMXTreeNodeBase::ENodeType InNodeType);
+	virtual ~FDMXTreeNodeBase() {}
+
+	/**  @return The string to be used in the tree display. */
+	virtual FString GetDisplayString() const;
+	/**  @return The name of this node in text. */
+	virtual FText GetDisplayName() const;
+	/**  @return The SCS node that is represented by this object, or NULL if there is no SCS node associated with the component template. */
+	virtual UDMXEntity* GetEntity() const;
+	/**  @return The type of this node. */
+	virtual ENodeType GetNodeType() const;
+
+	/** Add a child node to this node */
+	virtual void AddChild(TSharedPtr<FDMXTreeNodeBase> InChildPtr);
+	/** Remove a child node from this node */
+	virtual void RemoveChild(TSharedPtr<FDMXTreeNodeBase> InChildPtr);
+	/** Remove this node from its parent one */
+	virtual void RemoveFromParent();
+	/**  @return Child nodes for this object. */
+	virtual const TArray<TSharedPtr<FDMXTreeNodeBase>>& GetChildren() const;
+	/**  Remove all child nodes from this node. */
+	virtual void ClearChildren();
+	/**  Sort children by name */
+	virtual void SortChildren();
+	/**  Sort children using custom predicate */
+	virtual void SortChildren(TFunction<bool (const TSharedPtr<FDMXTreeNodeBase>&, const TSharedPtr<FDMXTreeNodeBase>&)> Predicate);
+	/**  @return This object's parent node (or an invalid reference if no parent is assigned). */
+	virtual TWeakPtr<FDMXTreeNodeBase> GetParent() const { return ParentNodePtr; }
+
+	/**  @return Whether or not this node can be deleted from the entities tree */
+	virtual bool CanDelete() const { return false; }
+
+	/**  @return Whether or not this object represents a node that can be renamed from the entities tree. */
+	virtual bool CanRename() const { return false; }
+
+	/** To be overridden by a subclass that represents an Entity node */
+	virtual bool IsEntityNode() const { return false; }
+
+	/** If this node is a category node, represents the type of category. Otherwise, ECategoryType::NONE */
+	virtual ECategoryType GetCategoryType() const { return ECategoryType::NONE; }
+
+	/**
+	 * Accessor to the node's RenameRequestEvent (for binding purposes). Do not
+	 * Execute() the delegate from this function, instead call
+	 * BroadcastRenameRequest() on the node.
+	 *
+	 * @return The node's internal RenameRequestEvent.
+	 */
+	FOnRenameRequested& OnRenameRequest() { return RenameRequestEvent; }
+
+	/**
+	 * Executes the node's RenameRequestEvent if it is bound. Otherwise, it will
+	 * mark the node as having a pending rename request.
+	 *
+	 * @return True if the broadcast went through, false if the "pending rename request" flag was set.
+	 */
+	bool BroadcastRenameRequest();
+
+	/**
+	 * Sometimes a call to BroadcastRenameRequest() is made before the
+	 * RenameRequestEvent has been bound. When that happens, this node is
+	 * marked with a pending rename request. This method determines if that is
+	 * the case for this node.
+	 *
+	 * @return True if a call to BroadcastRenameRequest() was made without a valid RenameRequestEvent.
+	 */
+	bool IsRenameRequestPending() const;
+
+	/**
+	 * Attempts to find a reference to the child node that matches the given component template.
+	 *
+	 * @param InComponentTemplate The component template instance to match.
+	 * @param bRecursiveSearch Whether or not to recursively search child nodes (default == false).
+	 * @param OutDepth If non-NULL, the depth of the child node will be returned in this parameter on success (default == NULL).
+	 * @return The child node with a component template that matches the given component template instance, or an invalid node reference if no match was found.
+	 */
+	TSharedPtr<FDMXTreeNodeBase> FindChild(const UDMXEntity* InEntity, bool bRecursiveSearch = false, uint32* OutDepth = NULL) const;
+
+	/** Query that determines if this item should be filtered out or not */
+	virtual bool IsFlaggedForFiltration() const
+	{
+		// if Unknown, throws an error message
+		return ensureMsgf(FilterFlags != (uint8)EFilteredState::Unknown, TEXT("Querying a bad filtration state.")) ?
+			(FilterFlags & (uint8)EFilteredState::FilteredInMask) == 0 : false;
+	}
+
+	/** Refreshes this item's filtration state. Use bUpdateParent to make sure the parent's EFilteredState::ChildMatches flag is properly updated based off the new state */
+	void UpdateCachedFilterState(bool bMatchesFilter, bool bUpdateParent);
+
+	/** Update this node's desired expansion state for when there're no filters */
+	void SetExpansionState(bool bNewExpansionState);
+	/** This node's desired expansion state for when there're no filters */
+	bool GetExpansionState() const { return bShouldBeExpanded; }
+
+	/** If the warning tool tip is not empty, the node will display a warning icon with said tool tip. */
+	void SetWarningStatus(const FText& InWarningToolTip);
+	const FText& GetWarningStatus() const { return WarningToolTip; }
+	/** If the error tool tip is not empty, the node will display an error icon with said tool tip. */
+	void SetErrorStatus(const FText& InErrorToolTip);
+	const FText& GetErrorStatus() const { return ErrorToolTip; }
+
+	/** Operator used when sorting categories by name/number */
+	bool operator<(const FDMXTreeNodeBase& Other) const;
+
+protected:
+	/** Updates the EFilteredState::ChildMatches flag, based off of children's current state */
+	void RefreshCachedChildFilterState(bool bUpdateParent);
+	/** Used to update the EFilteredState::ChildMatches flag for parent nodes, when this item's filtration state has changed */
+	void ApplyFilteredStateToParent();
+
+	/** DMX Entity represented by this node, if it's an entity node, otherwise invalid */
+	TWeakObjectPtr<UDMXEntity> DMXEntity;
+
+	FText WarningToolTip;
+	FText ErrorToolTip;
+
+private:
+	ENodeType NodeType;
+
+	// Actual tree structure
+	TWeakPtr<FDMXTreeNodeBase> ParentNodePtr;
+	TArray<TSharedPtr<FDMXTreeNodeBase>> Children;
+
+	/** When the item is first created, a rename request may occur before everything is setup for it. This toggles to true in those cases */
+	bool bPendingRenameRequest;
+	/** Delegate to trigger when a rename was requested on this node */
+	FOnRenameRequested RenameRequestEvent;
+
+	/** Register whether the node should be expanded when there's no search filter text */
+	bool bShouldBeExpanded;
+
+	enum EFilteredState
+	{
+		FilteredOut = 0x00,
+		MatchesFilter = (1 << 0),
+		ChildMatches = (1 << 1),
+
+		FilteredInMask = (MatchesFilter | ChildMatches),
+		Unknown = 0xFC // ~FilteredInMask
+	};
+	uint8 FilterFlags;
+};
+
+class FDMXCategoryTreeNode
+	: public FDMXTreeNodeBase
+{
+public:
+	FDMXCategoryTreeNode(ECategoryType InCategoryType, FText InCategoryName, const void* InCategoryValuePtr, const FText& ToolTip = FText::GetEmpty());
+	
+	// ~ start FDMXTreeNodeBase interface
+	virtual FString GetDisplayString() const override;
+	virtual FText GetDisplayName() const override;
+	virtual ECategoryType GetCategoryType() const override;
+	// ~ end FDMXTreeNodeBase interface
+
+	const FText& GetToolTip() const { return ToolTip; }
+
+	bool IsCategoryValueValid() const { return CategoryValuePtr != nullptr; }
+	template<typename ValType>
+	const ValType* GetCategoryValue() const
+	{
+		return CategoryValuePtr == nullptr ? nullptr : (ValType*)(CategoryValuePtr);
+	}
+
+protected:
+	/** This node's category type */
+	ECategoryType CategoryType;
+	const void* CategoryValuePtr;
+
+	FText CategoryName;
+	FText ToolTip;
+};
+
+class FDMXEntityBaseTreeNode : public FDMXTreeNodeBase
+{
+public:
+	FDMXEntityBaseTreeNode(UDMXEntity* InEntity);
+
+	// ~ start FDMXTreeNodeBase interface
+	virtual bool IsEntityNode() const override;
+	virtual bool CanDelete() const override { return true; }
+	virtual bool CanRename() const override { return true; }
+	// ~ end FDMXTreeNodeBase interface
+};
+
+///////////////////////////////////////////////////////////////////////////////
+// FDMXEntityDragDropOperation
+
+class FDMXEntityDragDropOperation
+	: public FDragDropOperation
+{
+public:
+	DRAG_DROP_OPERATOR_TYPE(FDMXEntityDragDropOperation, FDragDropOperation)
+
+	FDMXEntityDragDropOperation(UDMXLibrary* InLibrary, TWeakPtr<SDMXEntityList> EntityList, TArray<TSharedPtr<FDMXEntityBaseTreeNode>>&& InEntities);
+
+	void SetHoveredEntity(TSharedPtr<FDMXEntityBaseTreeNode> InEntityNode, UDMXLibrary* InLibrary, TSubclassOf<UDMXEntity> TabEntitiesType);
+	void SetHoveredCategory(TSharedPtr<FDMXCategoryTreeNode> InCategory, UDMXLibrary* InLibrary, TSubclassOf<UDMXEntity> TabEntitiesType);
+
+	void DroppedOnEntity(TSharedRef<FDMXEntityBaseTreeNode> InEntity, UDMXLibrary* InLibrary, TSubclassOf<UDMXEntity> TabEntitiesType);
+	void DroppedOnCategory(TSharedRef<FDMXCategoryTreeNode> InCategory, UDMXLibrary* InLibrary, TSubclassOf<UDMXEntity> TabEntitiesType);
+
+protected:
+	/** Constructs the tooltip widget that follows the mouse */
+	virtual void Construct() override;
+
+private:
+	void HoverTargetChanged();
+
+	// Common after-dropping actions
+	/** Move all dragged entities to NewIndex */
+	void ReorderEntities(int32 NewIndex);
+	/** Set the required property to move the dragged entities into a specific category from the list */
+	void SetPropertyForNewCategory();
+
+	void SetDraggingFromMultipleCategories();
+	bool GetCategoryPropertyNameFromTabType(FText& PropertyName, FText& CategoryPropertyValue);
+	bool IsDraggingToDifferentCategory() const;
+	TSharedPtr<FDMXCategoryTreeNode> GetTargetCategory() const;
+
+	void SetFeedbackMessageError(const FText& Message);
+	void SetFeedbackMessageOK(const FText& Message);
+	void SetFeedbackMessage(const FSlateBrush* Icon, const FText& Message);
+
+private:
+	UDMXLibrary* DraggedFromLibrary;
+	TArray<TSharedPtr<FDMXEntityBaseTreeNode>> DraggedEntities;
+	TWeakPtr<SDMXEntityList> EntityList;
+
+	TSharedPtr<FDMXEntityBaseTreeNode> HoveredEntity;
+	TSharedPtr<FDMXCategoryTreeNode> HoveredCategory;
+	TSubclassOf<UDMXEntity> HoveredTabType;
+	UDMXLibrary* HoveredLibrary;
+
+	bool bValidDropTarget;
+	bool bDraggingFromMultipleCategories;
+
+	/** Name of the entity being dragged or entities type for several ones */
+	FText DraggedLabel;
+};
+
+///////////////////////////////////////////////////////////////////////////////
+// Row widgets
+
+typedef STableRow<TSharedPtr<FDMXTreeNodeBase>> SDMXTableRowType;
+
+class SDMXCategoryRow
+	: public SDMXTableRowType
+{
+public:
+
+	SLATE_BEGIN_ARGS(SDMXCategoryRow)
+		{}
+		
+		SLATE_DEFAULT_SLOT(typename SDMXCategoryRow::FArguments, Content)
+	SLATE_END_ARGS()
+
+	void Construct(const FArguments& InArgs, const TSharedRef<STableViewBase>& InOwnerTableView, TSharedPtr<FDMXTreeNodeBase> InNodePtr, bool bInIsRootCategory, TWeakPtr<SDMXEntityList> InEditorList);
+
+	virtual void SetContent(TSharedRef< SWidget > InContent) override;
+	virtual void SetRowContent(TSharedRef< SWidget > InContent) override;
+
+	virtual const FSlateBrush* GetBorder() const { return nullptr; }
+	/* Get the node used by the row Widget */
+	virtual TSharedPtr<FDMXCategoryTreeNode> GetNode() const { return TreeNodePtr.Pin(); };
+
+protected:
+	//~ SWidget interface begin
+	virtual void OnDragEnter(const FGeometry& MyGeometry, const FDragDropEvent& DragDropEvent) override;
+	virtual void OnDragLeave(const FDragDropEvent& DragDropEvent) override;
+	virtual FReply OnDrop(const FGeometry& MyGeometry, const FDragDropEvent& DragDropEvent) override;
+	//~ SWidget interface end
+
+protected:
+	/** Pointer to node we represent */
+	TWeakPtr<FDMXCategoryTreeNode> TreeNodePtr;
+	TWeakPtr<SDMXEntityList> EditorListPtr;
+
+private:
+	const FSlateBrush* GetBackgroundImage() const;
+
+	TSharedPtr<SBorder> ContentBorder;
+};
+
+class SDMXEntityRow : public SDMXTableRowType
+{
+public:
+	DECLARE_DELEGATE_RetVal_TwoParams(FReply, FOnEntityDragged, TSharedPtr<FDMXTreeNodeBase>, const FPointerEvent&);
+	DECLARE_DELEGATE_RetVal(FText, FOnGetFilterText);
+	DECLARE_DELEGATE_OneParam(FOnAutoAssignChannelStateChanged, bool);
+
+	SLATE_BEGIN_ARGS(SDMXEntityRow)
+		: _OnEntityDragged()
+		, _OnGetFilterText()
+		, _OnAutoAssignChannelStateChanged()
+		{}
+
+		SLATE_EVENT(FOnEntityDragged, OnEntityDragged)
+		SLATE_EVENT(FOnGetFilterText, OnGetFilterText)
+		SLATE_EVENT(FOnAutoAssignChannelStateChanged, OnAutoAssignChannelStateChanged)
+	SLATE_END_ARGS()
+
+	void Construct(const FArguments& InArgs, TSharedPtr<FDMXTreeNodeBase> InNodePtr, TSharedPtr<STableViewBase> InOwnerTableView, TWeakPtr<SDMXEntityList> InEditorList);
+
+	/* Get the node used by the row Widget */
+	virtual TSharedPtr<FDMXEntityBaseTreeNode> GetNode() const { return TreeNodePtr.Pin(); };
+
+	FOnAutoAssignChannelStateChanged& GetOnAutoAssignChannelStateChanged() { return OnAutoAssignChannelStateChanged; }
+
+protected:
+	/** Data accessors */
+	FText GetDisplayText() const;
+	FText GetStartingChannelLabel() const;
+	FText GetEndingChannelLabel() const;
+
+	/**  Called when the auto-assign channel check-box state is changed */
+	void OnAutoAssignChannelBoxStateChanged(ECheckBoxState NewState);
+
+	//~ SWidget interface begin
+	virtual void OnDragEnter(const FGeometry& MyGeometry, const FDragDropEvent& DragDropEvent) override;
+	virtual void OnDragLeave(const FDragDropEvent& DragDropEvent) override;
+	virtual FReply OnDrop(const FGeometry& MyGeometry, const FDragDropEvent& DragDropEvent) override;
+	//~ SWidget interface end
+
+protected:
+	/** Pointer to node we represent */
+	TWeakPtr<FDMXEntityBaseTreeNode> TreeNodePtr;
+
+	TWeakPtr<SDMXEntityList> EditorListPtr;
+
+private:
+	/** Verifies the name of the component when changing it */
+	bool OnNameTextVerifyChanged(const FText& InNewText, FText& OutErrorMessage);
+
+	/** Commits the new name of the component */
+	void OnNameTextCommit(const FText& InNewName, ETextCommit::Type InTextCommit);
+
+	/** Returns the tooltip text for this row */
+	FText GetToolTipText() const;
+
+	/** Drag-drop handlers */
+	FReply HandleOnDragDetected(const FGeometry& MyGeometry, const FPointerEvent& MouseEvent);
+
+	/** Get the current filter text from the search box */
+	FText GetFilterText() const;
+
+	/** Get the icon for the Entity usability status. If it's all good, it's an emtpy image. */
+	const FSlateBrush* GetStatusIcon() const;
+	/** Get the tool tip text for the status icon */
+	FText GetStatusToolTip() const;
+
+private:
+	FOnEntityDragged OnEntityDragged;
+	FOnGetFilterText OnGetFilterText;
+	FText StatusIconToolTip;
+
+	FOnAutoAssignChannelStateChanged OnAutoAssignChannelStateChanged;
+
+	TSharedPtr<SInlineEditableTextBlock> InlineRenameWidget;
+};
+
+//////////////////////////////////////////////////////////////////////////
+// SDMXEntityListBase
+
+/** DMX entities list editor widget */
+class SDMXEntityList
+	: public SCompoundWidget, public FEditorUndoClient
+{
+public:
+	DECLARE_DELEGATE_OneParam(FOnSelectionUpdated, TArray<TSharedPtr<FDMXTreeNodeBase>>);
+	DECLARE_DELEGATE_OneParam(FOnItemDoubleClicked, const TSharedPtr<FDMXTreeNodeBase>);
+
+	SLATE_BEGIN_ARGS(SDMXEntityList)
+		: _OnSelectionUpdated()
+		{}
+
+		SLATE_ARGUMENT(TWeakPtr<FDMXEditor>,	DMXEditor)
+		SLATE_EVENT(FOnSelectionUpdated,		OnSelectionUpdated)
+
+	SLATE_END_ARGS()
+
+	/** Constructs the widget */
+	void Construct(const FArguments& InArgs, const TSubclassOf<UDMXEntity> InListType);
+
+	~SDMXEntityList();
+
+	/** SWidget interface */
+	virtual FReply OnKeyDown(const FGeometry& MyGeometry, const FKeyEvent& InKeyEvent);
+
+	bool IsListEmpty() const { return EntitiesCount == 0; }
+
+	/** Cut selected node(s) */
+	void OnCutSelectedNodes();
+	bool CanCutNodes() const;
+
+	/** Copy selected node(s) */
+	void OnCopySelectedNodes();
+	bool CanCopyNodes() const;
+
+	/** Pastes previously copied node(s) */
+	void OnPasteNodes();
+	bool CanPasteNodes() const;
+
+	/** Callbacks to duplicate the selected component */
+	bool CanDuplicateNodes() const;
+	void OnDuplicateNodes();
+
+	/** Removes existing selected component nodes from the SCS */
+	void OnDeleteNodes();
+	bool CanDeleteNodes() const;
+
+	/** Requests a rename on the selected Entity. */
+	void OnRenameNode();
+	/** Checks to see if renaming is allowed on the selected Entity */
+	bool CanRenameNode() const;
+
+	/** Get only the valid selected entities */
+	TArray<UDMXEntity*> GetSelectedEntities() const;
+	/** Selects an item by name */
+	void SelectItemByName(const FString& ItemName, ESelectInfo::Type SelectInfo = ESelectInfo::Direct);
+	/** Selects an item by Entity */
+	void SelectItemByEntity(const UDMXEntity* InEntity, ESelectInfo::Type SelectInfo = ESelectInfo::Direct);
+	/** Selects items by Entity */
+	void SelectItemsByEntity(const TArray<UDMXEntity*>& InEntities, ESelectInfo::Type SelectInfo = ESelectInfo::Direct);
+	/** Called when selection in the tree changes */
+	void OnTreeSelectionChanged(TSharedPtr<FDMXTreeNodeBase> InSelectedNodePtr, ESelectInfo::Type SelectInfo);
+
+	/** Update any associated selection from the passed in nodes */
+	void UpdateSelectionFromNodes(const TArray<TSharedPtr<FDMXTreeNodeBase>>& SelectedNodes);
+
+	/**
+	 * Set the expansion state of a node
+	 *
+	 * @param InNodeToChange	The node to be expanded/collapsed
+	 * @param bIsExpanded		True to expand the node, false to collapse it
+	 */
+	void SetNodeExpansionState(TSharedPtr<FDMXTreeNodeBase> InNodeToChange, const bool bIsExpanded);
+
+	/** Refresh the tree control to reflect changes in the editor */
+	void UpdateTree(bool bRegenerateTreeNodes = true);
+
+	/** Gets current filter from the FilterBox */
+	FText GetFilterText() const;
+
+	/** Gets the DMX Library object being edited */
+	UDMXLibrary* GetDMXLibrary() const;
+
+	TSubclassOf<UDMXEntity> GetListType() const;
+
+	//~ Begin FEditorUndoClient Interface
+	virtual void PostUndo(bool bSuccess) override;
+	virtual void PostRedo(bool bSuccess) override { PostUndo(bSuccess); }
+	// End of FEditorUndoClient
+
+protected:
+
+	/** Called to display context menu when right clicking on an Entity */
+	TSharedPtr< SWidget > OnContextMenuOpen();
+	void BuildAddNewMenu(FMenuBuilder& MenuBuilder);
+
+	/** Callback when an entity item is scrolled into view */
+	void OnItemScrolledIntoView(TSharedPtr<FDMXTreeNodeBase> InItem, const TSharedPtr<ITableRow>& InWidget);
+
+	/** Returns the set of expandable nodes that are currently collapsed in the UI */
+	void GetCollapsedNodes(TSet<TSharedPtr<FDMXTreeNodeBase>>& OutCollapsedNodes, TSharedPtr<FDMXTreeNodeBase> InParentNodePtr = nullptr) const;
+
+	/** Helper method to recursively find a tree node for the given DMX Entity starting at the given tree node */
+	TSharedPtr<FDMXTreeNodeBase> FindTreeNode(const UDMXEntity* InEntity, TSharedPtr<FDMXTreeNodeBase> InStartNodePtr = nullptr) const;
+	/** Helper method to recursively find a tree node with the given name starting at the given tree node */
+	TSharedPtr<FDMXTreeNodeBase> FindTreeNode(const FText& InName, TSharedPtr<FDMXTreeNodeBase> InStartNodePtr = nullptr) const;
+
+	/**
+	 * Creates a new category node directly under the passed parent or just retrieves it if existent.
+	 * If InParentNodePtr is null, a root category node is created/retrieved.
+	 */
+	TSharedPtr<FDMXTreeNodeBase> GetOrCreateCategoryNode(const FDMXTreeNodeBase::ECategoryType InCategoryType, const FText InCategoryName, void* CategoryValuePtr, TSharedPtr<FDMXTreeNodeBase> InParentNodePtr = nullptr, const FText& InToolTip = FText::GetEmpty());
+
+	/** Called when the active tab in the editor changes */
+	void OnActiveTabChanged(TSharedPtr<SDockTab> PreviouslyActive, TSharedPtr<SDockTab> NewlyActivated);
+
+	/** Searches this widget's parents to see if it's a child of InDockTab */
+	bool IsInTab(TSharedPtr<SDockTab> InDockTab) const;
+
+protected:
+	/** Pointer back to the DMXEditor tool that owns us */
+	TWeakPtr<FDMXEditor> DMXEditor;
+
+	/** Entity type we're editing. Might change the list layout */
+	TSubclassOf<UDMXEntity> ListType;
+
+	/** Tree widget */
+	TSharedPtr<STreeView<TSharedPtr<FDMXTreeNodeBase>>> EntitiesTreeWidget;
+
+	/** Command list for handling actions in the SSCSEditor */
+	TSharedPtr< FUICommandList > CommandList;
+
+	/** Delegate to invoke on selection update. */
+	FOnSelectionUpdated OnSelectionUpdated;
+
+	/**
+	 * Dummy root tree node.
+	 * It's not added to the tree, but the main categories and all their children
+	 * (entity and sub-category nodes) belong to it to make recursive searching algorithms nicer.
+	 */
+	TSharedPtr<FDMXTreeNodeBase> RootNode;
+
+private:
+	/** Empty Nodes array and create a FDMXEntityListTreeNode for each relevant entity and category */
+	void InitializeNodes();
+
+	TSharedPtr<FDMXEntityBaseTreeNode> CreateEntityTreeNode(UDMXEntity* Entity);
+
+	FReply OnAddNewClicked();
+
+	/** Callback when the filter is changed, forces the action tree(s) to filter */
+	void OnFilterTextChanged(const FText& InFilterText);
+
+	/**
+	 * Compares the filter bar's text with the item's component name. Use
+	 * bRecursive to refresh the state of child nodes as well. Returns true if
+	 * the node is set to be filtered out
+	 */
+	bool RefreshFilteredState(TSharedPtr<FDMXTreeNodeBase> TreeNode, bool bRecursive);
+
+	/** Used by tree control - make a widget for a table row from a node */
+	TSharedRef<ITableRow> MakeNodeWidget(TSharedPtr<FDMXTreeNodeBase> InNodePtr, const TSharedRef<STableViewBase>& OwnerTable);
+
+	/** Used by tree control - get children for a specified node */
+	void OnGetChildrenForTree(TSharedPtr<FDMXTreeNodeBase> InNodePtr, TArray<TSharedPtr<FDMXTreeNodeBase>>& OutChildren);
+
+	/** Expand all categories during filtering and resets nodes expansion state after filtering is cleared */
+	void UpdateNodesExpansion(TSharedRef<FDMXTreeNodeBase> InRootNode, bool bFilterIsEmpty);
+
+	/** Handler for expanding/collapsing items */
+	void OnItemExpansionChanged(TSharedPtr<FDMXTreeNodeBase> InNodePtr, bool bInExpansionState);
+
+	/** Handler for when an entity from the list is dragged */
+	FReply OnEntityDragged(TSharedPtr<FDMXTreeNodeBase> InNodePtr, const FPointerEvent& MouseEvent);
+
+	/** Called when the user selects a Fixture Type to create a Fixture Patch from */
+	void OnFixtureTypeSelected(UDMXEntity* InSelectedFixtureType);
+	/** Called by the editor to set a base name for an Entity about to be created */
+	void OnEditorGetBaseNameForNewFixturePatch(TSubclassOf<UDMXEntity> InEntityClass, FString& OutBaseName, UDMXEntityFixtureType* InSelectedFixtureType);
+	/** Called by the editor to setup the properties of a new Entity before its selection and renaming in the editor */
+	void OnEditorSetupNewFixturePatch(UDMXEntity* InNewEntity, UDMXEntityFixtureType* InSelectedFixtureType);
+
+	/** Called when Auto Assign Channel is changed for a Patch */
+	void OnAutoAssignChannelStateChanged(bool NewState, TSharedPtr<FDMXTreeNodeBase> InNodePtr);
+
+private:
+	/** The filter box that handles filtering entities */
+	TSharedPtr< SSearchBox > FilterBox;
+
+	/** Handle to the registered OnActiveTabChanged delegate */
+	FDelegateHandle OnActiveTabChangedDelegateHandle;
+	
+	FDelegateHandle OnGetBaseNameForNewEntityHandle;
+	FDelegateHandle OnSetupNewEntityHandle;
+
+	/** Gate to prevent changing the selection while selection change is being broadcast. */
+	bool bUpdatingSelection;
+
+	/** The green Add Button. We need to reference it in Fixture Patches tab to update its list */
+	TSharedPtr<SComboButton> AddComboButton;
+	TSharedPtr<SDMXEntityDropdownMenu<UDMXEntityFixtureType>> AddButtonDropdownList;
+
+	/** Keeps the current number of entities for quick checking */
+	int32 EntitiesCount;
+};
diff --git a/Engine/Plugins/Experimental/DigitalMultiplex/DMXEngine/Source/DMXEditor/Private/Widgets/SDMXGDTFOptionWindow.h b/Engine/Plugins/Experimental/DigitalMultiplex/DMXEngine/Source/DMXEditor/Private/Widgets/SDMXGDTFOptionWindow.h
new file mode 100644
--- /dev/null
+++ b/Engine/Plugins/Experimental/DigitalMultiplex/DMXEngine/Source/DMXEditor/Private/Widgets/SDMXGDTFOptionWindow.h
@@ -0,0 +1,104 @@
+// Copyright Epic Games, Inc. All Rights Reserved.
+
+#pragma once
+
+#include "CoreMinimal.h"
+#include "InputCoreTypes.h"
+#include "Widgets/DeclarativeSyntaxSupport.h"
+#include "Input/Reply.h"
+#include "Widgets/SCompoundWidget.h"
+#include "Widgets/SWindow.h"
+
+class SButton;
+class UDMXGDTFImportUI;
+class IDetailsView;
+
+class SDMXGDTFOptionWindow : public SCompoundWidget
+{
+public:
+    SLATE_BEGIN_ARGS( SDMXGDTFOptionWindow )
+    : _ImportUI(nullptr)
+    , _WidgetWindow()
+    , _FullPath()
+    , _MaxWindowHeight(0.0f)
+    , _MaxWindowWidth(0.0f)
+    {}
+
+    SLATE_ARGUMENT( UDMXGDTFImportUI*, ImportUI )
+    SLATE_ARGUMENT( TSharedPtr<SWindow>, WidgetWindow )
+    SLATE_ARGUMENT( FText, FullPath )
+    SLATE_ARGUMENT( float, MaxWindowHeight)
+    SLATE_ARGUMENT(float, MaxWindowWidth)
+    SLATE_END_ARGS()
+
+public:
+    void Construct(const FArguments& InArgs);
+    virtual bool SupportsKeyboardFocus() const override { return true; }
+
+    FReply OnImport()
+    {
+        bShouldImport = true;
+        if ( WidgetWindow.IsValid() )
+        {
+            WidgetWindow.Pin()->RequestDestroyWindow();
+        }
+        return FReply::Handled();
+    }
+
+    FReply OnImportAll()
+    {
+        bShouldImportAll = true;
+        return OnImport();
+    }
+
+    FReply OnCancel()
+    {
+        bShouldImport = false;
+        bShouldImportAll = false;
+        if ( WidgetWindow.IsValid() )
+        {
+            WidgetWindow.Pin()->RequestDestroyWindow();
+        }
+        return FReply::Handled();
+    }
+
+    virtual FReply OnKeyDown( const FGeometry& MyGeometry, const FKeyEvent& InKeyEvent ) override
+    {
+        if( InKeyEvent.GetKey() == EKeys::Escape )
+        {
+            return OnCancel();
+        }
+
+        return FReply::Unhandled();
+    }
+
+    bool ShouldImport() const
+    {
+        return bShouldImport;
+    }
+
+    bool ShouldImportAll() const
+    {
+        return bShouldImportAll;
+    }
+
+    SDMXGDTFOptionWindow()
+        : ImportUI(nullptr)
+        , bShouldImport(false)
+        , bShouldImportAll(false)
+    {}
+
+private:
+
+    FReply OnResetToDefaultClick() const;
+    FText GetImportTypeDisplayText() const;
+
+
+private:
+    TWeakObjectPtr<UDMXGDTFImportUI> ImportUI;
+    TSharedPtr<IDetailsView> DetailsView;
+    TWeakPtr< SWindow > WidgetWindow;
+    TSharedPtr< SButton > ImportButton;
+    bool			bShouldImport;
+    bool			bShouldImportAll;
+};
diff --git a/Engine/Plugins/Experimental/DigitalMultiplex/DMXEngine/Source/DMXEditor/Private/Widgets/SDMXInputConsole.cpp b/Engine/Plugins/Experimental/DigitalMultiplex/DMXEngine/Source/DMXEditor/Private/Widgets/SDMXInputConsole.cpp
new file mode 100644
--- /dev/null
+++ b/Engine/Plugins/Experimental/DigitalMultiplex/DMXEngine/Source/DMXEditor/Private/Widgets/SDMXInputConsole.cpp
@@ -0,0 +1,52 @@
+// Copyright Epic Games, Inc. All Rights Reserved.
+
+#include "Widgets/SDMXInputConsole.h"
+#include "Widgets/SDMXInputInfoSelecter.h"
+#include "Widgets/SDMXInputInfo.h"
+#include "Widgets/SDMXInputInfo.h"
+#include "Widgets/Layout/SSeparator.h"
+
+#include "Widgets/SBoxPanel.h"
+
+#define LOCTEXT_NAMESPACE "SDMXInputConsole"
+
+void SDMXInputConsole::Construct(const FArguments& InArgs)
+{
+	InputInfoSelecter = SNew(SDMXInputInfoSelecter)
+		.OnUniverseSelectionChanged(this, &SDMXInputConsole::OnUniverseSelectionChanged);
+	InputInfo = SNew(SDMXInputInfo)
+		.InfoSelecter(InputInfoSelecter);
+
+	ChildSlot
+		[
+			SNew(SVerticalBox)
+			+ SVerticalBox::Slot()
+				.VAlign(VAlign_Top)
+				.AutoHeight()
+				[
+					InputInfoSelecter.ToSharedRef()
+				]
+			+ SVerticalBox::Slot()
+				.AutoHeight()
+				.Padding(3.0f)
+				[
+					SNew(SSeparator)
+						.Orientation(Orient_Horizontal)
+				]
+			+ SVerticalBox::Slot()
+				.FillHeight(1.0f)
+				[
+					InputInfo.ToSharedRef()
+				]
+		];
+}
+
+void SDMXInputConsole::OnUniverseSelectionChanged(const FName& InProtocol)
+{
+	if (InputInfo.IsValid())
+	{
+		InputInfo->ResetUISequanceID();
+	}
+}
+
+#undef LOCTEXT_NAMESPACE
diff --git a/Engine/Plugins/Experimental/DigitalMultiplex/DMXEngine/Source/DMXEditor/Private/Widgets/SDMXInputConsole.h b/Engine/Plugins/Experimental/DigitalMultiplex/DMXEngine/Source/DMXEditor/Private/Widgets/SDMXInputConsole.h
new file mode 100644
--- /dev/null
+++ b/Engine/Plugins/Experimental/DigitalMultiplex/DMXEngine/Source/DMXEditor/Private/Widgets/SDMXInputConsole.h
@@ -0,0 +1,37 @@
+// Copyright Epic Games, Inc. All Rights Reserved.
+
+#pragma once
+
+#include "CoreMinimal.h"
+#include "Widgets/DeclarativeSyntaxSupport.h"
+#include "Widgets/SCompoundWidget.h"
+
+class SDMXInputInfoSelecter;
+class SDMXInputInfo;
+
+/** Widget to hold tab widgets */
+class SDMXInputConsole
+	: public SCompoundWidget
+{
+public:
+	SLATE_BEGIN_ARGS(SDMXInputConsole)
+	{}
+
+	SLATE_END_ARGS()
+
+	/** Constructs the widget */
+	void Construct(const FArguments& InArgs);
+
+	const TSharedRef<SDMXInputInfoSelecter> GetInputInfoSelecter() const { return InputInfoSelecter.ToSharedRef(); }
+
+	const TSharedRef<SDMXInputInfo> GetInputInfo() const { return InputInfo.ToSharedRef(); }
+
+private:
+	void OnUniverseSelectionChanged(const FName& InProtocol);
+
+private:
+
+	TSharedPtr<SDMXInputInfoSelecter> InputInfoSelecter;
+
+	TSharedPtr<SDMXInputInfo> InputInfo;
+};
diff --git a/Engine/Plugins/Experimental/DigitalMultiplex/DMXEngine/Source/DMXEditor/Private/Widgets/SDMXInputInfo.cpp b/Engine/Plugins/Experimental/DigitalMultiplex/DMXEngine/Source/DMXEditor/Private/Widgets/SDMXInputInfo.cpp
new file mode 100644
--- /dev/null
+++ b/Engine/Plugins/Experimental/DigitalMultiplex/DMXEngine/Source/DMXEditor/Private/Widgets/SDMXInputInfo.cpp
@@ -0,0 +1,128 @@
+// Copyright Epic Games, Inc. All Rights Reserved.
+
+#include "Widgets/SDMXInputInfo.h"
+
+#include "DMXEditor.h"
+#include "Library/DMXLibrary.h"
+#include "Interfaces/IDMXProtocol.h"
+#include "Interfaces/IDMXProtocolUniverse.h"
+#include "DMXEditorLog.h"
+#include "DMXProtocolTypes.h"
+#include "DMXProtocolConstants.h"
+#include "Widgets/SDMXInputInfoChannelValue.h"
+#include "Widgets/SDMXInputInfoSelecter.h"
+
+#include "Widgets/Layout/SWrapBox.h"
+#include "Widgets/Layout/SScrollBox.h"
+
+#define LOCTEXT_NAMESPACE "SDMXInputInfo"
+
+void SDMXInputInfo::Construct(const FArguments& InArgs)
+{
+	WeakInfoSelecter = InArgs._InfoSelecter;
+
+	SetVisibility(EVisibility::SelfHitTestInvisible);
+
+	// Tick each frame but update only if new data is coming
+	// In case of ticking we update UI only once per frame if any DMX data is coming
+	SetCanTick(true);
+
+	static const float PaddingBorders = 15.0f;
+	static const float PaddingChannelValues = 3.0f;
+
+	ChildSlot
+	.Padding(PaddingBorders)
+	[
+		// root
+		SNew(SScrollBox)
+		.Orientation(Orient_Vertical)
+		.ScrollBarAlwaysVisible(false)
+
+		+SScrollBox::Slot()
+		.HAlign(HAlign_Fill)
+		[
+			SAssignNew(ChannelValuesBox, SWrapBox)
+			.UseAllottedWidth(true).InnerSlotPadding(FVector2D(1.0f))
+		]
+	];
+
+	ResetUISequanceID();
+	CreateChannelValueWidgets();
+
+	// Set buffer values to 0
+	ChannelsValues.SetNum(DMX_UNIVERSE_SIZE);
+	FMemory::Memset(ChannelsValues.GetData(), 0, DMX_UNIVERSE_SIZE);
+	UpdateChannelWidgetsValues(ChannelsValues);
+}
+
+void SDMXInputInfo::Tick(const FGeometry& AllottedGeometry, const double InCurrentTime, const float InDeltaTime)
+{
+	UpdateChannelsValues();
+}
+
+void SDMXInputInfo::CreateChannelValueWidgets()
+{
+	if (!ChannelValuesBox.IsValid())
+	{
+		return;
+	}
+
+	ChannelValueWidgets.Reserve(DMX_UNIVERSE_SIZE);
+
+	for (uint32 ChannelIndex = 0; ChannelIndex < DMX_UNIVERSE_SIZE; ++ChannelIndex)
+	{
+		TSharedPtr<SDMXInputInfoChannelValue> ChannelValueWidget;
+
+		ChannelValuesBox->AddSlot()
+			.HAlign(HAlign_Fill)
+			.VAlign(VAlign_Fill)
+			[
+				SAssignNew(ChannelValueWidget, SDMXInputInfoChannelValue)
+				.ID(ChannelIndex + 1) // +1 because channels start at 1
+				.Value(0)
+			];
+
+		ChannelValueWidgets.Add(ChannelValueWidget);
+	}
+}
+
+void SDMXInputInfo::UpdateChannelsValues()
+{
+	if (TSharedPtr<SDMXInputInfoSelecter> InfoSelecterPtr = WeakInfoSelecter.Pin())
+	{
+		if (TSharedPtr<IDMXProtocol> DMXProtocolPtr = IDMXProtocol::Get(InfoSelecterPtr->GetCurrentProtocolName()))
+		{
+			if (IDMXProtocolUniversePtr Universe = DMXProtocolPtr->GetUniverseById(InfoSelecterPtr->GetCurrentUniverseID()))
+			{
+				if (TSharedPtr<FDMXBuffer> DMXBuffer = Universe->GetInputDMXBuffer())
+				{
+					// check the sequence ID 
+					uint32 BufferSequanceID = DMXBuffer->GetSequanceID();
+					if (BufferSequanceID != UISequanceID)
+					{
+						FMemory::Memcpy(ChannelsValues.GetData(), DMXBuffer->GetDMXData().GetData(), DMX_UNIVERSE_SIZE);
+						UpdateChannelWidgetsValues(ChannelsValues);
+					}
+
+					UISequanceID = BufferSequanceID;
+				}
+			}
+		}
+	}
+}
+
+void SDMXInputInfo::UpdateChannelWidgetsValues(const TArray<uint8>& NewValues)
+{
+	if (NewValues.Num() != DMX_UNIVERSE_SIZE)
+	{
+		UE_LOG_DMXEDITOR(Error, TEXT("%S: Input values has the wrong number of channels!"), __FUNCTION__);
+		return;
+	}
+
+	for (uint32 ChannelID = 0; ChannelID < DMX_UNIVERSE_SIZE; ++ChannelID)
+	{
+		ChannelValueWidgets[ChannelID]->SetValue(NewValues[ChannelID]);
+	}
+}
+
+#undef LOCTEXT_NAMESPACE
\ No newline at end of file
diff --git a/Engine/Plugins/Experimental/DigitalMultiplex/DMXEngine/Source/DMXEditor/Private/Widgets/SDMXInputInfo.h b/Engine/Plugins/Experimental/DigitalMultiplex/DMXEngine/Source/DMXEditor/Private/Widgets/SDMXInputInfo.h
new file mode 100644
--- /dev/null
+++ b/Engine/Plugins/Experimental/DigitalMultiplex/DMXEngine/Source/DMXEditor/Private/Widgets/SDMXInputInfo.h
@@ -0,0 +1,57 @@
+// Copyright Epic Games, Inc. All Rights Reserved.
+
+#pragma once
+
+#include "CoreMinimal.h"
+#include "Widgets/DeclarativeSyntaxSupport.h"
+#include "Widgets/SCompoundWidget.h"
+
+class SWrapBox;
+class SDMXInputInfoChannelValue;
+class SDMXInputInfoSelecter;
+
+/** DMX Input inspector widget */
+class SDMXInputInfo
+	: public SCompoundWidget
+{
+public:
+	SLATE_BEGIN_ARGS(SDMXInputInfo)
+	{}
+	SLATE_ARGUMENT(TWeakPtr<SDMXInputInfoSelecter>, InfoSelecter)
+	SLATE_END_ARGS()
+
+	/** Constructs the widget */
+	void Construct(const FArguments& InArgs);
+
+	//~ Begin SWidget interface
+	virtual void Tick(const FGeometry& AllottedGeometry, const double InCurrentTime, const float InDeltaTime) override;
+	//~End of SWidget interface
+
+	const TArray<uint8>& GetChannelsValues() const { return ChannelsValues; }
+
+	void ResetUISequanceID() { UISequanceID = 0; }
+	
+protected:
+	/** Container widget for all the channels' values */
+	TSharedPtr<SWrapBox> ChannelValuesBox;
+
+	/** Channels' values for testing purpose */
+	TArray<uint8> ChannelsValues;
+
+	TArray<TSharedPtr<SDMXInputInfoChannelValue>> ChannelValueWidgets;
+
+	TWeakPtr<SDMXInputInfoSelecter> WeakInfoSelecter;
+
+protected:
+	/** Spawns the channel value widgets */
+	void CreateChannelValueWidgets();
+
+	/** channel values */
+	void UpdateChannelsValues();
+
+	void UpdateChannelWidgetsValues(const TArray<uint8>& NewValues);
+
+private:
+	/** ID of the sequence on input info widget */
+	uint32 UISequanceID;
+}; 
\ No newline at end of file
diff --git a/Engine/Plugins/Experimental/DigitalMultiplex/DMXEngine/Source/DMXEditor/Private/Widgets/SDMXInputInfoChannelValue.cpp b/Engine/Plugins/Experimental/DigitalMultiplex/DMXEngine/Source/DMXEditor/Private/Widgets/SDMXInputInfoChannelValue.cpp
new file mode 100644
--- /dev/null
+++ b/Engine/Plugins/Experimental/DigitalMultiplex/DMXEngine/Source/DMXEditor/Private/Widgets/SDMXInputInfoChannelValue.cpp
@@ -0,0 +1,174 @@
+// Copyright Epic Games, Inc. All Rights Reserved.
+
+#include "Widgets/SDMXInputInfoChannelValue.h"
+
+#include "Interfaces/IDMXProtocol.h"
+#include "DMXEditorLog.h"
+#include "DMXProtocolConstants.h"
+#include "DMXEditorStyle.h"
+
+#include "Widgets/SOverlay.h"
+#include "Widgets/Text/STextBlock.h"
+#include "Widgets/SBoxPanel.h"
+#include "Widgets/Layout/SSeparator.h"
+#include "Widgets/Notifications/SProgressBar.h"
+#include "Widgets/Images/SImage.h"
+#include "Brushes/SlateColorBrush.h"
+#include "UObject/ConstructorHelpers.h"
+#include "Engine/Font.h"
+
+#define LOCTEXT_NAMESPACE "SDMXInputInfo"
+
+const float SDMXInputInfoChannelValue::NewValueChangedAnimDuration = 0.8f;
+
+const FLinearColor SDMXInputInfoChannelValue::IDColor = FLinearColor(1,1,1, 0.6f);
+const FLinearColor SDMXInputInfoChannelValue::ValueColor = FLinearColor(1,1,1, 0.9f);
+
+void SDMXInputInfoChannelValue::Construct(const FArguments& InArgs)
+{
+	SetVisibility(EVisibility::SelfHitTestInvisible);
+	SetCanTick(false);
+
+	BoundID = InArgs._ID;
+	BoundValue = InArgs._Value;
+	NewValueFreshness = 0.0f;
+
+	const float PaddingInfo = 3.0f;
+
+	ChildSlot
+	[
+		// root
+		SNew(SOverlay)
+
+		+ SOverlay::Slot()
+		.HAlign(HAlign_Fill)
+		.VAlign(VAlign_Fill)
+		[
+			// Background color image
+			SAssignNew(BarColorBorder, SImage)
+			.Image(FDMXEditorStyle::Get().GetBrush(TEXT("DMXEditor.WhiteBrush")))
+			.ColorAndOpacity(this, &SDMXInputInfoChannelValue::GetBackgroundColor)
+		]
+
+		// Info
+		+ SOverlay::Slot()
+		.HAlign(HAlign_Fill)
+		.VAlign(VAlign_Fill)
+		.Padding(PaddingInfo)
+		[
+			SNew(SVerticalBox)
+
+			// ID Label
+			+ SVerticalBox::Slot()
+			.FillHeight(0.5f)
+			.VAlign(VAlign_Center)
+			.HAlign(HAlign_Center)
+			[
+				SNew(STextBlock)
+				.Text(this, &SDMXInputInfoChannelValue::GetIDLabel)
+				.ColorAndOpacity(FSlateColor(IDColor))
+				.MinDesiredWidth(23.0f)
+				.Justification(ETextJustify::Center)
+				.Font(FDMXEditorStyle::Get().GetFontStyle("DMXEditor.Font.InputChannelID"))
+			]
+
+			// Separator
+			+ SVerticalBox::Slot()
+			.AutoHeight()
+			.HAlign(HAlign_Fill)
+			[
+				SNew(SSeparator)
+				.Orientation(Orient_Horizontal)
+			]
+
+			// Value Label
+			+ SVerticalBox::Slot()
+			.FillHeight(0.5f)
+			.VAlign(VAlign_Center)
+			.HAlign(HAlign_Center)
+			[
+				SNew(STextBlock)
+				.Text(this, &SDMXInputInfoChannelValue::GetValueLabel)
+				.ColorAndOpacity(FSlateColor(ValueColor))
+				.MinDesiredWidth(23.0f)
+				.Justification(ETextJustify::Center)
+				.Font(FDMXEditorStyle::Get().GetFontStyle("DMXEditor.Font.InputChannelValue"))
+			]
+		]
+	];
+}
+
+void SDMXInputInfoChannelValue::SetID(const TAttribute<uint32>& NewID)
+{
+	BoundID = NewID;
+}
+
+void SDMXInputInfoChannelValue::SetValue(const TAttribute<uint8>& NewValue)
+{
+	// is NewValue a different value from current one?
+	if (NewValue.Get() != BoundValue.Get())
+	{
+		// Activate timer to animate value bar color
+		if (!AnimationTimerHandle.IsValid())
+		{
+			AnimationTimerHandle = RegisterActiveTimer(0.0f, FWidgetActiveTimerDelegate::CreateSP(this, &SDMXInputInfoChannelValue::UpdateValueChangedAnim));
+		}
+		// restart value change animation
+		NewValueFreshness = 1.0f;
+	}
+	BoundValue = NewValue;
+}
+
+EActiveTimerReturnType SDMXInputInfoChannelValue::UpdateValueChangedAnim(double InCurrentTime, float InDeltaTime)
+{
+	NewValueFreshness = FMath::Max(NewValueFreshness - InDeltaTime / NewValueChangedAnimDuration, 0.0f);
+	
+	// disable timer when the value bar color animation ends
+	if (NewValueFreshness <= 0.0f)
+	{
+		TSharedPtr<FActiveTimerHandle> PinnedTimerHandle = AnimationTimerHandle.Pin();
+		if (PinnedTimerHandle.IsValid())
+		{
+			UnRegisterActiveTimer(PinnedTimerHandle.ToSharedRef());
+		}
+	}
+	return EActiveTimerReturnType::Continue;
+}
+
+FText SDMXInputInfoChannelValue::GetIDLabel() const
+{
+	return FText::AsNumber(BoundID.Get());
+}
+
+FText SDMXInputInfoChannelValue::GetValueLabel() const
+{
+	return FText::AsNumber(BoundValue.Get());
+}
+
+FSlateColor SDMXInputInfoChannelValue::GetBackgroundColor() const
+{
+	const float CurrentPercent = static_cast<float>(BoundValue.Get()) / DMX_MAX_CHANNEL_VALUE;
+
+	// totally transparent when 0
+	if (CurrentPercent <= 0.0f)
+	{
+		return FLinearColor(0, 0, 0, 0);
+	}
+
+	// Intensities to be animated when a new value is set and then multiplied by the background color
+	static const float NormalIntensity = 0.3f;
+	static const float FreshValueIntensity = 0.7f;
+	// lerp intensity depending on NewValueFreshness^2 to make it pop for a while when it has just been updated
+	const float ValueFreshnessIntensity = FMath::Lerp(NormalIntensity, FreshValueIntensity, NewValueFreshness * NewValueFreshness);
+
+	// color variations for low and high channel values
+	static const FVector LowValueColor = FVector(0, 0.045f, 0.15f);
+	static const FVector HighValueColor = FVector(0, 0.3f, 1.0f);
+	const FVector ColorFromChannelValue = FMath::Lerp(LowValueColor, HighValueColor, CurrentPercent);
+
+	// returning a FVector, a new FSlateColor will be created from it with (RGB = vector, Alpha = 1.0)
+	FVector Result = ColorFromChannelValue * ValueFreshnessIntensity;
+	return FSlateColor(FLinearColor(Result.X, Result.Y, Result.Z));
+}
+
+#undef LOCTEXT_NAMESPACE
diff --git a/Engine/Plugins/Experimental/DigitalMultiplex/DMXEngine/Source/DMXEditor/Private/Widgets/SDMXInputInfoChannelValue.h b/Engine/Plugins/Experimental/DigitalMultiplex/DMXEngine/Source/DMXEditor/Private/Widgets/SDMXInputInfoChannelValue.h
new file mode 100644
--- /dev/null
+++ b/Engine/Plugins/Experimental/DigitalMultiplex/DMXEngine/Source/DMXEditor/Private/Widgets/SDMXInputInfoChannelValue.h
@@ -0,0 +1,91 @@
+// Copyright Epic Games, Inc. All Rights Reserved.
+
+#pragma once
+
+#include "CoreMinimal.h"
+#include "Widgets/DeclarativeSyntaxSupport.h"
+#include "Widgets/SCompoundWidget.h"
+
+class SImage;
+class FActiveTimerHandle;
+
+/** DMX Input Channel value representation widget */
+class SDMXInputInfoChannelValue
+	: public SCompoundWidget
+{
+public:
+	SLATE_BEGIN_ARGS(SDMXInputInfoChannelValue)
+		: _ID(0)
+		, _Value(0)
+		{}
+
+	/** The channel ID this widget represents */
+	SLATE_ATTRIBUTE(uint32, ID)
+
+	/** The current value from the channel */
+	SLATE_ATTRIBUTE(uint8, Value)
+
+	SLATE_END_ARGS()
+
+	/** Constructs the widget */
+	void Construct(const FArguments& InArgs);
+
+public:
+	/** Sets the channel ID this widget represents */
+	void SetID(const TAttribute<uint32>& NewID);
+	/** Gets the channel ID this widget represents */
+	uint32 GetID() const
+	{
+		return BoundID.Get();
+	}
+
+	/** Sets the current value from the channel */
+	void SetValue(const TAttribute<uint8>& NewValue);
+	/** Gets the current value from the channel */
+	uint8 GetValue() const
+	{
+		return BoundValue.Get();
+	}
+
+	/**
+	 * Updates the variable that controls the color animation progress for the Value Bar.
+	 * This is called by a timer.
+	 */
+	EActiveTimerReturnType UpdateValueChangedAnim(double InCurrentTime, float InDeltaTime);
+
+private:
+	/** The channel ID this widget represents */
+	TAttribute<uint32> BoundID;
+	/** The current value from the channel */
+	TAttribute<uint8> BoundValue;
+
+	/** The ProgressBar widget to display the channel value graphically */
+	TSharedPtr<SImage> BarColorBorder;
+
+	/**
+	 * Used to animate the color when the value changes.
+	 * 0..1 range: 1 = value has just changed, 0 = standard color
+	 */
+	float NewValueFreshness;
+	/** How long it takes to become standard color again after a new value is set */
+	static const float NewValueChangedAnimDuration;
+	/** Used to stop the animation timer once the animation is completed */
+	TWeakPtr<FActiveTimerHandle> AnimationTimerHandle;
+
+	// ~ VISUAL STYLE CONSTANTS
+
+	/** Color of the ID label */
+	static const FLinearColor IDColor;
+
+	/** Color of the Value label */
+	static const FLinearColor ValueColor;
+
+private:
+	/** Returns the channel ID in Text form to display it in the UI */
+	FText GetIDLabel() const;
+	/** Returns the channel value in Text form to display it in the UI */
+	FText GetValueLabel() const;
+
+	/** Returns the fill color for the ValueBar */
+	FSlateColor GetBackgroundColor() const;
+};
diff --git a/Engine/Plugins/Experimental/DigitalMultiplex/DMXEngine/Source/DMXEditor/Private/Widgets/SDMXInputInfoSelecter.cpp b/Engine/Plugins/Experimental/DigitalMultiplex/DMXEngine/Source/DMXEditor/Private/Widgets/SDMXInputInfoSelecter.cpp
new file mode 100644
--- /dev/null
+++ b/Engine/Plugins/Experimental/DigitalMultiplex/DMXEngine/Source/DMXEditor/Private/Widgets/SDMXInputInfoSelecter.cpp
@@ -0,0 +1,177 @@
+// Copyright Epic Games, Inc. All Rights Reserved.
+
+#include "Widgets/SDMXInputInfoSelecter.h"
+
+#include "DMXEditor.h"
+#include "Library/DMXLibrary.h"
+#include "DMXEditorLog.h"
+#include "Interfaces/IDMXProtocol.h"
+#include "DMXProtocolSettings.h"
+#include "Widgets/SNameListPicker.h"
+
+#include "Widgets/Input/SSpinBox.h"
+#include "Widgets/Layout/SWrapBox.h"
+
+#define LOCTEXT_NAMESPACE "SDMXInputInfoSelecter"
+
+void SDMXInputInfoSelecter::Construct(const FArguments& InArgs)
+{
+	// Sets delegates
+	OnUniverseSelectionChanged = InArgs._OnUniverseSelectionChanged;
+
+	SetVisibility(EVisibility::SelfHitTestInvisible);
+
+	const float PaddingBorders = 15.0f;
+	const float PaddingKeyVal = 10.0f;
+	const float PaddingNewInput = 35.0f;
+	const FVector2D PaddingInner(PaddingNewInput, 10.0f);
+
+	// Get values from plugin settings
+	UDMXProtocolSettings* ProtocolSettings = GetMutableDefault<UDMXProtocolSettings>();
+	CurrentUniverseID = ProtocolSettings->InputConsoleUniverseID;
+	if (ProtocolSettings->InputConsoleProtocol.IsNone())
+	{
+		// Create default ProtocolName struct, which will have a valid protocol name
+		CurrentProtocol = FDMXProtocolName();
+		ProtocolSettings->InputConsoleProtocol = CurrentProtocol;
+		ProtocolSettings->SaveConfig();
+	}
+	else
+	{
+		CurrentProtocol = FDMXProtocolName(ProtocolSettings->InputConsoleProtocol);
+	}
+
+	ChildSlot
+	.Padding(PaddingBorders)
+	[
+		// root
+		SNew(SWrapBox)
+		.InnerSlotPadding(PaddingInner)
+		.UseAllottedWidth(true)
+
+		// Protocol
+		+ SWrapBox::Slot()
+		.HAlign(HAlign_Left)
+		.VAlign(VAlign_Center)
+		[
+			SNew(SHorizontalBox)
+			//Label
+			+ SHorizontalBox::Slot()
+			.VAlign(VAlign_Center)
+			.AutoWidth()
+			[
+				SNew(STextBlock)
+				.Text(LOCTEXT("ProtocolLabel", "Protocol"))
+			]
+			// Protocol combo box
+			+ SHorizontalBox::Slot()
+			.Padding(PaddingKeyVal, 0, 0, 0)
+			.AutoWidth()
+			[
+				SNew(SNameListPicker)
+				.Value(this, &SDMXInputInfoSelecter::GetCurrentProtocolName)
+				.OnValueChanged(this, &SDMXInputInfoSelecter::HandleProtocolChanged)
+				.OptionsSource(FDMXProtocolName::GetPossibleValues())
+			]
+		]
+
+		// Universe ID
+		+ SWrapBox::Slot()
+		.HAlign(HAlign_Left)
+		.VAlign(VAlign_Center)
+		[
+			SNew(SHorizontalBox)
+			//Label
+			+ SHorizontalBox::Slot()
+			.VAlign(VAlign_Center)
+			.AutoWidth()
+			[
+				SNew(STextBlock)
+				.Text(LOCTEXT("UniverseIDLabel", "Universe ID"))
+			]
+			// Final Universe ID current value text.
+			+ SHorizontalBox::Slot()
+			.Padding(PaddingKeyVal, 0, 0, 0)
+			.AutoWidth()
+			[
+				SAssignNew(UniverseIDField, SSpinBox<uint16>)
+				.Value(this, &SDMXInputInfoSelecter::GetCurrentUniverseID)
+				.OnValueChanged(this, &SDMXInputInfoSelecter::HandleUniverseIDChanged)
+				.OnValueCommitted(this, &SDMXInputInfoSelecter::HandleUniverseIDValueCommitted)
+				.MinValue(0).MaxValue(MAX_uint16)
+				.MinSliderValue(0).MaxSliderValue(MAX_uint16)
+				.MinDesiredWidth(50.0f)
+			]
+		]
+	];
+
+	// Update UniverseID field min and max values
+	HandleProtocolChanged(CurrentProtocol);
+}
+
+TSharedRef<SWidget> SDMXInputInfoSelecter::GenerateProtocolItemWidget(TSharedPtr<FName> InItem)
+{
+	if (!InItem.IsValid())
+	{
+		UE_LOG_DMXEDITOR(Warning, TEXT("InItem for GenerateProtocolItemWidget was null!"));
+		return SNew(STextBlock)
+			.Text(LOCTEXT("NullComboBoxItemLabel", "Null Error"));
+	}
+
+	return SNew(STextBlock)
+		.Text(FText::FromName(*InItem));
+}
+
+void SDMXInputInfoSelecter::HandleProtocolChanged(FName SelectedProtocol)
+{
+	FDMXProtocolName ProtocolName(SelectedProtocol);
+	if (!ProtocolName.IsValid())
+	{
+		UE_LOG_DMXEDITOR(Error, TEXT("%S: Selected null protocol!"), __FUNCTION__);
+		return;
+	}
+
+	CurrentProtocol = ProtocolName;
+
+	uint16 MinUniverseID = CurrentProtocol.GetProtocol()->GetMinUniverseID();
+	uint16 MaxUniverseID = CurrentProtocol.GetProtocol()->GetMaxUniverses();
+
+	UniverseIDField->SetMinSliderValue(MinUniverseID);
+	UniverseIDField->SetMinValue(MinUniverseID);
+	UniverseIDField->SetMaxSliderValue(MaxUniverseID);
+	UniverseIDField->SetMaxValue(MaxUniverseID);
+
+	CurrentUniverseID = FMath::Clamp(CurrentUniverseID, MinUniverseID, MaxUniverseID);
+
+	// Execute delegate
+	OnUniverseSelectionChanged.ExecuteIfBound(SelectedProtocol);
+
+	// Update stored settings
+	UDMXProtocolSettings* ProtocolSettings = GetMutableDefault<UDMXProtocolSettings>();
+	ProtocolSettings->InputConsoleProtocol = SelectedProtocol;
+	ProtocolSettings->InputConsoleUniverseID = CurrentUniverseID;
+	ProtocolSettings->SaveConfig();
+}
+
+void SDMXInputInfoSelecter::HandleUniverseIDChanged(uint16 NewValue)
+{
+	if (CurrentUniverseID != NewValue)
+	{
+		CurrentUniverseID = NewValue;
+	}
+}
+
+void SDMXInputInfoSelecter::HandleUniverseIDValueCommitted(uint16 NewValue, ETextCommit::Type CommitType)
+{
+	if (CurrentUniverseID != NewValue)
+	{
+		CurrentUniverseID = NewValue;
+	}
+
+	// Update stored settings
+	UDMXProtocolSettings* ProtocolSettings = GetMutableDefault<UDMXProtocolSettings>();
+	ProtocolSettings->InputConsoleUniverseID = CurrentUniverseID;
+	ProtocolSettings->SaveConfig();
+}
+
+#undef LOCTEXT_NAMESPACE
diff --git a/Engine/Plugins/Experimental/DigitalMultiplex/DMXEngine/Source/DMXEditor/Private/Widgets/SDMXInputInfoSelecter.h b/Engine/Plugins/Experimental/DigitalMultiplex/DMXEngine/Source/DMXEditor/Private/Widgets/SDMXInputInfoSelecter.h
new file mode 100644
--- /dev/null
+++ b/Engine/Plugins/Experimental/DigitalMultiplex/DMXEngine/Source/DMXEditor/Private/Widgets/SDMXInputInfoSelecter.h
@@ -0,0 +1,67 @@
+// Copyright Epic Games, Inc. All Rights Reserved.
+
+#pragma once
+
+#include "CoreMinimal.h"
+#include "Widgets/SCompoundWidget.h"
+#include "Widgets/DeclarativeSyntaxSupport.h"
+#include "Types/SlateEnums.h"
+
+#include "Interfaces/IDMXProtocol.h"
+#include "DMXProtocolTypes.h"
+
+template<typename NumericType>
+class SSpinBox;
+
+/** Widget to configure DMX Input inspector settings */
+class SDMXInputInfoSelecter
+	: public SCompoundWidget
+{
+public:
+	DECLARE_DELEGATE_OneParam(FOnUniverseSelectionChanged, const FName&);
+
+	SLATE_BEGIN_ARGS(SDMXInputInfoSelecter)
+	{}
+	SLATE_EVENT(FOnUniverseSelectionChanged, OnUniverseSelectionChanged)
+	SLATE_END_ARGS()
+
+	/** Constructs the widget */
+	void Construct(const FArguments& InArgs);
+
+	const TSharedRef<SSpinBox<uint16>> GetUniverseField() const { return UniverseIDField.ToSharedRef(); }
+	const void SetProtocol(const FName& Name) { CurrentProtocol = FDMXProtocolName(Name); };
+
+protected:
+	/** Universe ID value computed using Net, Subnet and Universe values */
+	uint16 CurrentUniverseID;
+	/** Final Universe field widget */
+	TSharedPtr<SSpinBox<uint16>> UniverseIDField;
+
+	/** The user-selected protocol */
+	FDMXProtocolName CurrentProtocol;
+
+public:
+	/** Returns a label text with the currently selected directionality */
+	//const FText& GetCurrentDirectionality() const { return CurrentDirectionality; }
+
+	/** Returns the UniverseID value computed from Net, Subnet and Universe values */
+	uint16 GetCurrentUniverseID() const { return CurrentUniverseID; }
+
+	/** Returns the user-selected DMX protocol */
+	FName GetCurrentProtocolName() const { return CurrentProtocol; }
+
+protected:
+	/** Create an option for the Protocol combo box */
+	TSharedRef<SWidget> GenerateProtocolItemWidget(TSharedPtr<FName> InItem);
+	/** Handles new selection from the Directionality combo box */
+	void HandleProtocolChanged(FName SelectedProtocol);
+
+	/** Handles when the user changes the universe value, including while spinning the value */
+	void HandleUniverseIDChanged(uint16 NewValue);
+
+	/** Handles when the user commit the Universe value. Doesn't fire while spinning the value */
+	void HandleUniverseIDValueCommitted(uint16 NewValue, ETextCommit::Type CommitType);
+
+protected:
+	FOnUniverseSelectionChanged OnUniverseSelectionChanged;
+}; 
diff --git a/Engine/Plugins/Experimental/DigitalMultiplex/DMXEngine/Source/DMXEditor/Private/Widgets/SDMXOutputConsole.cpp b/Engine/Plugins/Experimental/DigitalMultiplex/DMXEngine/Source/DMXEditor/Private/Widgets/SDMXOutputConsole.cpp
new file mode 100644
--- /dev/null
+++ b/Engine/Plugins/Experimental/DigitalMultiplex/DMXEngine/Source/DMXEditor/Private/Widgets/SDMXOutputConsole.cpp
@@ -0,0 +1,56 @@
+// Copyright Epic Games, Inc. All Rights Reserved.
+
+#include "Widgets/SDMXOutputConsole.h"
+#include "Widgets/OutputFader/SDMXOutputFaderList.h"
+#include "Widgets/SDMXEntityInspector.h"
+
+#include "DMXEditor.h"
+#include "DMXEditorLog.h"
+#include "Library/DMXEntityFader.h"
+#include "DMXEditorUtils.h"
+
+#include "Widgets/SBoxPanel.h"
+#include "Widgets/Layout/SScrollBox.h"
+
+void SDMXOutputConsole::Construct(const FArguments& InArgs)
+{
+	DMXEditorPtr = InArgs._DMXEditor;
+	check(DMXEditorPtr.Pin().IsValid());
+
+	TSharedPtr<SDMXEntityInspectorFaders> Inspector = SNew(SDMXEntityInspectorFaders)
+		.DMXEditor(DMXEditorPtr);
+
+	// GC can't delete this object
+	OutputConsoleFaderTemplateGuard = TStrongObjectPtr<UDMXEntityFader>(FDMXEditorUtils::CreateFaderTemplate(DMXEditorPtr.Pin()->GetDMXLibrary()));
+	Inspector->ShowDetailsForSingleEntity(OutputConsoleFaderTemplateGuard.Get());
+
+	ChildSlot
+	[
+		SNew(SVerticalBox)
+		+ SVerticalBox::Slot()
+		.AutoHeight()
+		.MaxHeight(350)
+		[
+			SNew(SScrollBox)
+			.Orientation(EOrientation::Orient_Vertical)
+			+ SScrollBox::Slot()
+			[
+				Inspector.ToSharedRef()
+			]
+		]
+		+SVerticalBox::Slot()
+		.VAlign(VAlign_Fill)
+		.HAlign(HAlign_Fill)
+		.AutoHeight()
+		[
+			SNew(SDMXOutputFaderList)
+			.DMXEditor(DMXEditorPtr)
+			.FaderTemplate(OutputConsoleFaderTemplateGuard.Get())
+		]
+	];
+}
+
+SDMXOutputConsole::~SDMXOutputConsole()
+{
+	OutputConsoleFaderTemplateGuard.Reset(); // GC can now delete this object
+}
diff --git a/Engine/Plugins/Experimental/DigitalMultiplex/DMXEngine/Source/DMXEditor/Private/Widgets/SDMXOutputConsole.h b/Engine/Plugins/Experimental/DigitalMultiplex/DMXEngine/Source/DMXEditor/Private/Widgets/SDMXOutputConsole.h
new file mode 100644
--- /dev/null
+++ b/Engine/Plugins/Experimental/DigitalMultiplex/DMXEngine/Source/DMXEditor/Private/Widgets/SDMXOutputConsole.h
@@ -0,0 +1,38 @@
+// Copyright Epic Games, Inc. All Rights Reserved.
+
+#pragma once
+
+#include "CoreMinimal.h"
+#include "Widgets/SCompoundWidget.h"
+#include "Widgets/DeclarativeSyntaxSupport.h"
+#include "UObject/StrongObjectPtr.h"
+
+class FDMXEditor;
+class UDMXEntityFader;
+
+/**
+ * Widget for the Output Console tab, to configure output faders 
+ */
+class SDMXOutputConsole
+	: public SCompoundWidget
+{
+public:
+	SLATE_BEGIN_ARGS(SDMXOutputConsole)
+		: _DMXEditor(nullptr)
+	{}
+		
+	SLATE_ARGUMENT(TWeakPtr<FDMXEditor>, DMXEditor)
+	SLATE_END_ARGS()
+
+	/** Constructs the widget */
+	void Construct(const FArguments& InArgs);
+
+	~SDMXOutputConsole();
+
+private:
+	/** The fader used in the output console to display fader properties */
+	TStrongObjectPtr<UDMXEntityFader> OutputConsoleFaderTemplateGuard;
+
+	/** Weak reference to DMX editor */
+	TWeakPtr<FDMXEditor> DMXEditorPtr;
+};
diff --git a/Engine/Plugins/Experimental/DigitalMultiplex/DMXEngine/Source/DMXEditor/Private/Widgets/SNameListPicker.cpp b/Engine/Plugins/Experimental/DigitalMultiplex/DMXEngine/Source/DMXEditor/Private/Widgets/SNameListPicker.cpp
new file mode 100644
--- /dev/null
+++ b/Engine/Plugins/Experimental/DigitalMultiplex/DMXEngine/Source/DMXEditor/Private/Widgets/SNameListPicker.cpp
@@ -0,0 +1,119 @@
+// Copyright Epic Games, Inc. All Rights Reserved.
+
+#include "Widgets/SNameListPicker.h"
+
+#include "DMXEditorLog.h"
+
+#include "Widgets/Input/SComboBox.h"
+
+#define LOCTEXT_NAMESPACE "SDMXProtocolNamePicker"
+
+void SNameListPicker::Construct(const FArguments& InArgs)
+{
+	ValueAttribute = InArgs._Value;
+	OnValueChangedDelegate = InArgs._OnValueChanged;
+	HasMultipleValuesAttribute = InArgs._HasMultipleValues;
+
+	OptionsSourceAttr = InArgs._OptionsSource;
+	UpdateOptionsSource();
+
+	UpdateOptionsDelegate = InArgs._UpdateOptionsDelegate;
+	if (UpdateOptionsDelegate)
+	{
+		UpdateOptionsHandle = UpdateOptionsDelegate->Add(FSimpleDelegate::CreateSP(this, &SNameListPicker::UpdateOptionsSource));
+	}
+
+	TSharedPtr<FName> InitiallySelected = nullptr;
+	for (TSharedPtr<FName> NameItem : OptionsSource)
+	{
+		if (NameItem)
+		{
+			if (ValueAttribute.Get() == *NameItem)
+			{
+				InitiallySelected = NameItem;
+				break;
+			}
+		}
+	}
+
+	ChildSlot
+	[
+		SAssignNew(PickerComboButton, SComboBox<TSharedPtr<FName>>)
+		.OptionsSource(&OptionsSource)
+		.OnGenerateWidget(this, &SNameListPicker::GenerateNameItemWidget)
+		.OnSelectionChanged(this, &SNameListPicker::HandleSelectionChanged)
+		.InitiallySelectedItem(InitiallySelected)
+		[
+			SNew(STextBlock)
+			.Text(this, &SNameListPicker::GetCurrentNameLabel)
+		]
+	];
+}
+
+void SNameListPicker::UpdateOptionsSource()
+{
+	OptionsSource.Reset();
+	OptionsSource.Reserve(OptionsSourceAttr.Get().Num());
+	for (const FName& Name : OptionsSourceAttr.Get())
+	{
+		OptionsSource.Add(MakeShared<FName>(Name));
+	}
+}
+
+SNameListPicker::~SNameListPicker()
+{
+	if (UpdateOptionsDelegate)
+	{
+		UpdateOptionsDelegate->Remove(UpdateOptionsHandle);
+	}
+}
+
+TSharedRef<SWidget> SNameListPicker::GenerateNameItemWidget(TSharedPtr<FName> InItem)
+{
+	if (!InItem.IsValid())
+	{
+		UE_LOG_DMXEDITOR(Warning, TEXT("InItem for GenerateProtocolItemWidget was null!"));
+		return SNew(STextBlock)
+			.Text(LOCTEXT("NullComboBoxItemLabel", "Null Error"));
+	}
+
+	return SNew(STextBlock)
+		.Text(FText::FromName(*InItem));
+}
+
+void SNameListPicker::HandleSelectionChanged(const TSharedPtr<FName> Item, ESelectInfo::Type SelectInfo)
+{
+	if (!Item.IsValid())
+	{
+		UE_LOG_DMXEDITOR(Error, TEXT("HandleProtocolChanged called with null Item pointer"));
+		return;
+	}
+
+	if (OnValueChangedDelegate.IsBound())
+	{
+		OnValueChangedDelegate.ExecuteIfBound(*Item);
+	}
+	else if (!ValueAttribute.IsBound())
+	{
+		ValueAttribute = *Item;
+	}
+
+	TSharedPtr<SComboButton> PickerComboButtonPin = PickerComboButton.Pin();
+	if (PickerComboButtonPin.IsValid())
+	{
+		PickerComboButtonPin->SetIsOpen(false);
+	}
+}
+
+FText SNameListPicker::GetCurrentNameLabel() const
+{
+	const bool bHasMultipleValues = HasMultipleValuesAttribute.Get();
+	if (bHasMultipleValues)
+	{
+		return LOCTEXT("MultipleValuesText", "<multiple values>");
+	}
+
+	return FText::FromName(ValueAttribute.Get());
+}
+
+#undef LOCTEXT_NAMESPACE
diff --git a/Engine/Plugins/Experimental/DigitalMultiplex/DMXEngine/Source/DMXEditor/Public/Commands/DMXEditorCommands.h b/Engine/Plugins/Experimental/DigitalMultiplex/DMXEngine/Source/DMXEditor/Public/Commands/DMXEditorCommands.h
new file mode 100644
--- /dev/null
+++ b/Engine/Plugins/Experimental/DigitalMultiplex/DMXEngine/Source/DMXEditor/Public/Commands/DMXEditorCommands.h
@@ -0,0 +1,38 @@
+// Copyright Epic Games, Inc. All Rights Reserved.
+
+#pragma once
+
+#include "CoreMinimal.h"
+#include "Framework/Commands/Commands.h"
+#include "EditorStyleSet.h"
+
+class DMXEDITOR_API FDMXEditorCommandsImpl 
+	: public TCommands<FDMXEditorCommandsImpl>
+{
+public:
+	FDMXEditorCommandsImpl();
+
+	virtual ~FDMXEditorCommandsImpl() {}
+
+	//~ Begin TCommands implementation
+	virtual void RegisterCommands() override;
+	//~ End TCommands implementation
+
+	// Go to node documentation
+	TSharedPtr< FUICommandInfo > GoToDocumentation;
+
+	// Create entites
+	TSharedPtr< FUICommandInfo > AddNewEntityController;
+	TSharedPtr< FUICommandInfo > AddNewEntityFixtureType;
+	TSharedPtr< FUICommandInfo > AddNewEntityFixturePatch;
+};
+
+class DMXEDITOR_API FDMXEditorCommands
+{
+public:
+	static void Register();
+
+	static const FDMXEditorCommandsImpl& Get();
+
+	static void Unregister();
+};
diff --git a/Engine/Plugins/Experimental/DigitalMultiplex/DMXEngine/Source/DMXEditor/Public/DMXEditor.h b/Engine/Plugins/Experimental/DigitalMultiplex/DMXEngine/Source/DMXEditor/Public/DMXEditor.h
new file mode 100644
--- /dev/null
+++ b/Engine/Plugins/Experimental/DigitalMultiplex/DMXEngine/Source/DMXEditor/Public/DMXEditor.h
@@ -0,0 +1,149 @@
+// Copyright Epic Games, Inc. All Rights Reserved.
+
+#pragma once
+
+#include "WorkflowOrientedApp/WorkflowCentricApplication.h"
+#include "Misc/NotifyHook.h"
+
+class UDMXLibrary;
+class SDockTab;
+class UFactory;
+class SDMXControllers;
+class SDMXInputConsole;
+class SDMXFixturePatch;
+class SDMXFixtureTypes;
+class SDMXEntityEditor;
+class FDMXEditorToolbar;
+class UDMXEntity;
+
+// Used to enable Entity creator code to inject a base name before the entity creation
+DECLARE_MULTICAST_DELEGATE_TwoParams(FOnGetBaseNameForNewEntity, TSubclassOf<UDMXEntity>, FString&);
+// Used to enable Entity creator code to set values in a newly created entity
+DECLARE_MULTICAST_DELEGATE_OneParam(FOnSetupNewEntity, UDMXEntity*);
+
+class DMXEDITOR_API FDMXEditor
+	: public FWorkflowCentricApplication	// Allow Add ApplicationModes
+	, public FNotifyHook
+{
+public:
+	FDMXEditor();
+
+	virtual ~FDMXEditor();
+
+	//~ Begin IToolkit implementation
+	virtual FName GetToolkitFName() const override;
+	virtual FText GetBaseToolkitName() const override;
+	virtual FString GetWorldCentricTabPrefix() const override;
+	virtual FLinearColor GetWorldCentricTabColorScale() const override;
+	//~ End IToolkit implementationFExtensibilityManager
+
+	/** Adds a new Entity to this DMX Library */
+	void OnAddNewEntity(TSubclassOf<UDMXEntity> InEntityClass);
+
+	/** Activate the editor tab suited to edit Entities of type InEntityClass */
+	bool InvokeEditorTabFromEntityType(TSubclassOf<UDMXEntity> InEntityClass);
+
+	FReply OnAddNewEntity_OnClick(TSubclassOf<UDMXEntity> InEntityClass) { OnAddNewEntity(InEntityClass); return FReply::Handled(); }
+	/** Checks if adding a new Entity is allowed in the current list */
+	bool CanAddNewEntity(TSubclassOf<UDMXEntity> InEntityClass) const;
+	bool NewEntity_IsVisibleForType(TSubclassOf<UDMXEntity> InEntityClass) const;
+
+	/** Utility function to handle all steps required to rename a newly added Entity */
+	void RenameNewlyAddedEntity(UDMXEntity* InEntity, TSubclassOf<UDMXEntity> InEntityClass);
+
+	/** Gets the content widget for the tab that edits Entities from InEntityClass */
+	TSharedPtr<SDMXEntityEditor> GetEditorWidgetForEntityType(TSubclassOf<UDMXEntity> InEntityClass) const;
+
+	/** Switch to the correct tab to select an Entity for editing */
+	void SelectEntityInItsTypeTab(UDMXEntity* InEntity, ESelectInfo::Type InSelectionType = ESelectInfo::Type::Direct);
+	/** Switch to the correct tab for the first Entity's type and select the Entities */
+	void SelectEntitiesInTypeTab(const TArray<UDMXEntity*>& InEntities, ESelectInfo::Type InSelectionType = ESelectInfo::Type::Direct);
+	/** Get the selected entities from the tab that stores the passed in type */
+	TArray<UDMXEntity*> GetSelectedEntitiesFromTypeTab(TSubclassOf<UDMXEntity> InEntityClass) const;
+
+	//~ Getters for the various DMX widgets
+	TSharedRef<SDMXControllers> GetControllersTab() const { return ControllersWidget.ToSharedRef(); }
+	TSharedRef<SDMXFixturePatch> GetFixturePatchTab() const { return FixturePatchWidget.ToSharedRef(); }
+	TSharedRef<SDMXFixtureTypes> GetFixtureTypesTab() const { return FixtureTypesWidget.ToSharedRef(); }
+	TSharedRef<SDMXInputConsole> GetInputConsoleTab() const { return InputConsoleWidget.ToSharedRef(); }
+	TSharedRef<SWidget> GetOutputConsoleTab() const { return OutputConsoleWidget.ToSharedRef(); }
+
+	//~ Getters for event dispatchers
+	/**
+	 * Called before a new Entity creation to set its base name.
+	 * It's highly advisable to only bind to this event right before needing it
+	 * and unbinding right after it's called. To avoid getting calls for every new Entity creation
+	 * (unless that's the desired behavior).
+	 */
+	FOnGetBaseNameForNewEntity& GetOnGetBaseNameForNewEntity() { return OnGetBaseNameForNewEntity; }
+	/**
+	 * Called right after a new Entity is created, to set its values up before it gets selected and renamed.
+	 * It's highly advisable to only bind to this event right before needing it
+	 * and unbinding right after it's called. To avoid getting calls for every new Entity creation
+	 * (unless that's the desired behavior).
+	 */
+	FOnSetupNewEntity& GetOnSetupNewEntity() { return OnSetupNewEntity; }
+
+	/** Edits the specified DMX library */
+	virtual void InitEditor(const EToolkitMode::Type Mode, const TSharedPtr< class IToolkitHost >& InitToolkitHost, UDMXLibrary* DMXLibrary);
+
+	/** Should be called when initializing */
+	void CommonInitialization(UDMXLibrary* DMXLibrary);
+
+	/** Get the DMX library being edited */
+	UDMXLibrary* GetDMXLibrary() const;
+
+	TSharedPtr<FDMXEditorToolbar> GetToolbarBuilder() { return Toolbar; }
+
+	void RegisterToolbarTab(const TSharedRef<class FTabManager>& TabManager);
+	
+protected:
+	/** Creates the widgets that go into the tabs (note: does not create the tabs themselves) **/
+	virtual void CreateDefaultTabContents(UDMXLibrary* DMXLibrary);
+
+	/** Create Default Commands **/
+	virtual void CreateDefaultCommands();
+
+	/** Called during initialization of the DMX editor to register commands and extenders. */
+	virtual void InitalizeExtenders();
+
+	/** Called during initialization of the DMX editor to register any application modes. */
+	virtual void RegisterApplicationModes(UDMXLibrary* DMXLibrary, bool bShouldOpenInDefaultsMode, bool bNewlyCreated = false);
+
+private:
+	UDMXLibrary* GetEditableDMXLibrary() const;
+
+	//~ Generate Editor widgets for tabs
+	TSharedRef<SDMXInputConsole> CreateInputConsoleWidget();
+	TSharedRef<SWidget> CreateOutputConsoleWidget();
+	TSharedRef<SDMXControllers> CreateControllersWidget();
+	TSharedRef<SDMXFixtureTypes> CreateFixtureTypesWidget();
+	TSharedRef<SDMXFixturePatch> CreateFixturePatchWidget();
+
+private:
+	/** The toolbar builder class */
+	TSharedPtr<FDMXEditorToolbar> Toolbar;
+
+	/** The name given to all instances of this type of editor */
+	static const FName ToolkitFName;
+
+	/** UI for the "DMX Monitor" tab */
+	TSharedPtr<SDMXInputConsole> InputConsoleWidget;
+
+	/** UI for the "DMX Output Console" tab */
+	TSharedPtr<SWidget> OutputConsoleWidget;
+
+	// TODO. Should sold object for particular DMX object, like actor components in blueprint editor
+	/** UI for the "DMX Controllers" tab */
+	TSharedPtr<SDMXControllers> ControllersWidget;
+
+	/** UI for the "DMX Fixture Types" tab */
+	TSharedPtr<SDMXFixtureTypes> FixtureTypesWidget;
+
+	/** UI for the "DMX Fixture Patch" tab */
+	TSharedPtr<SDMXFixturePatch> FixturePatchWidget;
+
+	//~ Event dispatchers
+	FOnGetBaseNameForNewEntity OnGetBaseNameForNewEntity;
+	FOnSetupNewEntity OnSetupNewEntity;
+};
diff --git a/Engine/Plugins/Experimental/DigitalMultiplex/DMXEngine/Source/DMXEditor/Public/DMXEditorModule.h b/Engine/Plugins/Experimental/DigitalMultiplex/DMXEngine/Source/DMXEditor/Public/DMXEditorModule.h
new file mode 100644
--- /dev/null
+++ b/Engine/Plugins/Experimental/DigitalMultiplex/DMXEngine/Source/DMXEditor/Public/DMXEditorModule.h
@@ -0,0 +1,110 @@
+// Copyright Epic Games, Inc. All Rights Reserved.
+
+#pragma once
+
+#include "CoreMinimal.h"
+#include "Modules/ModuleInterface.h"
+#include "Toolkits/AssetEditorToolkit.h"
+#include "AssetTypeCategories.h"
+
+class FDMXEditor;
+class UDMXLibrary;
+class IAssetTools;
+class IAssetTypeActions;
+
+class DMXEDITOR_API FDMXEditorModule
+	: public IModuleInterface
+	, public IHasMenuExtensibility			// Extender for adds or removes extenders for menu
+	, public IHasToolBarExtensibility		// Extender for adds or removes extenders for toolbar
+{
+public:
+
+	//~ Begin IModuleInterface implementation
+	virtual void StartupModule() override;
+
+	virtual void ShutdownModule() override;
+
+	//~ End IModuleInterface implementation
+
+	//~ Begin IHasMenuExtensibility implementation
+	virtual TSharedPtr<FExtensibilityManager> GetMenuExtensibilityManager() override { return MenuExtensibilityManager; }
+	//~ End IHasMenuExtensibility implementation
+
+	//~ Begin IHasToolBarExtensibility implementation
+	virtual TSharedPtr<FExtensibilityManager> GetToolBarExtensibilityManager() override { return ToolBarExtensibilityManager; }
+	//~ End IHasToolBarExtensibility implementation
+
+	/** Get the instance of this module. */
+	static FDMXEditorModule& Get();
+
+	/**
+	 * Creates an instance of a DMX editor object.
+	 *
+	 * Note: This function should not be called directly. It should be called from AssetTools handler
+	 *
+	 * @param	Mode					Mode that this editor should operate in
+	 * @param	InitToolkitHost			When Mode is WorldCentric, this is the level editor instance to spawn this editor within
+	 * @param	DMXLibrary				The DMX object to start editing
+	 *
+	 * @return	Interface to the new DMX editor
+	 */
+	TSharedRef<FDMXEditor> CreateEditor( const EToolkitMode::Type Mode, const TSharedPtr< class IToolkitHost >& InitToolkitHost, UDMXLibrary* DMXLibrary );
+
+	static EAssetTypeCategories::Type GetAssetCategory() { return DMXEditorAssetCategory; }
+
+	/**
+	 * Exposes a way for other modules to add in their own DMX editor
+	 * commands (appended to other DMX editor commands, when the editor is
+	 * first opened).
+	 */
+	virtual const TSharedRef<FUICommandList> GetsSharedDMXEditorCommands() const { return SharedDMXEditorCommands.ToSharedRef(); }
+
+public:
+	/** DataTable Editor app identifier string */
+	static const FName DMXEditorAppIdentifier;
+
+	static const FName ModuleName;
+
+private:
+	void RegisterAssetTypeAction(IAssetTools& AssetTools, TSharedRef<IAssetTypeActions> Action);
+
+private:
+	void RegisterPropertyTypeCustomizations();
+	void RegisterObjectCustomizations();
+
+	/**
+	 * Registers a custom class
+	 *
+	 * @param ClassName				The class name to register for property customization
+	 * @param DetailLayoutDelegate	The delegate to call to get the custom detail layout instance
+	 */
+	void RegisterCustomClassLayout(FName ClassName, FOnGetDetailCustomizationInstance DetailLayoutDelegate);
+
+	/**
+	* Registers a custom struct
+	*
+	* @param StructName				The name of the struct to register for property customization
+	* @param StructLayoutDelegate	The delegate to call to get the custom detail layout instance
+	*/
+	void RegisterCustomPropertyTypeLayout(FName PropertyTypeName, FOnGetPropertyTypeCustomizationInstance PropertyTypeLayoutDelegate);
+
+private:
+	//~ Gets the extensibility managers for outside entities to DMX editor's menus and toolbars
+	TSharedPtr<FExtensibilityManager> MenuExtensibilityManager;
+	TSharedPtr<FExtensibilityManager> ToolBarExtensibilityManager;
+
+	/** All created asset type actions.  Cached here so that we can unregister it during shutdown. */
+	TArray< TSharedPtr<IAssetTypeActions> > CreatedAssetTypeActions;
+
+	/**
+	 * A command list that can be passed around and isn't bound to an instance
+	 * of the DMX editor.
+	 */
+	TSharedPtr<FUICommandList> SharedDMXEditorCommands;
+
+	static EAssetTypeCategories::Type DMXEditorAssetCategory;
+
+	/** List of registered class that we must unregister when the module shuts down */
+	TSet< FName > RegisteredClassNames;
+	TSet< FName > RegisteredPropertyTypes;
+};
\ No newline at end of file
diff --git a/Engine/Plugins/Experimental/DigitalMultiplex/DMXEngine/Source/DMXEditor/Public/DMXEditorStyle.h b/Engine/Plugins/Experimental/DigitalMultiplex/DMXEngine/Source/DMXEditor/Public/DMXEditorStyle.h
new file mode 100644
--- /dev/null
+++ b/Engine/Plugins/Experimental/DigitalMultiplex/DMXEngine/Source/DMXEditor/Public/DMXEditorStyle.h
@@ -0,0 +1,35 @@
+// Copyright Epic Games, Inc. All Rights Reserved.
+
+#pragma once
+
+#include "CoreMinimal.h"
+#include "Styling/SlateStyle.h"
+
+/**  DMX editor style, for loading and creation of UI resources */
+class DMXEDITOR_API FDMXEditorStyle
+{
+public:
+
+	/** Initializes this style's singleton and register the Slate Style */
+	static void Initialize();
+
+	/** Unregister the Slate Style and release the Singleton */
+	static void Shutdown();
+
+	/**  Reloads textures used by slate renderer */
+	static void ReloadTextures();
+
+	/**  Returns the singleton ISlateStyle instance for DMX editor style */
+	static const ISlateStyle& Get();
+
+	/**  Returns the DMX editor style name */
+	static FName GetStyleSetName();
+
+private:
+	/** Instantiate and returns a Slate Style Set object */
+	static TSharedRef< class FSlateStyleSet > Create();
+
+private:
+	/** Singleton style instance */
+	static TSharedPtr< class FSlateStyleSet > StyleInstance;
+};
\ No newline at end of file
diff --git a/Engine/Plugins/Experimental/DigitalMultiplex/DMXEngine/Source/DMXEditor/Public/DMXEditorUtils.h b/Engine/Plugins/Experimental/DigitalMultiplex/DMXEngine/Source/DMXEditor/Public/DMXEditorUtils.h
new file mode 100644
--- /dev/null
+++ b/Engine/Plugins/Experimental/DigitalMultiplex/DMXEngine/Source/DMXEditor/Public/DMXEditorUtils.h
@@ -0,0 +1,108 @@
+// Copyright Epic Games, Inc. All Rights Reserved.
+
+#pragma once
+
+#include "CoreMinimal.h"
+#include "Templates/SubclassOf.h"
+#include "Library/DMXEntity.h"
+
+class UDMXLibrary;
+class UDMXEntityFixtureType;
+class UDMXEntityFader;
+
+class DMXEDITOR_API FDMXEditorUtils
+{
+public:
+	/**
+	 * Utility to separate a name from an index at the end.
+	 * @param InString	The string to be separated.
+	 * @param OutName	The string without an index at the end. White spaces and '_' are also removed.
+	 * @param OutIndex	Index that was separated from the name. If there was none, it's zero.
+	 *					Check the return value to know if there was an index on InString.
+	 * @return True if there was an index on InString.
+	 */
+	static bool GetNameAndIndexFromString(const FString& InString, FString& OutName, int32& OutIndex);
+
+	/**
+	 * Generates a unique name given a base one and a list of existing ones, by appending an index to
+	 * existing names. If InBaseName is an empty String, it returns "Default name".
+	 */
+	static FString GenerateUniqueNameFromExisting(const TSet<FString>& InExistingNames, const FString& InBaseName);
+
+	/**
+	 * Creates an unique name for an Entity from a specific type, using the type name as base.
+	 * @param InLibrary		The DMXLibrary object the entity will belong to.
+	 * @param InEntityClass	The class of the Entity, to check the name against others from same type.
+	 * @param InBaseName	Optional base name to use instead of the type name.
+	 * @return Unique name for an Entity amongst others from the same type.
+	 */
+	static FString FindUniqueEntityName(const UDMXLibrary* InLibrary, TSubclassOf<UDMXEntity> InEntityClass, const FString& InBaseName = TEXT(""));
+
+	/**
+	 * Set unique names for Fixture Types' Modes, Functions and Sub Functions when they have just been created.
+	 * We simply rename the Modes/Functions with no name. The user can't set a blank name afterwards, so it's
+	 * a good way to know which are the new ones.
+	 */
+	static void SetNewFixtureFunctionsNames(UDMXEntityFixtureType* InFixtureType);
+
+	/**
+	 * Creates a named Entity from the specified type and add it to the DMXLibrary.
+	 * @param InLibrary			The DMXLibrary object the entity will belong to.
+	 * @param NewEntityName		A name for the new Entity.
+	 * @param NewEntityClass	Type of the new Entity.
+	 * @param OutNewEntity		Returns the created entity.
+	 * @return True if the creation was successful.
+	 */
+	static bool AddEntity(UDMXLibrary* InLibrary, const FString& NewEntityName, TSubclassOf<UDMXEntity> NewEntityClass, UDMXEntity** OutNewEntity = nullptr);
+
+	/**
+	 * Validates an Entity name, also checking for uniqueness among others of the same type.
+	 * @param NewEntityName		The name to validate.
+	 * @param InLibrary			The DMXLibrary object to check for name uniqueness.
+	 * @param InEntityClass		The type to check other Entities' names
+	 * @param OutReason			If false is returned, contains a text with the reason for it.
+	 * @return True if the name would be a valid one.
+	 */
+	static bool ValidateEntityName(const FString& NewEntityName, const UDMXLibrary* InLibrary, UClass* InEntityClass, FText& OutReason);
+
+	/**
+	 * Creates new fader template.
+	 * @param InLibrary			The DMXLibrary object to check for name uniqueness.
+	 * @return New transient fader template object.
+	 */
+	static UDMXEntityFader* CreateFaderTemplate(const UDMXLibrary* InLibrary);
+
+	/**  Renames an Entity */
+	static void RenameEntity(UDMXLibrary* InLibrary, UDMXEntity* InEntity, const FString& NewName);
+
+	/**  Checks if the Entity is being referenced by other objects. */
+	static bool IsEntityUsed(const UDMXLibrary* InLibrary, const UDMXEntity* InEntity);
+
+	/**  Removes the entities and fixes references to it. */
+	static void RemoveEntities(UDMXLibrary* InLibrary, const TArray<UDMXEntity*>&& InEntities);
+
+	/**  Copies Entities to the operating system's clipboard. */
+	static void CopyEntities(const TArray<UDMXEntity*>&& EntitiesToCopy);
+
+	/**  Determines whether the current contents of the clipboard contain paste-able DMX Entity information */
+	static bool CanPasteEntities();
+
+	/**
+	 * Gets the copied DMX Entities from the clipboard without attempting to paste/apply them in any way
+	 * @param OutNewObjectMap			Contains the name->instance object mapping of the copied DMX Entities
+	 */
+	static void GetEntitiesFromClipboard(TArray<UDMXEntity*>& OutNewObjects);
+
+	/**
+	 * Compares the property values of two Fixture Types, including properties in arrays,
+	 * and returns true if they are almost all the same.
+	 * Name, ID and Parent Library are ignored.
+	 */
+	static bool AreFixtureTypesIdentical(const UDMXEntityFixtureType* A, const UDMXEntityFixtureType* B);
+
+	/**  Returns the Entity class type name (e.g: Fixture Type for UDMXEntityFixtureType) in singular or plural */
+	static FText GetEntityTypeNameText(TSubclassOf<UDMXEntity> EntityClass, bool bPlural = false);
+
+	// can't instantiate this class
+	FDMXEditorUtils() = delete;
+};
diff --git a/Engine/Plugins/Experimental/DigitalMultiplex/DMXEngine/Source/DMXRuntime/Private/DMXSubsystem.cpp b/Engine/Plugins/Experimental/DigitalMultiplex/DMXEngine/Source/DMXRuntime/Private/DMXSubsystem.cpp
new file mode 100644
--- /dev/null
+++ b/Engine/Plugins/Experimental/DigitalMultiplex/DMXEngine/Source/DMXRuntime/Private/DMXSubsystem.cpp
@@ -0,0 +1,473 @@
+// Copyright Epic Games, Inc. All Rights Reserved.
+
+#include "DMXSubsystem.h"
+#include "Interfaces/IDMXProtocol.h"
+#include "Interfaces/IDMXProtocolUniverse.h"
+#include "DMXProtocolTypes.h"
+#include "Library/DMXLibrary.h"
+#include "Library/DMXEntityReference.h"
+#include "Library/DMXEntity.h"
+#include "Library/DMXEntityController.h"
+#include "Library/DMXEntityFixtureType.h"
+#include "Library/DMXEntityFixturePatch.h"
+
+#include "EngineUtils.h"
+#include "UObject/UObjectIterator.h"
+
+void UDMXSubsystem::SendDMX(FDMXProtocolName SelectedProtocol, UDMXEntityFixturePatch* FixturePatch, TMap<FName, int32> FunctionMap, EDMXSendResult& OutResult)
+{
+	OutResult = EDMXSendResult::ErrorSetBuffer;
+
+	if (FixturePatch != nullptr)
+	{
+		IDMXFragmentMap DMXFragmentMap;
+		for (const TPair<FName, int32>& Elem : FunctionMap)
+		{
+			const int32& ActiveMode = FixturePatch->ActiveMode;
+			const FDMXFixtureMode& RelevantMode = FixturePatch->ParentFixtureTypeTemplate->Modes[ActiveMode];
+			for (const FDMXFixtureFunction& Function : RelevantMode.Functions)
+			{
+				const FName FunctionName(*Function.FunctionName);
+				if (FunctionName == Elem.Key && Function.Channel > 0)
+				{
+					const uint8 ChannelsToAdd = UDMXEntityFixtureType::NumChannelsToOccupy(Function.DataType);
+					const int32& Channel = Function.Channel + (FixturePatch->GetStartingChannel() - 1);
+					const uint32 FinalValue = UDMXEntityFixtureType::ClampValueToDataType(Function.DataType, Elem.Value);
+
+					for (uint8 ChannelIt = 0; ChannelIt < ChannelsToAdd; ChannelIt++)
+					{
+						const uint8 BytesOffset = (ChannelIt) * 8;
+						const uint8 RetVal = FinalValue >> BytesOffset & 0xff;
+
+						DMXFragmentMap.Add(Channel + ChannelIt, RetVal);
+					}
+
+				}
+			}
+		}
+
+		if (SelectedProtocol)
+		{
+			const int32& Universe = FixturePatch->UniverseID;
+			const TArray<UDMXEntityController*>&& RelevantControllers = FixturePatch->GetRelevantControllers();
+			TArray<EDMXSendResult> Results;
+			TSet<uint32> UniversesUsed;
+
+			Results.Reserve(RelevantControllers.Num());
+			UniversesUsed.Reserve(RelevantControllers.Num());
+
+			// Send using the Remote Offset from each Controller with this Fixture's Universe in its range
+			for (const UDMXEntityController* Controller : RelevantControllers)
+			{
+				const uint32 RemoteUniverse = Universe + Controller->RemoteOffset;
+				if (!UniversesUsed.Contains(RemoteUniverse))
+				{
+					Results.Add(SelectedProtocol.GetProtocol()->SendDMXFragment(Universe + Controller->RemoteOffset, DMXFragmentMap));
+					UniversesUsed.Add(RemoteUniverse); // Avoid setting values in the same Universe more than once
+				}
+			}
+
+			for (const EDMXSendResult& Result : Results)
+			{
+				if (Result != EDMXSendResult::Success)
+				{
+					OutResult = Result;
+					return;
+				}
+			}
+
+			OutResult = EDMXSendResult::Success;
+		}
+	}
+}
+
+void UDMXSubsystem::SendDMXRaw(FDMXProtocolName SelectedProtocol, int32 UniverseIndex, TMap<int32, uint8> ChannelValuesMap, EDMXSendResult& OutResult)
+{
+	OutResult = EDMXSendResult::ErrorSetBuffer;
+
+	if (SelectedProtocol)
+	{
+		IDMXFragmentMap DMXFragmentMap;
+		for (auto& Elem : ChannelValuesMap)
+		{
+			if (Elem.Key != 0)
+			{
+				DMXFragmentMap.Add(Elem.Key, Elem.Value);
+			}
+		}
+		OutResult = SelectedProtocol.GetProtocol()->SendDMXFragment(UniverseIndex, DMXFragmentMap);
+	}
+}
+
+void UDMXSubsystem::GetAllFixturesOfType(const UDMXLibrary* DMXLibrary, const FName& FixtureType, TArray<UDMXEntityFixturePatch*>& OutResult)
+{
+	OutResult.Reset();
+
+	if (DMXLibrary != nullptr)
+	{
+		DMXLibrary->ForEachEntityOfType<UDMXEntityFixturePatch>([&](UDMXEntityFixturePatch* Fixture)
+			{
+				if (Fixture->ParentFixtureTypeTemplate != nullptr && Fixture->ParentFixtureTypeTemplate->GetDisplayName() == FixtureType.ToString())
+				{
+					OutResult.Add(Fixture);
+				}
+			});
+	}
+}
+
+void UDMXSubsystem::GetAllFixturesOfCategory(const UDMXLibrary* DMXLibrary, FDMXFixtureCategory Category, TArray<UDMXEntityFixturePatch*>& OutResult)
+{
+	OutResult.Reset();
+
+	if (DMXLibrary != nullptr)
+	{
+		DMXLibrary->ForEachEntityOfType<UDMXEntityFixturePatch>([&](UDMXEntityFixturePatch* Fixture)
+			{
+				if (Fixture->ParentFixtureTypeTemplate != nullptr && Fixture->ParentFixtureTypeTemplate->DMXCategory == Category)
+				{
+					OutResult.Add(Fixture);
+				}
+			});
+	}
+}
+
+void UDMXSubsystem::GetAllFixturesInUniverse(const UDMXLibrary* DMXLibrary, int32 UniverseId, TArray<UDMXEntityFixturePatch*>& OutResult)
+{
+	OutResult.Reset();
+
+	if (DMXLibrary != nullptr)
+	{
+		DMXLibrary->ForEachEntityOfType<UDMXEntityFixturePatch>([&](UDMXEntityFixturePatch* Fixture)
+			{
+				if (Fixture->UniverseID == UniverseId)
+				{
+					OutResult.Add(Fixture);
+				}
+			});
+	}
+}
+
+void UDMXSubsystem::GetAllUniversesInController(const UDMXLibrary* DMXLibrary, FString ControllerName, TArray<int32>& OutResult)
+{
+	OutResult.Reset();
+
+	if (DMXLibrary != nullptr)
+	{
+		const UDMXEntityController* Controller = Cast<UDMXEntityController>(DMXLibrary->FindEntity(ControllerName));
+		if (Controller != nullptr)
+		{
+			OutResult.Reserve(Controller->Universes.Num());
+			const int32& RemoteOffset = Controller->RemoteOffset;
+
+			// Get All Universes
+			for (const FDMXUniverse& Universe : Controller->Universes)
+			{
+				// Remove remote offset to get local Universe IDs
+				OutResult.Add(Universe.UniverseNumber - RemoteOffset);
+			}
+		}
+	}
+}
+
+void UDMXSubsystem::GetRawBuffer(FDMXProtocolName SelectedProtocol, int32 UniverseIndex, TArray<uint8>& DMXBuffer)
+{
+	DMXBuffer.Reset();
+	if (SelectedProtocol)
+	{
+		TSharedPtr<IDMXProtocolUniverse, ESPMode::ThreadSafe> ProtocolUniverse = SelectedProtocol.GetProtocol()->GetUniverseById(UniverseIndex);
+		if (ProtocolUniverse.IsValid())
+		{
+			TSharedPtr<FDMXBuffer> Buffer = ProtocolUniverse.Get()->GetInputDMXBuffer();
+			if (Buffer.IsValid())
+			{
+				DMXBuffer = Buffer->GetDMXData();
+			}
+		}
+	}
+}
+
+void UDMXSubsystem::GetFixtureFunctions(const UDMXEntityFixturePatch* InFixturePatch, const TArray<uint8>& DMXBuffer, TMap<FName, int32>& OutResult)
+{
+	OutResult.Reset();
+
+	if (InFixturePatch != nullptr)
+	{
+		const int32& ActiveMode = InFixturePatch->ActiveMode;
+		const int32 StartingAddress = InFixturePatch->GetStartingChannel();
+
+		const UDMXEntityFixtureType* FixtureType = InFixturePatch->ParentFixtureTypeTemplate;
+		if (FixtureType != nullptr)
+		{
+			const FDMXFixtureMode& CurrentMode = FixtureType->Modes[ActiveMode];
+
+			for (const FDMXFixtureFunction& Function : CurrentMode.Functions)
+			{
+				const int32 IndexValue = (Function.Channel - 1) + (StartingAddress - 1);
+
+				const uint8 ChannelsToAdd = UDMXEntityFixtureType::NumChannelsToOccupy(Function.DataType);
+
+				TArray<uint8> Bytes;
+				for (uint8 ChannelIt = 0; ChannelIt < ChannelsToAdd; ChannelIt++)
+				{
+					const uint8& ByteValue = DMXBuffer[IndexValue + ChannelIt];
+					Bytes.Add(ByteValue);
+				}
+
+				uint32 IntVal = 0;
+				for (uint8 ByteIndex = 0; ByteIndex < Bytes.Num() && ByteIndex < 4; ++ByteIndex)
+				{
+					IntVal += Bytes[ByteIndex] << (ByteIndex * 8);
+				}
+				OutResult.Add(FName(*Function.FunctionName), IntVal);
+			}
+		}
+	}
+}
+
+UDMXEntityFixturePatch* UDMXSubsystem::GetFixturePatch(FDMXEntityFixturePatchRef InFixturePatch)
+{
+	return InFixturePatch.GetFixturePatch();
+}
+
+bool UDMXSubsystem::GetFunctionsMap(UDMXEntityFixturePatch* InFixturePatch, const FDMXProtocolName& SelectedProtocol, TMap<FName, int32>& OutFunctionsMap)
+{
+	OutFunctionsMap.Empty();
+
+	if (InFixturePatch == nullptr)
+	{
+		UE_LOG(LogTemp, Warning, TEXT("No FixturePatch"));
+
+		return false;
+	}
+
+	TSharedPtr<IDMXProtocol> Protocol = SelectedProtocol.GetProtocol();
+	if (!Protocol.IsValid())
+	{
+		UE_LOG(LogTemp, Warning, TEXT("Protocol Not Valid"));
+
+		return false;
+	}
+
+	TSharedPtr<IDMXProtocolUniverse, ESPMode::ThreadSafe> Universe = Protocol->GetUniverseById(InFixturePatch->UniverseID);
+	if (!Universe.IsValid())
+	{
+		UE_LOG(LogTemp, Warning, TEXT("Universe Not Valid"));
+
+		return false;
+	}
+
+	TSharedPtr<FDMXBuffer> InputDMXBuffer = Universe->GetInputDMXBuffer();
+	if (!InputDMXBuffer.IsValid())
+	{
+		UE_LOG(LogTemp, Warning, TEXT("InputDMXBuffer Not Valid"));
+
+		return false;
+	}
+
+	if (InFixturePatch == nullptr)
+	{
+		UE_LOG(LogTemp, Warning, TEXT("InFixturePatch == nullptr"));
+		return false;
+	}
+
+	UDMXEntityFixtureType* TypeTemplate = InFixturePatch->ParentFixtureTypeTemplate;
+	if (TypeTemplate == nullptr)
+	{
+		UE_LOG(LogTemp, Warning, TEXT("TypeTemplate == nullptr"));
+
+		return false;
+	}
+
+	if (TypeTemplate->Modes.Num() < InFixturePatch->ActiveMode + 1)
+	{
+		UE_LOG(LogTemp, Warning, TEXT("Wrong ActiveMode %d, Num of modes %d"), InFixturePatch->ActiveMode, TypeTemplate->Modes.Num());
+
+		return false;
+	}
+
+	// Read the correct amount of channels for each function and store each in the correct byte of an int32
+	const FDMXFixtureMode& Mode = TypeTemplate->Modes[InFixturePatch->ActiveMode];
+	const TArray<uint8>& DMXData = InputDMXBuffer->GetDMXData();
+	
+	// Unsigned to make sure there won't be any system that changes how the bytes sum is done with int32
+	uint32 ChannelValue = 0;
+
+	for (const FDMXFixtureFunction& Function : Mode.Functions)
+	{
+		if (Function.Channel > DMX_MAX_ADDRESS)
+		{
+			UE_LOG(LogTemp, Warning, TEXT("Channel %d is more then %d"), Function.Channel, DMX_MAX_ADDRESS);
+
+			return false;
+		}
+
+		if (UDMXEntityFixtureType::GetFunctionLastChannel(Function) > Mode.ChannelSpan)
+		{
+			// We reached the functions outside the valid channels for this mode
+			break;
+		}
+
+		// Always reset all bytes to zero on our int, in case the new ones don't use all bytes.
+		// Otherwise, the bytes not used now could have values from the previous function's channels.
+		ChannelValue = 0;
+
+		const int32 NumChannelsToRead = UDMXEntityFixtureType::NumChannelsToOccupy(Function.DataType);
+		for (int32 ChannelIndex = 0; ChannelIndex < NumChannelsToRead && ChannelIndex + Function.Channel < DMXData.Num() && ChannelIndex < 4; ++ChannelIndex)
+		{
+			ChannelValue += DMXData[Function.Channel + ChannelIndex] << ChannelIndex * 8;
+		}
+
+		OutFunctionsMap.Add(*Function.FunctionName, ChannelValue);
+	}
+
+	return true;
+}
+
+int32 UDMXSubsystem::GetFunctionsValue(const FName& InName, const TMap<FName, int32>& InFunctionsMap)
+{
+	const int32* Result = InFunctionsMap.Find(InName);
+	if (Result != nullptr)
+	{
+		return *Result;
+	}
+
+	return 0;
+}
+
+TArray<UDMXEntityFixturePatch*> UDMXSubsystem::GetAllFixturesWithTag(const UDMXLibrary* DMXLibrary, FName CustomTag)
+{
+	TArray<UDMXEntityFixturePatch*> FoundPatches;
+
+	if (DMXLibrary != nullptr)
+	{
+		DMXLibrary->ForEachEntityOfType<UDMXEntityFixturePatch>([&](UDMXEntityFixturePatch* Patch)
+			{
+				if (Patch->CustomTags.Contains(CustomTag))
+				{
+					FoundPatches.Add(Patch);
+				}
+			});
+	}
+
+	return FoundPatches;
+}
+
+TArray<UDMXEntityFixturePatch*> UDMXSubsystem::GetAllFixturesInLibrary(const UDMXLibrary* DMXLibrary)
+{
+	TArray<UDMXEntityFixturePatch*> FoundPatches;
+
+	if (DMXLibrary != nullptr)
+	{
+		DMXLibrary->ForEachEntityOfType<UDMXEntityFixturePatch>([&](UDMXEntityFixturePatch* Patch)
+			{
+				FoundPatches.Add(Patch);
+			});
+	}
+
+	return FoundPatches;
+}
+
+template<class TEntityType>
+TEntityType* GetDMXEntityByName(const UDMXLibrary* DMXLibrary, const FString& Name)
+{
+	if (DMXLibrary != nullptr)
+	{
+		TEntityType* FoundEntity = nullptr;
+		DMXLibrary->ForEachEntityOfTypeWithBreak<TEntityType>([&](TEntityType* Entity)
+			{
+				if (Entity->Name.Equals(Name))
+				{
+					FoundEntity = Entity;
+					return false;
+				}
+				return true;
+			});
+
+		return FoundEntity;
+	}
+
+	return nullptr;
+}
+
+UDMXEntityFixturePatch* UDMXSubsystem::GetFixtureByName(const UDMXLibrary* DMXLibrary, const FString& Name)
+{
+	return GetDMXEntityByName<UDMXEntityFixturePatch>(DMXLibrary, Name);
+}
+
+TArray<UDMXEntityFixtureType*> UDMXSubsystem::GetAllFixtureTypesInLibrary(const UDMXLibrary* DMXLibrary)
+{
+	TArray<UDMXEntityFixtureType*> FoundTypes;
+
+	if (DMXLibrary != nullptr)
+	{
+		DMXLibrary->ForEachEntityOfType<UDMXEntityFixtureType>([&](UDMXEntityFixtureType* Type)
+		{
+			FoundTypes.Add(Type);
+		});
+	}
+
+	return FoundTypes;
+}
+
+UDMXEntityFixtureType* UDMXSubsystem::GetFixtureTypeByName(const UDMXLibrary* DMXLibrary, const FString& Name)
+{
+	return GetDMXEntityByName<UDMXEntityFixtureType>(DMXLibrary, Name);
+}
+
+TArray<UDMXEntityController*> UDMXSubsystem::GetAllControllersInLibrary(const UDMXLibrary* DMXLibrary)
+{
+	TArray<UDMXEntityController*> FoundControllers;
+
+	if (DMXLibrary != nullptr)
+	{
+		DMXLibrary->ForEachEntityOfType<UDMXEntityController>([&](UDMXEntityController* Controller)
+		{
+			FoundControllers.Add(Controller);
+		});
+	}
+
+	return FoundControllers;
+}
+
+UDMXEntityController* UDMXSubsystem::GetControllerByName(const UDMXLibrary* DMXLibrary, const FString& Name)
+{
+	return GetDMXEntityByName<UDMXEntityController>(DMXLibrary, Name);
+}
+
+TArray<UDMXLibrary*> UDMXSubsystem::GetAllDMXLibraries()
+{
+	TArray<UDMXLibrary*> LibrariesFound;
+
+	for (TObjectIterator<UDMXLibrary> LibraryItr; LibraryItr; ++LibraryItr)
+	{
+		LibrariesFound.Add(*LibraryItr);
+	}
+
+	return LibrariesFound;
+}
+
+int32 UDMXSubsystem::BytesToInt(const TArray<uint8>& Bytes)
+{
+	uint32 IntVal = 0;
+	for (uint8 ByteIndex = 0; ByteIndex < Bytes.Num() && ByteIndex < 4; ++ByteIndex)
+	{
+		IntVal += Bytes[ByteIndex] << (ByteIndex * 8);
+	}
+	return IntVal;
+}
+
+void UDMXSubsystem::Initialize(FSubsystemCollectionBase& Collection)
+{
+	for (const FName& ProtocolName : IDMXProtocol::GetProtocolNames())
+	{
+		if (TSharedPtr<IDMXProtocol> Protocol = IDMXProtocol::Get(ProtocolName))
+		{
+			Protocol->GetOnUniverseInputUpdate().AddUObject(this, &UDMXSubsystem::BufferReceivedBroadcast);
+		}
+	}
+}
+
+void UDMXSubsystem::BufferReceivedBroadcast(FName Protocol, uint16 UniverseID, const TArray<uint8>& Values)
+{
+	OnProtocolReceived.Broadcast(FDMXProtocolName(Protocol), UniverseID, Values);
+}
diff --git a/Engine/Plugins/Experimental/DigitalMultiplex/DMXEngine/Source/DMXRuntime/Private/Game/DMXComponent.cpp b/Engine/Plugins/Experimental/DigitalMultiplex/DMXEngine/Source/DMXRuntime/Private/Game/DMXComponent.cpp
new file mode 100644
--- /dev/null
+++ b/Engine/Plugins/Experimental/DigitalMultiplex/DMXEngine/Source/DMXRuntime/Private/Game/DMXComponent.cpp
@@ -0,0 +1,42 @@
+// Copyright Epic Games, Inc. All Rights Reserved.
+
+#include "Game/DMXComponent.h"
+#include "Library/DMXEntityFixturePatch.h"
+
+UDMXEntityFixturePatch* UDMXComponent::GetFixturePatch() const
+{
+	return FixturePatch.GetFixturePatch();
+}
+
+void UDMXComponent::SetFixturePatch(UDMXEntityFixturePatch* InFixturePatch)
+{
+	FixturePatch.SetEntity(InFixturePatch);
+}
+
+// Sets default values for this component's properties
+UDMXComponent::UDMXComponent()
+{
+	// Set this component to be initialized when the game starts, and to be ticked every frame.  You can turn these features
+	// off to improve performance if you don't need them.
+	PrimaryComponentTick.bCanEverTick = true;
+
+	// ...
+}
+
+// Called when the game starts
+void UDMXComponent::BeginPlay()
+{
+	Super::BeginPlay();
+
+	// ...
+}
+
+
+// Called every frame
+void UDMXComponent::TickComponent(float DeltaTime, ELevelTick TickType, FActorComponentTickFunction* ThisTickFunction)
+{
+	Super::TickComponent(DeltaTime, TickType, ThisTickFunction);
+
+	// ...
+}
+
diff --git a/Engine/Plugins/Experimental/DigitalMultiplex/DMXEngine/Source/DMXRuntime/Private/Library/DMXEntity.cpp b/Engine/Plugins/Experimental/DigitalMultiplex/DMXEngine/Source/DMXRuntime/Private/Library/DMXEntity.cpp
new file mode 100644
--- /dev/null
+++ b/Engine/Plugins/Experimental/DigitalMultiplex/DMXEngine/Source/DMXRuntime/Private/Library/DMXEntity.cpp
@@ -0,0 +1,119 @@
+// Copyright Epic Games, Inc. All Rights Reserved.
+
+#include "Library/DMXEntity.h"
+#include "Library/DMXLibrary.h"
+#include "Interfaces/IDMXProtocol.h"
+
+#include "Dom/JsonObject.h"
+
+UDMXEntity::UDMXEntity()
+	: ParentLibrary(nullptr)
+{
+	FPlatformMisc::CreateGuid(Id);
+}
+
+FString UDMXEntity::GetDisplayName() const
+{
+	return Name;
+}
+
+void UDMXEntity::SetName(const FString& InNewName)
+{
+	// TODO check name for validity
+	Name = InNewName;
+}
+
+void UDMXEntity::SetParentLibrary(UDMXLibrary* InParent)
+{
+	ParentLibrary = InParent;
+}
+
+void UDMXEntity::RefreshID()
+{
+	FPlatformMisc::CreateGuid(Id);
+}
+
+void UDMXEntity::ReplicateID(UDMXEntity* Other)
+{
+	Id = Other->Id;
+}
+
+#if WITH_EDITOR
+void UDMXEntity::PostDuplicate(EDuplicateMode::Type DuplicateMode)
+{
+	Super::PostDuplicate(DuplicateMode);
+
+	RefreshID();
+}
+#endif // WITH_EDITOR
+
+UDMXEntityUniverseManaged::UDMXEntityUniverseManaged()
+{
+	if (!HasAnyFlags(RF_ClassDefaultObject | RF_ArchetypeObject))
+	{
+		DeviceProtocol = FDMXProtocolName(IDMXProtocol::GetFirstProtocolName());
+	}
+}
+
+void UDMXEntityUniverseManaged::PostLoad()
+{
+	Super::PostLoad();
+	UpdateProtocolUniverses();
+}
+
+#if WITH_EDITOR
+
+void UDMXEntityUniverseManaged::PostEditChangeChainProperty(FPropertyChangedChainEvent& PropertyChangedEvent)
+{
+	static const FName NAME_Universes = GET_MEMBER_NAME_CHECKED(UDMXEntityUniverseManaged, Universes);
+
+	Super::PostEditChangeChainProperty(PropertyChangedEvent);
+
+	if (PropertyChangedEvent.GetPropertyName() == GET_MEMBER_NAME_CHECKED(FDMXUniverse, UniverseNumber) 
+		|| PropertyChangedEvent.GetPropertyName() == NAME_Universes
+		|| PropertyChangedEvent.GetPropertyName() == GET_MEMBER_NAME_CHECKED(UDMXEntityUniverseManaged, DeviceProtocol))
+	{
+		// Keep the Universe ID values within the valid range for the current protocol
+		if (DeviceProtocol.IsValid())
+		{
+			const TSharedPtr<IDMXProtocol> Protocol = DeviceProtocol.GetProtocol();
+			const int32 ChangedUniverseIndex = PropertyChangedEvent.GetArrayIndex(NAME_Universes.ToString());
+
+			if (ChangedUniverseIndex != INDEX_NONE)
+			{
+				FDMXUniverse& ChangedUniverse = Universes[ChangedUniverseIndex];
+				ChangedUniverse.UniverseNumber = FMath::Clamp(ChangedUniverse.UniverseNumber, (uint32)Protocol->GetMinUniverseID(), (uint32)Protocol->GetMaxUniverses());
+			}
+			else
+			{
+				const uint32 MinUniverseID = Protocol->GetMinUniverseID();
+				const uint32 MaxUniverseID = Protocol->GetMaxUniverses();
+
+				for (FDMXUniverse& Universe : Universes)
+				{
+					Universe.UniverseNumber = FMath::Clamp(Universe.UniverseNumber, MinUniverseID, MaxUniverseID);
+				}
+			}
+		}
+
+		// New Universes will have their directionality set to Output
+		if (PropertyChangedEvent.ChangeType == EPropertyChangeType::ArrayAdd
+			&& PropertyChangedEvent.GetPropertyName() == GET_MEMBER_NAME_CHECKED(UDMXEntityUniverseManaged, Universes))
+		{
+			Universes.Last(0).DMXProtocolDirectionality = EDMXProtocolDirectionality::EOutput;
+		}
+	}
+
+	UpdateProtocolUniverses();
+}
+
+#endif // WITH_EDITOR
+
+void UDMXEntityUniverseManaged::UpdateProtocolUniverses() const
+{
+	if (DeviceProtocol.IsValid())
+	{
+		TSharedPtr<IDMXProtocol> DMXProtocol = DeviceProtocol;
+		DMXProtocol->CollectUniverses(Universes);
+	}
+}
