Fix static analysis warning  for DMX.

diff --git a/Engine/Plugins/Experimental/DigitalMultiplex/DMXEngine/Source/DMXBlueprintGraph/Private/Widgets/SDMXEntityReferenceGraphPin.cpp b/Engine/Plugins/Experimental/DigitalMultiplex/DMXEngine/Source/DMXBlueprintGraph/Private/Widgets/SDMXEntityReferenceGraphPin.cpp
--- a/Engine/Plugins/Experimental/DigitalMultiplex/DMXEngine/Source/DMXBlueprintGraph/Private/Widgets/SDMXEntityReferenceGraphPin.cpp
+++ b/Engine/Plugins/Experimental/DigitalMultiplex/DMXEngine/Source/DMXBlueprintGraph/Private/Widgets/SDMXEntityReferenceGraphPin.cpp
@@ -1,6 +1,6 @@
 // Copyright Epic Games, Inc. All Rights Reserved.
 
-#include "Widgets//SDMXEntityReferenceGraphPin.h"
+#include "Widgets/SDMXEntityReferenceGraphPin.h"
 
 #include "Library/DMXEntityReference.h"
 #include "Library/DMXEntity.h"
diff --git a/Engine/Plugins/Experimental/DigitalMultiplex/DMXEngine/Source/DMXBlueprintGraph/Private/Widgets/SNullGraphPin.cpp b/Engine/Plugins/Experimental/DigitalMultiplex/DMXEngine/Source/DMXBlueprintGraph/Private/Widgets/SNullGraphPin.cpp
--- a/Engine/Plugins/Experimental/DigitalMultiplex/DMXEngine/Source/DMXBlueprintGraph/Private/Widgets/SNullGraphPin.cpp
+++ b/Engine/Plugins/Experimental/DigitalMultiplex/DMXEngine/Source/DMXBlueprintGraph/Private/Widgets/SNullGraphPin.cpp
@@ -1,6 +1,6 @@
 // Copyright Epic Games, Inc. All Rights Reserved.
 
-#include "Widgets//SNullGraphPin.h"
+#include "Widgets/SNullGraphPin.h"
 
 #include "Widgets/SNullWidget.h"
 
