Merging DMX fixes from Dev-VirtualProduction, up to CL 11494311.

diff --git a/Engine/Plugins/Experimental/DigitalMultiplex/DMXEngine/DMXEngine.uplugin b/Engine/Plugins/Experimental/DigitalMultiplex/DMXEngine/DMXEngine.uplugin
--- a/Engine/Plugins/Experimental/DigitalMultiplex/DMXEngine/DMXEngine.uplugin
+++ b/Engine/Plugins/Experimental/DigitalMultiplex/DMXEngine/DMXEngine.uplugin
@@ -16,38 +16,17 @@
     {
       "Name": "DMXRuntime",
       "Type": "Runtime",
-      "LoadingPhase": "Default",
-      "WhitelistPlatforms": 
-      [
-        "Win64",
-        "Win32",
-        "Mac",
-        "Linux"
-      ]
+      "LoadingPhase": "Default"
     },
     {
       "Name": "DMXEditor",
       "Type": "Editor",
-      "LoadingPhase": "Default",
-      "WhitelistPlatforms": 
-      [
-        "Win64",
-        "Win32",
-        "Mac",
-        "Linux"
-      ]
+      "LoadingPhase": "Default"
     },
 	{
       "Name": "DMXBlueprintGraph",
       "Type": "UncookedOnly",
-      "LoadingPhase": "Default",
-      "WhitelistPlatforms": 
-      [
-        "Win64",
-        "Win32",
-        "Mac",
-        "Linux"
-      ]
+      "LoadingPhase": "Default"
     }
   ],
   "Plugins": [
diff --git a/Engine/Plugins/Experimental/DigitalMultiplex/DMXEngine/Source/DMXBlueprintGraph/Private/DMXGraphPanelPinFactory.cpp b/Engine/Plugins/Experimental/DigitalMultiplex/DMXEngine/Source/DMXBlueprintGraph/Private/DMXGraphPanelPinFactory.cpp
--- a/Engine/Plugins/Experimental/DigitalMultiplex/DMXEngine/Source/DMXBlueprintGraph/Private/DMXGraphPanelPinFactory.cpp
+++ b/Engine/Plugins/Experimental/DigitalMultiplex/DMXEngine/Source/DMXBlueprintGraph/Private/DMXGraphPanelPinFactory.cpp
@@ -2,7 +2,6 @@
 
 #include "DMXGraphPanelPinFactory.h"
 #include "K2Node.h"
-#include "K2Node_GetAllFixturesOfType.h"
 #include "EdGraphSchema_K2.h"
 #include "Widgets/DeclarativeSyntaxSupport.h"
 #include "SGraphPin.h"
@@ -18,36 +17,7 @@
 
 TSharedPtr<class SGraphPin> FDMXGraphPanelPinFactory::CreatePin(class UEdGraphPin* InPin) const
 {
-	if (InPin->PinType.PinCategory == UEdGraphSchema_K2::PC_Name)
-	{
-		UObject* Outer = InPin->GetOuter();
-
-		if (UK2Node_DMXBase* K2Node_DMXBase = Cast<UK2Node_DMXBase>(Outer))
-		{
-			const UEdGraphPin* DMXLibraryPin = K2Node_DMXBase->GetClassPin();
-			if (DMXLibraryPin != nullptr)
-			{
-				if (DMXLibraryPin->LinkedTo.Num() == 0)
-				{
-					if (UDMXLibrary* DMXLibrary = Cast<UDMXLibrary>(DMXLibraryPin->DefaultObject))
-					{
-						if (UK2Node_GetAllFixturesOfType * K2Node_GetAllFixturesOfType = Cast<UK2Node_GetAllFixturesOfType>(Outer))
-						{
-							TArray<TSharedPtr<FName>> FixtureTypeList;
-							DMXLibrary->ForEachEntityOfType<UDMXEntityFixtureType>([&FixtureTypeList](UDMXEntityFixtureType* Fixture)
-								{
-									TSharedPtr<FName> NewFixtureType = MakeShared<FName>(FName(*Fixture->GetDisplayName()));
-									FixtureTypeList.Add(NewFixtureType);
-								});
-
-							return SNew(SGraphPinNameList, InPin, FixtureTypeList);
-						}
-					}
-				}
-			}
-		}
-	}
-	else if (InPin->PinType.PinCategory == UEdGraphSchema_K2::PC_Struct)
+	if (InPin->PinType.PinCategory == UEdGraphSchema_K2::PC_Struct)
 	{
 		if (UScriptStruct* PinStructType = Cast<UScriptStruct>(InPin->PinType.PinSubCategoryObject.Get()))
 		{
diff --git a/Engine/Plugins/Experimental/DigitalMultiplex/DMXEngine/Source/DMXBlueprintGraph/Private/K2Node_GetAllFixturesOfType.cpp b/Engine/Plugins/Experimental/DigitalMultiplex/DMXEngine/Source/DMXBlueprintGraph/Private/K2Node_GetAllFixturesOfType.cpp
deleted file mode 100644
diff --git a/Engine/Plugins/Experimental/DigitalMultiplex/DMXEngine/Source/DMXBlueprintGraph/Private/K2Node_GetDMXActiveModeFunctionValues.cpp b/Engine/Plugins/Experimental/DigitalMultiplex/DMXEngine/Source/DMXBlueprintGraph/Private/K2Node_GetDMXActiveModeFunctionValues.cpp
--- a/Engine/Plugins/Experimental/DigitalMultiplex/DMXEngine/Source/DMXBlueprintGraph/Private/K2Node_GetDMXActiveModeFunctionValues.cpp
+++ b/Engine/Plugins/Experimental/DigitalMultiplex/DMXEngine/Source/DMXBlueprintGraph/Private/K2Node_GetDMXActiveModeFunctionValues.cpp
@@ -15,7 +15,6 @@
 #include "ScopedTransaction.h"
 #include "EdGraphUtilities.h"
 #include "K2Node_CallFunction.h"
-#include "Engine/Engine.h"
 
 #define LOCTEXT_NAMESPACE "UK2Node_GetDMXActiveModeFunctionValues"
 
@@ -108,13 +107,24 @@ void UK2Node_GetDMXActiveModeFunctionValues::ExpandNode(FKismetCompilerContext&
 	Super::ExpandNode(CompilerContext, SourceGraph);
 	const UEdGraphSchema_K2* Schema = CompilerContext.GetSchema();
 
-	UDMXSubsystem* Subsystem = GEngine->GetEngineSubsystem<UDMXSubsystem>();
+	// First node to execute. GetDMXSubsystem
+	FName GetDMXSubsystemFunctionName = GET_FUNCTION_NAME_CHECKED(UDMXSubsystem, GetDMXSubsystem_Callable);
+	UK2Node_CallFunction* DMXSubsystemNode = CompilerContext.SpawnIntermediateNode<UK2Node_CallFunction>(this, SourceGraph);
+	DMXSubsystemNode->FunctionReference.SetExternalMember(GetDMXSubsystemFunctionName, UDMXSubsystem::StaticClass());
+	DMXSubsystemNode->AllocateDefaultPins();
 
-	UEdGraphPin* SelfPin = CreatePin(EGPD_Input, UEdGraphSchema_K2::PC_Object, UDMXSubsystem::StaticClass(), UEdGraphSchema_K2::PN_Self);
-	SelfPin->DefaultObject = Subsystem;
+	UEdGraphPin* DMXSubsystemExecPin = DMXSubsystemNode->GetExecPin();
+	UEdGraphPin* DMXSubsystemThenPin = DMXSubsystemNode->GetThenPin();
+	UEdGraphPin* DMXSubsystemResult = DMXSubsystemNode->GetReturnValuePin();
 
-	// Call DMX function map. Always first and always call it
-	// Static C++ function
+	// Hook up inputs
+	CompilerContext.MovePinLinksToIntermediate(*GetExecPin(), *DMXSubsystemExecPin);
+
+	// Hook up outputs
+	UEdGraphPin* LastThenPin = DMXSubsystemThenPin;
+	Schema->TryCreateConnection(LastThenPin, DMXSubsystemThenPin);
+
+	// Second node to execute. GetFunctionsMap
 	static const FName FunctionName = GET_FUNCTION_NAME_CHECKED(UDMXSubsystem, GetFunctionsMap);
 	UFunction* GetFunctionsMapPointer = FindField<UFunction>(UDMXSubsystem::StaticClass(), FunctionName);
 	check(GetFunctionsMapPointer);
@@ -133,16 +143,17 @@ void UK2Node_GetDMXActiveModeFunctionValues::ExpandNode(FKismetCompilerContext&
 	UEdGraphPin* GetFunctionsMapNodeThenPin = GetFunctionsMapNode->GetThenPin();
 
 	// Hook up inputs
-	CompilerContext.MovePinLinksToIntermediate(*SelfPin, *GetFunctionsMapNodeSelfPin);
-	CompilerContext.MovePinLinksToIntermediate(*GetExecPin(), *GetFunctionsMapNodeExecPin);
+	Schema->TryCreateConnection(GetFunctionsMapNodeSelfPin, DMXSubsystemResult);
 	CompilerContext.MovePinLinksToIntermediate(*GetInputDMXFixturePatchPin(), *GetFunctionsMapNodeInFixturePatchPin);
 	CompilerContext.MovePinLinksToIntermediate(*GetInputDMXProtocolPin(), *GetFunctionsMapNodeSelectedProtocolPin);
 
 	// Hook up outputs
 	CompilerContext.MovePinLinksToIntermediate(*GetOutputFunctionsMapPin(), *GetFunctionsMapNodeOutFunctionsMapPin);
 	CompilerContext.MovePinLinksToIntermediate(*GetOutputIsSuccessPin(), *GetFunctionsMapNodeOutIsSuccessPin);
+	Schema->TryCreateConnection(LastThenPin, GetFunctionsMapNodeExecPin);
+	LastThenPin = GetFunctionsMapNodeThenPin;
 
-	// Call dmx active mode pin functions
+	// Loop GetFunctionsValueName nodes to execute. 
 	if (UserDefinedPins.Num() > 0)
 	{
 		TArray<UEdGraphPin*> IntPairs;
@@ -166,7 +177,6 @@ void UK2Node_GetDMXActiveModeFunctionValues::ExpandNode(FKismetCompilerContext&
 
 		check(IntPairs.Num() == NamePairs.Num());
 
-		UEdGraphPin* LastThenPin = GetFunctionsMapNodeThenPin;
 		bool bResult = false;
 		for (int32 PairIndex = 0; PairIndex < NamePairs.Num(); ++PairIndex)
 		{
@@ -186,7 +196,7 @@ void UK2Node_GetDMXActiveModeFunctionValues::ExpandNode(FKismetCompilerContext&
 			UEdGraphPin* GetFunctionsValueNodeThenPin = GetFunctionsValueNode->GetThenPin();
 
 			// Input
-			CompilerContext.MovePinLinksToIntermediate(*SelfPin, *GetFunctionsValueSelfPin);
+			Schema->TryCreateConnection(GetFunctionsValueSelfPin, DMXSubsystemResult);
 			CompilerContext.MovePinLinksToIntermediate(*NamePairs[PairIndex], *GetFunctionsValueNodeOutInNamePin);
 
 			// Output
@@ -202,7 +212,7 @@ void UK2Node_GetDMXActiveModeFunctionValues::ExpandNode(FKismetCompilerContext&
 	}
 	else
 	{
-		CompilerContext.MovePinLinksToIntermediate(*GetThenPin(), *GetFunctionsMapNodeThenPin);
+		CompilerContext.MovePinLinksToIntermediate(*GetThenPin(), *LastThenPin);
 	}
 }
 
diff --git a/Engine/Plugins/Experimental/DigitalMultiplex/DMXEngine/Source/DMXBlueprintGraph/Private/K2Node_GetDMXFixturePatch.cpp b/Engine/Plugins/Experimental/DigitalMultiplex/DMXEngine/Source/DMXBlueprintGraph/Private/K2Node_GetDMXFixturePatch.cpp
--- a/Engine/Plugins/Experimental/DigitalMultiplex/DMXEngine/Source/DMXBlueprintGraph/Private/K2Node_GetDMXFixturePatch.cpp
+++ b/Engine/Plugins/Experimental/DigitalMultiplex/DMXEngine/Source/DMXBlueprintGraph/Private/K2Node_GetDMXFixturePatch.cpp
@@ -12,7 +12,6 @@
 #include "BlueprintNodeSpawner.h"
 #include "K2Node_CallFunction.h"
 #include "Kismet2/BlueprintEditorUtils.h"
-#include "Engine/Engine.h"
 
 #define LOCTEXT_NAMESPACE "UK2Node_GetDMXFixturePatch"
 
@@ -28,6 +27,8 @@ void UK2Node_GetDMXFixturePatch::AllocateDefaultPins()
 	K2Schema->ConstructBasicPinTooltip(*InputDMXFixturePatchPin, LOCTEXT("InputDMXFixtureTypePin", "Get the fixture patch reference."), InputDMXFixturePatchPin->PinToolTip);
 	InputDMXFixturePatchPin->bNotConnectable = true;
 
+
+
 	// Output pins
 	UEdGraphPin* OutputDMXFixturePatchPin = CreatePin(EGPD_Output, UEdGraphSchema_K2::PC_Object, UDMXEntityFixturePatch::StaticClass(), OutputDMXFixturePatchPinName);
 	K2Schema->ConstructBasicPinTooltip(*OutputDMXFixturePatchPin, LOCTEXT("OutputDMXFixturePatch", "Fixture patch."), OutputDMXFixturePatchPin->PinToolTip);
@@ -67,22 +68,27 @@ void UK2Node_GetDMXFixturePatch::ExpandNode(FKismetCompilerContext& CompilerCont
 
 	const UEdGraphSchema_K2* K2Schema = CompilerContext.GetSchema();
 
-	UDMXSubsystem* Subsystem = GEngine->GetEngineSubsystem<UDMXSubsystem>();
+	// First node to execute. GetDMXSubsystem
+	FName GetDMXSubsystemFunctionName = GET_FUNCTION_NAME_CHECKED(UDMXSubsystem, GetDMXSubsystem_Pure);
+	const UFunction* GetDMXSubsystemFunction = UDMXSubsystem::StaticClass()->FindFunctionByName(GetDMXSubsystemFunctionName);
+	check(nullptr != GetDMXSubsystemFunction);
+	UK2Node_CallFunction* DMXSubsystemNode = CompilerContext.SpawnIntermediateNode<UK2Node_CallFunction>(this, SourceGraph);
+	DMXSubsystemNode->SetFromFunction(GetDMXSubsystemFunction);
+	DMXSubsystemNode->AllocateDefaultPins();
 
-	UEdGraphPin* SelfPin = CreatePin(EGPD_Input, UEdGraphSchema_K2::PC_Object, UDMXSubsystem::StaticClass(), UEdGraphSchema_K2::PN_Self);
-	SelfPin->DefaultObject = Subsystem;
+	UEdGraphPin* DMXSubsystemResult = DMXSubsystemNode->GetReturnValuePin();
 
-	// Function to call
-	const FName FunctionName = GET_FUNCTION_NAME_CHECKED(UDMXSubsystem, GetFixturePatch);
-	const UFunction* Function = UDMXSubsystem::StaticClass()->FindFunctionByName(FunctionName);
-	check(nullptr != Function);
+	// Second node to execute. GetFixturePatch
+	const FName GetFixturePatchName = GET_FUNCTION_NAME_CHECKED(UDMXSubsystem, GetFixturePatch);
+	const UFunction* GetFixturePatchFunction = UDMXSubsystem::StaticClass()->FindFunctionByName(GetFixturePatchName);
+	check(nullptr != GetFixturePatchFunction);
 
 	// Spawn call function node
 	UK2Node_CallFunction* SendDataGetFunction = CompilerContext.SpawnIntermediateNode<UK2Node_CallFunction>(this, SourceGraph);
-	SendDataGetFunction->FunctionReference.SetExternalMember(FunctionName, UK2Node_CallFunction::StaticClass());
+	SendDataGetFunction->FunctionReference.SetExternalMember(GetFixturePatchName, UK2Node_CallFunction::StaticClass());
 
 	// Set function node pins
-	SendDataGetFunction->SetFromFunction(Function);
+	SendDataGetFunction->SetFromFunction(GetFixturePatchFunction);
 	SendDataGetFunction->AllocateDefaultPins();
 
 	// Hook up function node inputs
@@ -93,7 +99,7 @@ void UK2Node_GetDMXFixturePatch::ExpandNode(FKismetCompilerContext& CompilerCont
 	const FString&& FixturePatchStr = GetFixturePatchValueAsString();
 
 	// Hook up input
-	CompilerContext.MovePinLinksToIntermediate(*SelfPin, *FunctionSelfPin);
+	K2Schema->TryCreateConnection(FunctionSelfPin, DMXSubsystemResult);
 	K2Schema->TrySetDefaultValue(*FunctionInFixturePatchPin, FixturePatchStr);
 	check(FunctionInFixturePatchPin->GetDefaultAsString().Equals(FixturePatchStr));
 
diff --git a/Engine/Plugins/Experimental/DigitalMultiplex/DMXEngine/Source/DMXBlueprintGraph/Public/K2Node_GetAllFixturesOfType.h b/Engine/Plugins/Experimental/DigitalMultiplex/DMXEngine/Source/DMXBlueprintGraph/Public/K2Node_GetAllFixturesOfType.h
deleted file mode 100644
diff --git a/Engine/Plugins/Experimental/DigitalMultiplex/DMXEngine/Source/DMXBlueprintGraph/Public/K2Node_GetDMXActiveModeFunctionValues.h b/Engine/Plugins/Experimental/DigitalMultiplex/DMXEngine/Source/DMXBlueprintGraph/Public/K2Node_GetDMXActiveModeFunctionValues.h
--- a/Engine/Plugins/Experimental/DigitalMultiplex/DMXEngine/Source/DMXBlueprintGraph/Public/K2Node_GetDMXActiveModeFunctionValues.h
+++ b/Engine/Plugins/Experimental/DigitalMultiplex/DMXEngine/Source/DMXBlueprintGraph/Public/K2Node_GetDMXActiveModeFunctionValues.h
@@ -17,8 +17,8 @@ class UDMXEntityFixturePatch;
 struct FDMXFixtureMode;
 struct FDMXFixtureFunction;
 
-UCLASS(BlueprintType)
-class UK2Node_GetDMXActiveModeFunctionValues
+UCLASS()
+class DMXBLUEPRINTGRAPH_API UK2Node_GetDMXActiveModeFunctionValues
 	: public UK2Node_EditablePinBase
 {
 	GENERATED_BODY()
diff --git a/Engine/Plugins/Experimental/DigitalMultiplex/DMXEngine/Source/DMXBlueprintGraph/Public/K2Node_GetDMXFixturePatch.h b/Engine/Plugins/Experimental/DigitalMultiplex/DMXEngine/Source/DMXBlueprintGraph/Public/K2Node_GetDMXFixturePatch.h
--- a/Engine/Plugins/Experimental/DigitalMultiplex/DMXEngine/Source/DMXBlueprintGraph/Public/K2Node_GetDMXFixturePatch.h
+++ b/Engine/Plugins/Experimental/DigitalMultiplex/DMXEngine/Source/DMXBlueprintGraph/Public/K2Node_GetDMXFixturePatch.h
@@ -13,8 +13,8 @@ class FNodeHandlingFunctor;
 class FKismetCompilerContext;
 class UDMXEntityFixturePatch;
 
-UCLASS(MinimalApi)
-class UK2Node_GetDMXFixturePatch
+UCLASS()
+class DMXBLUEPRINTGRAPH_API UK2Node_GetDMXFixturePatch
 	: public UK2Node
 {
 	GENERATED_BODY()
diff --git a/Engine/Plugins/Experimental/DigitalMultiplex/DMXEngine/Source/DMXEditor/Private/Customizations/DMXEditorPropertyEditorCustomization.cpp b/Engine/Plugins/Experimental/DigitalMultiplex/DMXEngine/Source/DMXEditor/Private/Customizations/DMXEditorPropertyEditorCustomization.cpp
--- a/Engine/Plugins/Experimental/DigitalMultiplex/DMXEngine/Source/DMXEditor/Private/Customizations/DMXEditorPropertyEditorCustomization.cpp
+++ b/Engine/Plugins/Experimental/DigitalMultiplex/DMXEngine/Source/DMXEditor/Private/Customizations/DMXEditorPropertyEditorCustomization.cpp
@@ -49,10 +49,93 @@ static void CollectChildPropertiesRecursive(TSharedPtr<IPropertyHandle> Node, TA
 
 void FDMXCustomization::CustomizeDetails(IDetailLayoutBuilder& DetailLayout)
 {
-	// keep Display Name as first property
+	// Customize the Name input field to check for repeated/invalid names
+	NamePropertyHandle = DetailLayout.GetProperty(GET_MEMBER_NAME_CHECKED(UDMXEntity, Name), UDMXEntity::StaticClass());
+	check(NamePropertyHandle->IsValidHandle());
+
+	// Don't allow editing the Name if multiple Entities are selected
+	FString Name;
+	const bool bCanEditName = NamePropertyHandle->GetValue(Name) == FPropertyAccess::Success;
+
+	DetailLayout.EditDefaultProperty(NamePropertyHandle)->CustomWidget()
+		.NameContent()
+		[
+			NamePropertyHandle->CreatePropertyNameWidget()
+		]
+		.ValueContent()
+		.MaxDesiredWidth(250.0f)
+		[
+			SAssignNew(NameEditableTextBox, SEditableTextBox)
+			.Text(this, &FDMXCustomization::OnGetEntityName)
+			.ToolTipText(NamePropertyHandle->GetToolTipText())
+			.OnTextChanged(this, &FDMXCustomization::OnEntityNameChanged)
+			.OnTextCommitted(this, &FDMXCustomization::OnEntityNameCommitted)
+			.Font(IDetailLayoutBuilder::GetDetailFont())
+			.IsEnabled(bCanEditName)
+		];
+
+	// Keep Display Name as first property
 	DetailLayout.EditCategory("Entity Properties", FText::GetEmpty(), ECategoryPriority::Important);
 }
 
+FText FDMXCustomization::OnGetEntityName() const
+{
+	check(NamePropertyHandle && NamePropertyHandle->IsValidHandle());
+
+	FString Name;
+	if (NamePropertyHandle->GetValue(Name) == FPropertyAccess::Success)
+	{
+		return FText::FromString(Name);
+	}
+
+	return LOCTEXT("EntityName_MultipleValues", "Multiple Values");
+}
+
+void FDMXCustomization::OnEntityNameChanged(const FText& InNewText)
+{
+	check(NameEditableTextBox.IsValid() && NamePropertyHandle.IsValid() && NamePropertyHandle->IsValidHandle());
+
+	FString CurrentName;
+	if (NamePropertyHandle->GetValue(CurrentName) != FPropertyAccess::Success)
+	{
+		return;
+	}
+
+	const FString& NewName = InNewText.ToString();
+	if (CurrentName.Equals(NewName))
+	{
+		NameEditableTextBox->SetError(FText::GetEmpty());
+		return;
+	}
+
+	check(DMXEditorPtr.IsValid());
+	TArray<UObject*> SelectedEntities;
+	NamePropertyHandle->GetOuterObjects(SelectedEntities);
+	check(SelectedEntities.Num() > 0);
+
+	FText OutErrorMessage;
+	FDMXEditorUtils::ValidateEntityName(
+		NewName,
+		DMXEditorPtr.Pin()->GetDMXLibrary(),
+		SelectedEntities[0]->GetClass(),
+		OutErrorMessage
+	);
+
+	NameEditableTextBox->SetError(OutErrorMessage);
+}
+
+void FDMXCustomization::OnEntityNameCommitted(const FText& InNewText, ETextCommit::Type InCommitType)
+{
+	check(NameEditableTextBox.IsValid() && NamePropertyHandle.IsValid() && NamePropertyHandle->IsValidHandle());
+
+	if (InCommitType != ETextCommit::OnCleared && !NameEditableTextBox->HasError())
+	{
+		NamePropertyHandle->SetValue(InNewText.ToString());
+	}
+
+	NameEditableTextBox->SetError(FText::GetEmpty());
+}
+
 void FDMXControllersDetails::CustomizeDetails(IDetailLayoutBuilder& DetailLayout)
 {
 	FDMXCustomization::CustomizeDetails(DetailLayout);
@@ -132,8 +215,6 @@ void FDMXFixtureTypeFunctionsDetails::CustomizeChildren(TSharedRef<IPropertyHand
 				.ValueContent()
 				.MaxDesiredWidth(250.0f)
 				[
-					// The lambdas below are a workaround to avoid ambiguous TSharedFromThis template resolution
-					// that could be both IPropertyTypeCustomization and IDetailCustomization.
 					SAssignNew(NameEditableTextBox, SEditableTextBox)
 					.Text(this, &FDMXFixtureTypeFunctionsDetails::OnGetFunctionName)
 					.ToolTipText(ToolTip)
@@ -249,10 +330,13 @@ void FDMXFixtureFunctionDetails::CustomizeChildren(TSharedRef<IPropertyHandle> I
 	StructPropertyHandle = InStructPropertyHandle;
 
 	DataTypeHandle = InStructPropertyHandle->GetChildHandle(GET_MEMBER_NAME_CHECKED(FDMXFixtureFunction, DataType));
-	check(DataTypeHandle);
+	check(DataTypeHandle && DataTypeHandle->IsValidHandle());
 
 	DefaultValueHandle = InStructPropertyHandle->GetChildHandle(GET_MEMBER_NAME_CHECKED(FDMXFixtureFunction, DefaultValue));
-	check(DefaultValueHandle);
+	check(DefaultValueHandle && DefaultValueHandle->IsValidHandle());
+
+	UseLSBHandle = InStructPropertyHandle->GetChildHandle(GET_MEMBER_NAME_CHECKED(FDMXFixtureFunction, bUseLSBMode));
+	check(UseLSBHandle && UseLSBHandle->IsValidHandle());
 
 	// Create fields for the individual channels (each byte) that will be displayed depending on the selected DataType
 	AddChannelInputFields(InStructBuilder);
@@ -366,9 +450,27 @@ TOptional<uint8> FDMXFixtureFunctionDetails::GetChannelValue(uint8 Channel) cons
 {
 	TOptional<uint8> RetVal;
 	int64 Value;
-	if (DefaultValueHandle->GetValue(Value) == FPropertyAccess::Success)
+	bool bUseLSBValue;
+	if (DefaultValueHandle->GetValue(Value) == FPropertyAccess::Success
+		&& UseLSBHandle->GetValue(bUseLSBValue) == FPropertyAccess::Success)
 	{
-		const uint8 BytesOffset = (Channel - 1) * 8;
+		uint8 BytesOffset = 0;
+
+		if (!bUseLSBValue)
+		{
+			void* DataTypePtr = nullptr;
+			if (DataTypeHandle->GetValueData(DataTypePtr) == FPropertyAccess::Success)
+			{
+				const EDMXFixtureSignalFormat& DataType = *(EDMXFixtureSignalFormat*)DataTypePtr;
+				const uint8 NumChannels = UDMXEntityFixtureType::NumChannelsToOccupy(DataType);
+				BytesOffset = (NumChannels - Channel) * 8;
+			}
+		}
+		else
+		{
+			BytesOffset = (Channel - 1) * 8;
+		}
+
 		RetVal = Value >> BytesOffset & 0xff;
 	}
 
@@ -397,10 +499,27 @@ EVisibility FDMXFixtureFunctionDetails::GetChannelInputVisibility(uint8 Channel)
 void FDMXFixtureFunctionDetails::HandleChannelValueChanged(uint8 NewValue, uint8 Channel)
 {
 	int64 DefaultValue;
-	if (DefaultValueHandle->GetValue(DefaultValue) == FPropertyAccess::Success)
+	bool bUseLSBValue;
+	if (DefaultValueHandle->GetValue(DefaultValue) == FPropertyAccess::Success
+		&& UseLSBHandle->GetValue(bUseLSBValue) == FPropertyAccess::Success)
 	{
 		uint8* ValueBytes = reinterpret_cast<uint8*>(&DefaultValue);
-		ValueBytes[Channel - 1] = NewValue;
+
+		if (!bUseLSBValue)
+		{
+			void* DataTypePtr = nullptr;
+			if (DataTypeHandle->GetValueData(DataTypePtr) == FPropertyAccess::Success)
+			{
+				const EDMXFixtureSignalFormat& DataType = *(EDMXFixtureSignalFormat*)DataTypePtr;
+				const uint8 NumChannels = UDMXEntityFixtureType::NumChannelsToOccupy(DataType);
+				ValueBytes[NumChannels - Channel] = NewValue;
+			}
+		}
+		else
+		{
+			ValueBytes[Channel - 1] = NewValue;
+		}
+
 		DefaultValueHandle->SetValue(DefaultValue, EPropertyValueSetFlags::InteractiveChange);
 	}
 }
@@ -410,6 +529,7 @@ void FDMXFixtureFunctionDetails::HandleChannelValueCommitted(uint8 NewValue, ETe
 	int64 DefaultValue;
 	if (DefaultValueHandle->GetValue(DefaultValue) == FPropertyAccess::Success)
 	{
+		// We need to set the value without the InteractiveChange flag to register the transaction
 		DefaultValueHandle->SetValue(DefaultValue);
 	}
 }
diff --git a/Engine/Plugins/Experimental/DigitalMultiplex/DMXEngine/Source/DMXEditor/Private/Customizations/DMXEditorPropertyEditorCustomization.h b/Engine/Plugins/Experimental/DigitalMultiplex/DMXEngine/Source/DMXEditor/Private/Customizations/DMXEditorPropertyEditorCustomization.h
--- a/Engine/Plugins/Experimental/DigitalMultiplex/DMXEngine/Source/DMXEditor/Private/Customizations/DMXEditorPropertyEditorCustomization.h
+++ b/Engine/Plugins/Experimental/DigitalMultiplex/DMXEngine/Source/DMXEditor/Private/Customizations/DMXEditorPropertyEditorCustomization.h
@@ -34,9 +34,19 @@ public:
 	/** IDetailCustomization interface */
 	virtual void CustomizeDetails(IDetailLayoutBuilder & DetailLayout) override;
 
+protected:
+	FText OnGetEntityName() const;
+	void OnEntityNameChanged(const FText& InNewText);
+	void OnEntityNameCommitted(const FText& InNewText, ETextCommit::Type InCommitType);
+
 protected:
 	/** Weak reference to the DMX editor */
 	TWeakPtr<FDMXEditor> DMXEditorPtr;
+
+	/** Custom Name text field to validate input name */
+	TSharedPtr<SEditableTextBox> NameEditableTextBox;
+	/** Handle to the Name property for getting and setting it */
+	TSharedPtr<IPropertyHandle> NamePropertyHandle;
 };
 
 /** Details customization for Controllers */
@@ -146,6 +156,7 @@ private:
 	TSharedPtr<IPropertyHandle> StructPropertyHandle;
 	TSharedPtr<IPropertyHandle> DataTypeHandle;
 	TSharedPtr<IPropertyHandle> DefaultValueHandle;
+	TSharedPtr<IPropertyHandle> UseLSBHandle;
 };
 
 /**
diff --git a/Engine/Plugins/Experimental/DigitalMultiplex/DMXEngine/Source/DMXEditor/Private/DMXEditorUtils.cpp b/Engine/Plugins/Experimental/DigitalMultiplex/DMXEngine/Source/DMXEditor/Private/DMXEditorUtils.cpp
--- a/Engine/Plugins/Experimental/DigitalMultiplex/DMXEngine/Source/DMXEditor/Private/DMXEditorUtils.cpp
+++ b/Engine/Plugins/Experimental/DigitalMultiplex/DMXEngine/Source/DMXEditor/Private/DMXEditorUtils.cpp
@@ -281,6 +281,12 @@ bool FDMXEditorUtils::ValidateEntityName(const FString& NewEntityName, const UDM
 		return false;
 	}
 
+	if (NewEntityName.TrimStartAndEnd().IsEmpty())
+	{
+		OutReason = LOCTEXT("NameEmpty", "The name can't be blank!");
+		return false;
+	}
+
 	for (const TCHAR& Character : DMX_INVALID_NAME_CHARACTERS)
 	{
 		if (NewEntityName.Contains(&Character))
@@ -392,6 +398,7 @@ void FDMXEditorUtils::RemoveEntities(UDMXLibrary* InLibrary, const TArray<UDMXEn
 			}
 
 			InLibrary->Modify();
+			EntityToDelete->Modify(); // Take a snapshot of the entity before setting its ParentLibrary to null
 			InLibrary->RemoveEntity(EntityToDelete);
 		}
 	}
@@ -518,7 +525,7 @@ bool FDMXEditorUtils::AreFixtureTypesIdentical(const UDMXEntityFixtureType* A, c
 		return false;
 	}
 
-	// Compare each UProperty in the projects
+	// Compare each UProperty in the Fixtures
 	const UStruct* Struct = UDMXEntityFixtureType::StaticClass();
 	TPropertyValueIterator<const FProperty> ItA(Struct, A);
 	TPropertyValueIterator<const FProperty> ItB(Struct, B);
@@ -528,16 +535,29 @@ bool FDMXEditorUtils::AreFixtureTypesIdentical(const UDMXEntityFixtureType* A, c
 
 	for (; ItA && ItB; ++ItA, ++ItB)
 	{
-		const FProperty* Property = ItA->Key;
+		const FProperty* PropertyA = ItA->Key;
+		const FProperty* PropertyB = ItB->Key;
+
+		if (PropertyA == nullptr || PropertyB == nullptr)
+		{
+			return false;
+		}
+
+		// Properties must be in the exact same order on both Fixtures. Otherwise, it means we have
+		// different properties being compared due to differences in array sizes.
+		if (!PropertyA->SameType(PropertyB))
+		{
+			return false;
+		}
 
 		// Name and Id don't have to be identical
-		if (Property->GetFName() == GET_MEMBER_NAME_CHECKED(UDMXEntity, Name)
-			|| Property->GetFName() == NAME_ParentLibrary) // Can't GET_MEMBER_NAME... with private properties
+		if (PropertyA->GetFName() == GET_MEMBER_NAME_CHECKED(UDMXEntity, Name)
+			|| PropertyA->GetFName() == NAME_ParentLibrary) // Can't GET_MEMBER_NAME... with private properties
 		{
 			continue;
 		}
 
-		if (Property->GetFName() == NAME_Id)
+		if (PropertyA->GetFName() == NAME_Id)
 		{
 			// Skip all properties from GUID struct
 			for (int32 PropertyCount = 0; PropertyCount < 4; ++PropertyCount)
@@ -551,12 +571,19 @@ bool FDMXEditorUtils::AreFixtureTypesIdentical(const UDMXEntityFixtureType* A, c
 		const void* ValueA = ItA->Value;
 		const void* ValueB = ItB->Value;
 
-		if (!Property->Identical(ValueA, ValueB))
+		if (!PropertyA->Identical(ValueA, ValueB))
 		{
 			return false;
 		}
 	}
 
+	// If one of the Property Iterators is still valid, one of the Fixtures had
+	// less properties due to an array size difference, which means the Fixtures are different.
+	if (ItA || ItB)
+	{
+		return false;
+	}
+
 	return true;
 }
 
diff --git a/Engine/Plugins/Experimental/DigitalMultiplex/DMXEngine/Source/DMXEditor/Private/Tests/DMXEditor.spec.cpp b/Engine/Plugins/Experimental/DigitalMultiplex/DMXEngine/Source/DMXEditor/Private/Tests/DMXEditor.spec.cpp
deleted file mode 100644
diff --git a/Engine/Plugins/Experimental/DigitalMultiplex/DMXEngine/Source/DMXEditor/Private/Tests/DMXEditorTests.cpp b/Engine/Plugins/Experimental/DigitalMultiplex/DMXEngine/Source/DMXEditor/Private/Tests/DMXEditorTests.cpp
new file mode 100644
--- /dev/null
+++ b/Engine/Plugins/Experimental/DigitalMultiplex/DMXEngine/Source/DMXEditor/Private/Tests/DMXEditorTests.cpp
@@ -0,0 +1,668 @@
+// Copyright Epic Games, Inc. All Rights Reserved.
+
+#include "Misc/AutomationTest.h"
+
+#include "Factories/DMXEditorFactoryNew.h"
+#include "Library/DMXLibrary.h"
+#include "Interfaces/IDMXProtocol.h"
+#include "Interfaces/IDMXProtocolUniverse.h"
+#include "DMXProtocolArtNetModule.h"
+#include "DMXProtocolSACNModule.h"
+#include "Library/DMXLibrary.h"
+#include "Library/DMXEntityFader.h"
+#include "Library/DMXEntityController.h"
+#include "DMXEditor.h"
+#include "DMXProtocolTypes.h"
+#include "DMXEditorUtils.h"
+
+#include "Widgets/OutputFader/SDMXOutputFaderList.h"
+#include "Widgets/OutputFader/SDMXFader.h"
+#include "Widgets/SDMXInputConsole.h"
+#include "Widgets/SDMXInputInfoSelecter.h"
+#include "Widgets/SDMXInputInfo.h"
+#include "Widgets/Common/SSpinBoxVertical.h"
+
+#include "Widgets/Input/SComboBox.h"
+#include "Widgets/Input/SSpinBox.h"
+
+#include "Utils.h"
+#include "Dom/JsonObject.h"
+#include "Engine/Engine.h"
+#include "AssetData.h"
+#include "Editor.h"
+#include "ObjectEditorUtils.h"
+
+#include "AssetRegistryModule.h"
+#include "AssetRegistryState.h"
+#include "AssetData.h"
+
+#if WITH_DEV_AUTOMATION_TESTS
+
+struct FDMXEditorHelper
+{
+	static void ResetUniverses()
+	{
+		const TMap<FName, IDMXProtocolPtr>& Protocols = IDMXProtocol::GetProtocols();
+		for (TMap<FName, IDMXProtocolPtr>::TConstIterator ProtocolIt = Protocols.CreateConstIterator(); ProtocolIt; ++ProtocolIt)
+		{
+			if (IDMXProtocolPtr ProtocolPtr = ProtocolIt->Value)
+			{
+				ProtocolPtr->RemoveAllUniverses();
+			}
+		}
+	}
+
+	FDMXEditorHelper()
+	{
+		DMXEditor = MakeShared<FDMXEditor>();
+		UPackage* Package = GetTransientPackage();
+		UDMXEditorFactoryNew* Factory = NewObject<UDMXEditorFactoryNew>(Package, MakeUniqueObjectName(GetTransientPackage(), UDMXEditorFactoryNew::StaticClass()));
+		const FName NewLibraryName = MakeUniqueObjectName(Package, UDMXLibrary::StaticClass(), *FString::Printf(TEXT("%d_LIB"), FMath::RandRange(0, 1000)));
+
+		DMXLibrary = ImportObject<UDMXLibrary>(Package, *NewLibraryName.ToString(), RF_Public | RF_Standalone | RF_Transient, *NewLibraryName.ToString(), nullptr, Factory);
+
+		FName NewTemplateName = MakeUniqueObjectName(DMXLibrary, UDMXEntityFader::StaticClass());
+		NewFaderTemplate = FDMXEditorUtils::CreateFaderTemplate(DMXLibrary);
+		FaderEntity = Cast<UDMXEntityFader>(DMXLibrary->GetOrCreateEntityObject(TEXT(""), UDMXEntityFader::StaticClass()));
+		FaderEntity->SetName(NewFaderTemplate->GetDisplayName());
+	}
+
+	TSharedPtr<FDMXEditor> DMXEditor;
+	UDMXLibrary* DMXLibrary;
+	UDMXEntityFader* FaderEntity;
+	UDMXEntityFader* NewFaderTemplate;
+};
+
+
+static const uint8 TestChannelValue = 155;
+static const uint8 ExistingUniverse = 1;
+static const uint8 NonExistingUniverse = 12;
+static const uint8 ChannelValue = 50;
+
+IMPLEMENT_SIMPLE_AUTOMATION_TEST(FDMXEditorFadersArtNetTest, "VirtualProduction.DMX.Editor.Faders.ArtNet",
+	EAutomationTestFlags::EditorContext | EAutomationTestFlags::EngineFilter);
+
+bool FDMXEditorFadersArtNetTest::RunTest(const FString& Parameters)
+{
+	FDMXEditorHelper::ResetUniverses();
+
+	// Create Helper
+	TSharedPtr<FDMXEditorHelper> Helper = MakeShared<FDMXEditorHelper>();
+
+	IDMXProtocolPtr DMXProtocol = IDMXProtocol::Get(FDMXProtocolArtNetModule::NAME_Artnet);
+
+	// Add first universe
+	{
+		FDMXUniverse DMXUniverse;
+		DMXUniverse.Channel = 10;
+		DMXUniverse.UniverseNumber = 1;
+		DMXUniverse.DMXProtocolDirectionality = EDMXProtocolDirectionality::EOutput;
+		Helper->FaderEntity->Universes.Add(DMXUniverse);
+	}
+
+	// Add second universe
+	{
+		FDMXUniverse DMXUniverse;
+		DMXUniverse.Channel = 20;
+		DMXUniverse.UniverseNumber = 2;
+		DMXUniverse.DMXProtocolDirectionality = EDMXProtocolDirectionality::EOutput;
+		Helper->FaderEntity->Universes.Add(DMXUniverse);
+	}
+	Helper->FaderEntity->DeviceProtocol = FDMXProtocolName(FDMXProtocolArtNetModule::NAME_Artnet);
+	Helper->FaderEntity->PostEditChange();
+
+	Helper->DMXEditor->InitEditor(EToolkitMode::Standalone, nullptr, Helper->DMXLibrary);
+
+	TSharedPtr<SDMXOutputFaderList> DMXOutputFaderList = SNew(SDMXOutputFaderList)
+		.DMXEditor(Helper->DMXEditor)
+		.FaderTemplate(Helper->NewFaderTemplate);
+
+	TestTrue(TEXT("Should be only one fader now"), DMXOutputFaderList->GetFaderWidgets().Num() == 1);
+
+	DMXOutputFaderList->GetFaderWidgets()[0]->GetFaderBoxVertical()->SetValue(ChannelValue);
+
+	// Check the data inside universe buffers
+	TArray<TSharedPtr<IDMXProtocolUniverse, ESPMode::ThreadSafe>> Universes;
+	Universes.Add(DMXProtocol->GetUniverseById(Helper->FaderEntity->Universes[0].UniverseNumber));
+	Universes.Add(DMXProtocol->GetUniverseById(Helper->FaderEntity->Universes[1].UniverseNumber));
+
+	TestEqual(TEXT("Two universes must exist"), Universes.Num(), 2);
+	if (Universes.Num() == 2)
+	{
+		TestTrue(TEXT("Universe should be exists"),
+			Universes[0].IsValid());
+		TestTrue(TEXT("Universe should be exists"),
+			Universes[1].IsValid());
+		if (Universes[0].IsValid() && Universes[1].IsValid())
+		{
+			const uint8 DMXData0 = Universes[0]->GetOutputDMXBuffer()->GetDMXDataAddress(
+				Helper->FaderEntity->Universes[0].Channel - 1);
+			TestEqual(TEXT("Buffer value should be same"), DMXData0, ChannelValue);
+			const uint8 DMXData1 = Universes[1]->GetOutputDMXBuffer()->GetDMXDataAddress(
+				Helper->FaderEntity->Universes[1].Channel - 1);
+			TestEqual(TEXT("Buffer value should be same"), DMXData1, ChannelValue);
+		}
+	}
+
+	Helper->DMXEditor->CloseWindow();
+
+	FDMXEditorHelper::ResetUniverses();
+
+	return true;
+}
+
+IMPLEMENT_SIMPLE_AUTOMATION_TEST(FDMXEditorFadersSACNTest, "VirtualProduction.DMX.Editor.Faders.SACN",
+	EAutomationTestFlags::EditorContext | EAutomationTestFlags::EngineFilter);
+
+bool FDMXEditorFadersSACNTest::RunTest(const FString& Parameters)
+{
+	FDMXEditorHelper::ResetUniverses();
+
+	// Create Helper
+	TSharedPtr<FDMXEditorHelper> Helper = MakeShared<FDMXEditorHelper>();
+
+	IDMXProtocolPtr DMXProtocol = IDMXProtocol::Get(FDMXProtocolSACNModule::NAME_SACN);
+
+	// Add first universe
+	{
+		FDMXUniverse DMXUniverse;
+		DMXUniverse.Channel = 10;
+		DMXUniverse.UniverseNumber = 31;
+		DMXUniverse.DMXProtocolDirectionality = EDMXProtocolDirectionality::EOutput;
+		Helper->FaderEntity->Universes.Add(DMXUniverse);
+	}
+
+	// Add second universe
+	{
+		FDMXUniverse DMXUniverse;
+		DMXUniverse.Channel = 20;
+		DMXUniverse.UniverseNumber = 51;
+		DMXUniverse.DMXProtocolDirectionality = EDMXProtocolDirectionality::EOutput;
+		Helper->FaderEntity->Universes.Add(DMXUniverse);
+	}
+
+	Helper->FaderEntity->DeviceProtocol = FDMXProtocolName(FDMXProtocolSACNModule::NAME_SACN);
+	Helper->FaderEntity->PostEditChange();
+
+	Helper->DMXEditor->InitEditor(EToolkitMode::Standalone, nullptr, Helper->DMXLibrary);
+
+	TSharedPtr<SDMXOutputFaderList> DMXOutputFaderList = SNew(SDMXOutputFaderList)
+		.DMXEditor(Helper->DMXEditor)
+		.FaderTemplate(Helper->NewFaderTemplate);
+
+	TestTrue(TEXT("Should be only one fader now"), DMXOutputFaderList->GetFaderWidgets().Num() == 1);
+
+	DMXOutputFaderList->GetFaderWidgets()[0]->GetFaderBoxVertical()->SetValue(ChannelValue);
+
+	// Check the data inside universe buffers
+	TArray<TSharedPtr<IDMXProtocolUniverse, ESPMode::ThreadSafe>> Universes;
+	Universes.Add(DMXProtocol->GetUniverseById(Helper->FaderEntity->Universes[0].UniverseNumber));
+	Universes.Add(DMXProtocol->GetUniverseById(Helper->FaderEntity->Universes[1].UniverseNumber));
+
+	TestEqual(TEXT("Two universes must exist"), Universes.Num(), 2);
+	if (Universes.Num() == 2)
+	{
+		TestTrue(TEXT("Universe should be exists"),
+			Universes[0].IsValid());
+		TestTrue(TEXT("Universe should be exists"),
+			Universes[1].IsValid());
+
+		if (Universes[0].IsValid() && Universes[1].IsValid())
+		{
+			const uint8 DMXData0 = Universes[0]->GetOutputDMXBuffer()->GetDMXDataAddress(
+				Helper->FaderEntity->Universes[0].Channel - 1);
+			TestEqual(TEXT("Buffer value should be same"), DMXData0, ChannelValue);
+			const uint8 DMXData1 = Universes[1]->GetOutputDMXBuffer()->GetDMXDataAddress(
+				Helper->FaderEntity->Universes[1].Channel - 1);
+			TestEqual(TEXT("Buffer value should be same"), DMXData1, ChannelValue);
+		}
+	}
+
+	Helper->DMXEditor->CloseWindow();
+
+	FDMXEditorHelper::ResetUniverses();
+
+	return true;
+}
+
+IMPLEMENT_SIMPLE_AUTOMATION_TEST(FDMXEditorControllersEmptyTest, "VirtualProduction.DMX.Editor.Controllers.Empty",
+	EAutomationTestFlags::EditorContext | EAutomationTestFlags::EngineFilter);
+
+bool FDMXEditorControllersEmptyTest::RunTest(const FString& Parameters)
+{
+	FDMXEditorHelper::ResetUniverses();
+
+	FDMXEditorHelper Helpers[] =
+	{
+		FDMXEditorHelper(),
+		FDMXEditorHelper()
+	};
+
+	Helpers[0].DMXEditor->InitEditor(EToolkitMode::Standalone, nullptr, Helpers[0].DMXLibrary);
+	Helpers[1].DMXEditor->InitEditor(EToolkitMode::Standalone, nullptr, Helpers[1].DMXLibrary);
+
+	const TMap<FName, IDMXProtocolPtr>& Protocols = IDMXProtocol::GetProtocols();
+	for (TMap<FName, IDMXProtocolPtr>::TConstIterator ProtocolIt = Protocols.CreateConstIterator(); ProtocolIt; ++ProtocolIt)
+	{
+		if (IDMXProtocolPtr ProtocolPtr = ProtocolIt->Value)
+		{
+			TestEqual(TEXT("Verify number of universes Is 0"), 0, ProtocolPtr->GetUniversesNum());
+		}
+	}
+
+	// Check our files
+	static const FName AssetRegistryName(TEXT("AssetRegistry"));
+	FAssetRegistryModule& AssetRegistryModule = FModuleManager::LoadModuleChecked<FAssetRegistryModule>(AssetRegistryName);
+	TArray<FAssetData> Items;
+	AssetRegistryModule.Get().GetAssetsByClass(UDMXLibrary::StaticClass()->GetFName(), Items);
+
+	uint32 OurAssetsCount = 0;
+	if (Items.Num() > 0)
+	{
+		for (const FAssetData& Item : Items)
+		{
+			UObject* Asset = FindObject<UObject>(nullptr, *Item.ObjectPath.ToString());
+			if (Asset != nullptr)
+			{
+				if (Asset == Helpers[0].DMXLibrary || Asset == Helpers[1].DMXLibrary)
+				{
+					OurAssetsCount++;
+				}
+			}
+		}
+	}
+	TestEqual(TEXT("Verify number of assets Is 0"), 0, OurAssetsCount);
+
+	Helpers[0].DMXEditor->CloseWindow();
+	Helpers[1].DMXEditor->CloseWindow();
+
+	FDMXEditorHelper::ResetUniverses();
+
+	return true;
+}
+
+IMPLEMENT_SIMPLE_AUTOMATION_TEST(FDMXEditorControllersNonEmptyTest, "VirtualProduction.DMX.Editor.Controllers.NonEmpty",
+	EAutomationTestFlags::EditorContext | EAutomationTestFlags::EngineFilter);
+
+bool FDMXEditorControllersNonEmptyTest::RunTest(const FString& Parameters)
+{
+	const int32 UniverseLocalStart = 1;
+	const int32 UniverseLocalEnd = 2000;
+
+	FDMXEditorHelper::ResetUniverses();
+
+	FDMXEditorHelper Helpers[] =
+	{
+		FDMXEditorHelper(),
+		FDMXEditorHelper()
+	};
+
+	FName DeviceProtocolName = GET_MEMBER_NAME_CHECKED(UDMXEntityController, DeviceProtocol);
+
+	// Add one Universe for editor 0
+	{
+		FName NewLibraryName = MakeUniqueObjectName(GetTransientPackage(), UDMXLibrary::StaticClass());
+		UDMXEntityController* Controller = Cast<UDMXEntityController>(Helpers[0].DMXLibrary->GetOrCreateEntityObject(NewLibraryName.ToString(), UDMXEntityController::StaticClass()));
+		FDMXUniverse Universe;
+		Universe.UniverseNumber = 2;
+		Controller->UniverseLocalStart = 0;
+		Controller->UniverseLocalNum = 1;
+		Controller->RemoteOffset = 0;
+		Controller->Universes.Add(Universe);
+		FObjectEditorUtils::SetPropertyValue(Controller, DeviceProtocolName,
+			FDMXProtocolName(FDMXProtocolArtNetModule::NAME_Artnet));
+	}
+
+	// Add same Universe for editor 0
+	{
+		FName NewLibraryName = MakeUniqueObjectName(GetTransientPackage(), UDMXLibrary::StaticClass());
+		UDMXEntityController* Controller = Cast<UDMXEntityController>(Helpers[0].DMXLibrary->GetOrCreateEntityObject(NewLibraryName.ToString(), UDMXEntityController::StaticClass()));
+		FDMXUniverse Universe;
+		Universe.UniverseNumber = 2;
+		Controller->UniverseLocalStart = 0;
+		Controller->UniverseLocalNum = 1;
+		Controller->RemoteOffset = 0;
+		Controller->Universes.Add(Universe);
+		FObjectEditorUtils::SetPropertyValue(Controller, DeviceProtocolName,
+			FDMXProtocolName(FDMXProtocolArtNetModule::NAME_Artnet));
+	}
+
+	// Add Universe for editor 1
+	{
+		FName NewLibraryName = MakeUniqueObjectName(GetTransientPackage(), UDMXLibrary::StaticClass());
+		UDMXEntityController* Controller = Cast<UDMXEntityController>(Helpers[1].DMXLibrary->GetOrCreateEntityObject(NewLibraryName.ToString(), UDMXEntityController::StaticClass()));
+		FDMXUniverse Universe;
+		Universe.UniverseNumber = 675;
+		Controller->UniverseLocalStart = 0;
+		Controller->UniverseLocalNum = 1;
+		Controller->RemoteOffset = 1;
+		Controller->Universes.Add(Universe);
+		FObjectEditorUtils::SetPropertyValue(Controller, DeviceProtocolName,
+			FDMXProtocolName(FDMXProtocolArtNetModule::NAME_Artnet));
+	}
+
+	// Add same Universe for editor 1
+	{
+		FName NewLibraryName = MakeUniqueObjectName(GetTransientPackage(), UDMXLibrary::StaticClass());
+		UDMXEntityController* Controller = Cast<UDMXEntityController>(Helpers[1].DMXLibrary->GetOrCreateEntityObject(NewLibraryName.ToString(), UDMXEntityController::StaticClass()));
+		FDMXUniverse Universe;
+		Universe.UniverseNumber = 675;
+		Controller->UniverseLocalStart = 0;
+		Controller->UniverseLocalNum = 1;
+		Controller->RemoteOffset = 1;
+		Controller->Universes.Add(Universe);
+		FObjectEditorUtils::SetPropertyValue(Controller, DeviceProtocolName,
+			FDMXProtocolName(FDMXProtocolArtNetModule::NAME_Artnet));
+	}
+
+	Helpers[0].DMXEditor->InitEditor(EToolkitMode::Standalone, nullptr, Helpers[0].DMXLibrary);
+	Helpers[1].DMXEditor->InitEditor(EToolkitMode::Standalone, nullptr, Helpers[1].DMXLibrary);
+
+	uint32 NumUniverses = 0;
+	const TMap<FName, IDMXProtocolPtr>& Protocols = IDMXProtocol::GetProtocols();
+	for (TMap<FName, IDMXProtocolPtr>::TConstIterator ProtocolIt = Protocols.CreateConstIterator(); ProtocolIt; ++ProtocolIt)
+	{
+		if (IDMXProtocolPtr ProtocolPtr = ProtocolIt->Value)
+		{
+			NumUniverses += ProtocolPtr->GetUniversesNum();
+		}
+	}
+
+	TestEqual(TEXT("Verify number of universes Is 2"), NumUniverses, 2);
+
+	Helpers[0].DMXEditor->CloseWindow();
+	Helpers[1].DMXEditor->CloseWindow();
+
+	FDMXEditorHelper::ResetUniverses();
+
+	return true;
+}
+
+IMPLEMENT_SIMPLE_AUTOMATION_TEST(FDMXEditorInputConsoleArtNetExistingTest, "VirtualProduction.DMX.Editor.InputConsole.ArtNet.Existing",
+	EAutomationTestFlags::EditorContext | EAutomationTestFlags::EngineFilter);
+
+bool FDMXEditorInputConsoleArtNetExistingTest::RunTest(const FString& Parameters)
+{
+	// Reset Universes
+	FDMXEditorHelper::ResetUniverses();
+
+	FName DeviceProtocolName = GET_MEMBER_NAME_CHECKED(UDMXEntityController, DeviceProtocol);
+	
+	// Create Helper
+	TSharedPtr<FDMXEditorHelper> Helper = MakeShared<FDMXEditorHelper>();
+
+	// 2. Create Universe
+	FName NewLibraryName = MakeUniqueObjectName(GetTransientPackage(), UDMXLibrary::StaticClass());
+	UDMXEntityController* Controller = Cast<UDMXEntityController>(Helper->DMXLibrary->GetOrCreateEntityObject(NewLibraryName.ToString(), UDMXEntityController::StaticClass()));
+	FDMXUniverse Universe;
+	Universe.UniverseNumber = ExistingUniverse;
+	Controller->DeviceProtocol = FDMXProtocolName(FDMXProtocolArtNetModule::NAME_Artnet);
+	Controller->UniverseLocalStart = 0;
+	Controller->UniverseLocalNum = 1;
+	Controller->RemoteOffset = 1;
+	Controller->Universes.Add(Universe);
+	FObjectEditorUtils::SetPropertyValue(Controller, DeviceProtocolName,
+		FDMXProtocolName(FDMXProtocolArtNetModule::NAME_Artnet));
+
+	// 3. Create editor
+	Helper->DMXEditor->InitEditor(EToolkitMode::Standalone, nullptr, Helper->DMXLibrary);
+
+	// 4. Set Input Tab
+	TSharedRef<SDMXInputConsole> InputConsole = Helper->DMXEditor->GetInputConsoleTab();
+	const TSharedRef<SDMXInputInfo>& InputInfo = InputConsole->GetInputInfo();
+	const TSharedRef<SDMXInputInfoSelecter>& InfoSelecter = InputConsole->GetInputInfoSelecter();
+	const TSharedRef<SSpinBox<uint32>>& UniverseField = InfoSelecter->GetUniverseField();
+
+	InfoSelecter->SetProtocol(FDMXProtocolArtNetModule::NAME_Artnet);
+	UniverseField->SetValue(ExistingUniverse);
+
+	// 5. Send DMX
+	IDMXProtocolPtr DMXProtocol = IDMXProtocol::Get(FDMXProtocolArtNetModule::NAME_Artnet);
+	IDMXFragmentMap FragmentMap;
+	FragmentMap.Add(1, TestChannelValue);
+	EDMXSendResult SendResult = DMXProtocol->SendDMXFragment(ExistingUniverse, FragmentMap);
+	TestEqual(TEXT("SendDMXFragment failed"), SendResult, EDMXSendResult::Success);
+
+	if (SendResult == EDMXSendResult::Success)
+	{
+		// Check the input console values
+		AddCommand(new FDelayedFunctionLatentCommand([=]
+			{
+				TSharedRef<SDMXInputConsole> InputConsole = Helper->DMXEditor->GetInputConsoleTab();
+				const TSharedPtr<SDMXInputInfo>& InputInfo = InputConsole->GetInputInfo();
+
+				// Force Tick to allow for InputInfo to self update
+				FGeometry AllottedGeometry;
+				InputInfo->Tick(AllottedGeometry, 0.0, 0.0);
+
+				// 6. Check input console
+				const TArray<uint8>& ChannelsValues = InputInfo->GetChannelsValues();
+				TestEqual(TEXT("Verify ChannelsValue"), TestChannelValue, ChannelsValues[0]);
+
+				FDMXEditorHelper::ResetUniverses();
+
+				// 7. Close the Editor
+				Helper->DMXEditor->CloseWindow();
+
+			}, 0.2f));
+	}
+
+	return true;
+}
+
+IMPLEMENT_SIMPLE_AUTOMATION_TEST(FDMXEditorInputConsoleArtNetNonExistingTest, "VirtualProduction.DMX.Editor.InputConsole.ArtNet.NonExisting",
+	EAutomationTestFlags::EditorContext | EAutomationTestFlags::EngineFilter);
+
+bool FDMXEditorInputConsoleArtNetNonExistingTest::RunTest(const FString& Parameters)
+{
+	// Reset Universes
+	FDMXEditorHelper::ResetUniverses();
+
+	FName DeviceProtocolName = GET_MEMBER_NAME_CHECKED(UDMXEntityController, DeviceProtocol);
+
+	// Create Helper
+	TSharedPtr<FDMXEditorHelper> Helper = MakeShared<FDMXEditorHelper>();
+
+	// 2. Create Universe
+	FName NewLibraryName = MakeUniqueObjectName(GetTransientPackage(), UDMXLibrary::StaticClass());
+	UDMXEntityController* Controller = Cast<UDMXEntityController>(Helper->DMXLibrary->GetOrCreateEntityObject(NewLibraryName.ToString(), UDMXEntityController::StaticClass()));
+	FDMXUniverse Universe;
+	Universe.UniverseNumber = ExistingUniverse;
+	Controller->UniverseLocalStart = 0;
+	Controller->UniverseLocalNum = 1;
+	Controller->RemoteOffset = 1;
+	Controller->DeviceProtocol = FDMXProtocolName(FDMXProtocolArtNetModule::NAME_Artnet);
+	Controller->Universes.Add(Universe);
+	FObjectEditorUtils::SetPropertyValue(Controller, DeviceProtocolName,
+		FDMXProtocolName(FDMXProtocolArtNetModule::NAME_Artnet));
+
+	// 3. Create editor
+	Helper->DMXEditor->InitEditor(EToolkitMode::Standalone, nullptr, Helper->DMXLibrary);
+
+	// 4. Set Input Tab
+	TSharedRef<SDMXInputConsole> InputConsole = Helper->DMXEditor->GetInputConsoleTab();
+	const TSharedRef<SDMXInputInfo>& InputInfo = InputConsole->GetInputInfo();
+	const TSharedRef<SDMXInputInfoSelecter>& InfoSelecter = InputConsole->GetInputInfoSelecter();
+	const TSharedRef<SSpinBox<uint32>>& UniverseField = InfoSelecter->GetUniverseField();
+
+	InfoSelecter->SetProtocol(FDMXProtocolArtNetModule::NAME_Artnet);
+	UniverseField->SetValue(NonExistingUniverse);
+
+	// 5. Send DMX
+	IDMXProtocolPtr DMXProtocol = IDMXProtocol::Get(FDMXProtocolArtNetModule::NAME_Artnet);
+	IDMXFragmentMap FragmentMap;
+	FragmentMap.Add(1, TestChannelValue);
+	EDMXSendResult SendResult = DMXProtocol->SendDMXFragment(ExistingUniverse, FragmentMap);
+	TestEqual(TEXT("SendDMXFragment failed"), SendResult, EDMXSendResult::Success);
+
+	if (SendResult == EDMXSendResult::Success)
+	{
+		// Check the input console values
+		AddCommand(new FDelayedFunctionLatentCommand([=]
+			{
+				TSharedRef<SDMXInputConsole> InputConsole = Helper->DMXEditor->GetInputConsoleTab();
+				const TSharedPtr<SDMXInputInfo>& InputInfo = InputConsole->GetInputInfo();
+
+				// Force Tick to allow for InputInfo to self update
+				FGeometry AllottedGeometry;
+				InputInfo->Tick(AllottedGeometry, 0.0, 0.0);
+
+				// 6. Check input console
+				const TArray<uint8>& ChannelsValues = InputInfo->GetChannelsValues();
+				TestNotEqual(TEXT("Verify ChannelsValue"), TestChannelValue, ChannelsValues[0]);
+
+				FDMXEditorHelper::ResetUniverses();
+
+				// 7. Close the Editor
+				Helper->DMXEditor->CloseWindow();
+
+			}, 0.2f));
+	}
+
+	return true;
+}
+
+IMPLEMENT_SIMPLE_AUTOMATION_TEST(FDMXEditorInputConsoleSACNExistingTest, "VirtualProduction.DMX.Editor.InputConsole.SACN.Existing",
+	EAutomationTestFlags::EditorContext | EAutomationTestFlags::EngineFilter);
+
+bool FDMXEditorInputConsoleSACNExistingTest::RunTest(const FString& Parameters)
+{
+	// Reset Universes
+	FDMXEditorHelper::ResetUniverses();
+
+	FName DeviceProtocolName = GET_MEMBER_NAME_CHECKED(UDMXEntityController, DeviceProtocol);
+
+	// Create Helper
+	TSharedPtr<FDMXEditorHelper> Helper = MakeShared<FDMXEditorHelper>();
+
+	// 2. Create Universe
+	FName NewLibraryName = MakeUniqueObjectName(GetTransientPackage(), UDMXLibrary::StaticClass());
+	UDMXEntityController* Controller = Cast<UDMXEntityController>(Helper->DMXLibrary->GetOrCreateEntityObject(NewLibraryName.ToString(), UDMXEntityController::StaticClass()));
+	FDMXUniverse Universe;
+	Universe.UniverseNumber = ExistingUniverse;
+	Controller->UniverseLocalStart = 1;
+	Controller->UniverseLocalNum = 1;
+	Controller->RemoteOffset = 0;
+	Controller->DeviceProtocol = FDMXProtocolName(FDMXProtocolSACNModule::NAME_SACN);
+	Controller->Universes.Add(Universe);
+	FObjectEditorUtils::SetPropertyValue(Controller, DeviceProtocolName,
+		FDMXProtocolName(FDMXProtocolSACNModule::NAME_SACN));
+
+	// 3. Create editor
+	Helper->DMXEditor->InitEditor(EToolkitMode::Standalone, nullptr, Helper->DMXLibrary);
+
+	// 4. Set Input Tab
+	TSharedRef<SDMXInputConsole> InputConsole = Helper->DMXEditor->GetInputConsoleTab();
+	const TSharedRef<SDMXInputInfo>& InputInfo = InputConsole->GetInputInfo();
+	const TSharedRef<SDMXInputInfoSelecter>& InfoSelecter = InputConsole->GetInputInfoSelecter();
+	const TSharedRef<SSpinBox<uint32>>& UniverseField = InfoSelecter->GetUniverseField();
+
+	InfoSelecter->SetProtocol(FDMXProtocolSACNModule::NAME_SACN);
+	UniverseField->SetValue(ExistingUniverse);
+
+	// 5. Send DMX
+	IDMXProtocolPtr DMXProtocol = IDMXProtocol::Get(FDMXProtocolSACNModule::NAME_SACN);
+	IDMXFragmentMap FragmentMap;
+	FragmentMap.Add(1, TestChannelValue);
+	EDMXSendResult SendResult = DMXProtocol->SendDMXFragment(ExistingUniverse, FragmentMap);
+	TestEqual(TEXT("SendDMXFragment failed"), SendResult, EDMXSendResult::Success);
+
+	if (SendResult == EDMXSendResult::Success)
+	{
+		// Check the input console values
+		AddCommand(new FDelayedFunctionLatentCommand([=]
+			{
+				TSharedRef<SDMXInputConsole> InputConsole = Helper->DMXEditor->GetInputConsoleTab();
+				const TSharedPtr<SDMXInputInfo>& InputInfo = InputConsole->GetInputInfo();
+
+				// Force Tick to allow for InputInfo to self update
+				FGeometry AllottedGeometry;
+				InputInfo->Tick(AllottedGeometry, 0.0, 0.0);
+
+				// 6. Check input console
+				const TArray<uint8>& ChannelsValues = InputInfo->GetChannelsValues();
+				TestEqual(TEXT("Verify ChannelsValue"), TestChannelValue, ChannelsValues[0]);
+
+				FDMXEditorHelper::ResetUniverses();
+
+				// 7. Close the Editor
+				Helper->DMXEditor->CloseWindow();
+
+			}, 0.2f));
+	}
+
+	return true;
+}
+
+IMPLEMENT_SIMPLE_AUTOMATION_TEST(FDMXEditorInputConsoleSACNNonExistingTest, "VirtualProduction.DMX.Editor.InputConsole.sACN.NonExisting",
+	EAutomationTestFlags::EditorContext | EAutomationTestFlags::EngineFilter);
+
+bool FDMXEditorInputConsoleSACNNonExistingTest::RunTest(const FString& Parameters)
+{
+	// Reset Universes
+	FDMXEditorHelper::ResetUniverses();
+	FName DeviceProtocolName = GET_MEMBER_NAME_CHECKED(UDMXEntityController, DeviceProtocol);
+
+	// Create Helper
+	TSharedPtr<FDMXEditorHelper> Helper = MakeShared<FDMXEditorHelper>();
+
+	// 2. Create Universe
+	FName NewLibraryName = MakeUniqueObjectName(GetTransientPackage(), UDMXLibrary::StaticClass());
+	UDMXEntityController* Controller = Cast<UDMXEntityController>(Helper->DMXLibrary->GetOrCreateEntityObject(NewLibraryName.ToString(), UDMXEntityController::StaticClass()));
+	FDMXUniverse Universe;
+	Universe.UniverseNumber = ExistingUniverse;
+	Controller->DeviceProtocol = FDMXProtocolName(FDMXProtocolSACNModule::NAME_SACN);
+	Controller->UniverseLocalStart = 1;
+	Controller->UniverseLocalNum = 1;
+	Controller->RemoteOffset = 0;
+	Controller->Universes.Add(Universe);
+	FObjectEditorUtils::SetPropertyValue(Controller, DeviceProtocolName,
+		FDMXProtocolName(FDMXProtocolSACNModule::NAME_SACN));
+
+	// 3. Create editor
+	Helper->DMXEditor->InitEditor(EToolkitMode::Standalone, nullptr, Helper->DMXLibrary);
+
+	// 4. Set Input Tab
+	TSharedRef<SDMXInputConsole> InputConsole = Helper->DMXEditor->GetInputConsoleTab();
+	const TSharedRef<SDMXInputInfo>& InputInfo = InputConsole->GetInputInfo();
+	const TSharedRef<SDMXInputInfoSelecter>& InfoSelecter = InputConsole->GetInputInfoSelecter();
+	const TSharedRef<SSpinBox<uint32>>& UniverseField = InfoSelecter->GetUniverseField();
+
+	InfoSelecter->SetProtocol(FDMXProtocolSACNModule::NAME_SACN);
+	UniverseField->SetValue(NonExistingUniverse);
+
+	// 5. Send DMX
+	IDMXProtocolPtr DMXProtocol = IDMXProtocol::Get(FDMXProtocolSACNModule::NAME_SACN);
+	IDMXFragmentMap FragmentMap;
+	FragmentMap.Add(1, TestChannelValue);
+	EDMXSendResult SendResult = DMXProtocol->SendDMXFragment(ExistingUniverse, FragmentMap);
+	TestEqual(TEXT("SendDMXFragment failed"), SendResult, EDMXSendResult::Success);
+
+	if (SendResult == EDMXSendResult::Success)
+	{
+		// Check the input console values
+		AddCommand(new FDelayedFunctionLatentCommand([=]
+			{
+				TSharedRef<SDMXInputConsole> InputConsole = Helper->DMXEditor->GetInputConsoleTab();
+				const TSharedPtr<SDMXInputInfo>& InputInfo = InputConsole->GetInputInfo();
+
+				// Force Tick to allow for InputInfo to self update
+				FGeometry AllottedGeometry;
+				InputInfo->Tick(AllottedGeometry, 0.0, 0.0);
+
+				// 6. Check input console
+				const TArray<uint8>& ChannelsValues = InputInfo->GetChannelsValues();
+				TestNotEqual(TEXT("Verify ChannelsValue"), TestChannelValue, ChannelsValues[0]);
+
+				FDMXEditorHelper::ResetUniverses();
+
+				// 7. Close the Editor
+				Helper->DMXEditor->CloseWindow();
+
+			}, 0.2f));
+	}
+
+	return true;
+}
+
+#endif
\ No newline at end of file
diff --git a/Engine/Plugins/Experimental/DigitalMultiplex/DMXEngine/Source/DMXEditor/Private/Widgets/OutputFader/SDMXFader.cpp b/Engine/Plugins/Experimental/DigitalMultiplex/DMXEngine/Source/DMXEditor/Private/Widgets/OutputFader/SDMXFader.cpp
--- a/Engine/Plugins/Experimental/DigitalMultiplex/DMXEngine/Source/DMXEditor/Private/Widgets/OutputFader/SDMXFader.cpp
+++ b/Engine/Plugins/Experimental/DigitalMultiplex/DMXEngine/Source/DMXEditor/Private/Widgets/OutputFader/SDMXFader.cpp
@@ -275,7 +275,7 @@ void SDMXFader::SetFaderLabel(const FString& InLabel)
 	CustomFaderLabel->SetText(FText::FromString(InLabel));
 }
 
-void SDMXFader::AddChannelWidget(const FString& InUniverse, const FString& InChannel, uint16 InUniverseNumber, uint32 InChannelNumber, const TSharedPtr<IDMXProtocol>& InDMXProtocol)
+void SDMXFader::AddChannelWidget(const FString& InUniverse, const FString& InChannel, uint16 InUniverseNumber, uint32 InChannelNumber, const IDMXProtocolPtr& InDMXProtocol)
 {
 	TSharedPtr<SDMXFaderChannel> NewChannel = SNew(SDMXFaderChannel)
 		.DMXEditor(WeakDMXEditor)
@@ -347,7 +347,7 @@ void SDMXFader::HandleFaderChanged(uint8 NewValue)
 		{
 			if (WeakFaderEntity.Get())
 			{
-				TSharedPtr<IDMXProtocol> DMXProtocol = WeakFaderEntity->DeviceProtocol;
+				IDMXProtocolPtr DMXProtocol = WeakFaderEntity->DeviceProtocol;
 				if (DMXProtocol.IsValid())
 				{
 					IDMXFragmentMap DMXFragmentMap;
diff --git a/Engine/Plugins/Experimental/DigitalMultiplex/DMXEngine/Source/DMXEditor/Private/Widgets/OutputFader/SDMXFader.h b/Engine/Plugins/Experimental/DigitalMultiplex/DMXEngine/Source/DMXEditor/Private/Widgets/OutputFader/SDMXFader.h
--- a/Engine/Plugins/Experimental/DigitalMultiplex/DMXEngine/Source/DMXEditor/Private/Widgets/OutputFader/SDMXFader.h
+++ b/Engine/Plugins/Experimental/DigitalMultiplex/DMXEngine/Source/DMXEditor/Private/Widgets/OutputFader/SDMXFader.h
@@ -7,6 +7,8 @@
 #include "Widgets/SCompoundWidget.h"
 #include "Styling/SlateTypes.h"
 
+#include "DMXProtocolCommon.h"
+
 template<typename NumericType>
 class SSpinBoxVertical;
 class SDMXOutputFaderList;
@@ -44,7 +46,7 @@ public:
 		
 	void SetFaderLabel(const FString& InLabel);
 
-	void AddChannelWidget(const FString& InUniverse, const FString& InChannel, uint16 InUniverseNumber, uint32 InChannelNumber, const TSharedPtr<IDMXProtocol>& DMXProtocol);
+	void AddChannelWidget(const FString& InUniverse, const FString& InChannel, uint16 InUniverseNumber, uint32 InChannelNumber, const IDMXProtocolPtr& DMXProtocol);
 
 	void RemoveAllChannelWidgets();
 
diff --git a/Engine/Plugins/Experimental/DigitalMultiplex/DMXEngine/Source/DMXEditor/Private/Widgets/OutputFader/SDMXFaderChannel.h b/Engine/Plugins/Experimental/DigitalMultiplex/DMXEngine/Source/DMXEditor/Private/Widgets/OutputFader/SDMXFaderChannel.h
--- a/Engine/Plugins/Experimental/DigitalMultiplex/DMXEngine/Source/DMXEditor/Private/Widgets/OutputFader/SDMXFaderChannel.h
+++ b/Engine/Plugins/Experimental/DigitalMultiplex/DMXEngine/Source/DMXEditor/Private/Widgets/OutputFader/SDMXFaderChannel.h
@@ -27,7 +27,7 @@ public:
 	SLATE_ARGUMENT(FText, InText)
 	SLATE_ARGUMENT(uint16, UniverseNumber)
 	SLATE_ARGUMENT(uint32, ChannelNumber)
-	SLATE_ARGUMENT(TWeakPtr<IDMXProtocol>, DMXProtocol)
+	SLATE_ARGUMENT(IDMXProtocolPtrWeak, DMXProtocol)
 
 	SLATE_END_ARGS()
 
@@ -42,7 +42,7 @@ public:
 
 	uint32 GetChannelNumber() const { return ChannelNumber; }
 
-	const TSharedPtr<IDMXProtocol> GetProtocol() const { return WeakDMXProtocol.Pin(); }
+	const IDMXProtocolPtr GetProtocol() const { return WeakDMXProtocol.Pin(); }
 
 private:
 	/** Pointer back to the DMXEditor tool that owns us */
@@ -56,7 +56,7 @@ private:
 
 	uint32 ChannelNumber;
 
-	TWeakPtr<IDMXProtocol> WeakDMXProtocol;
+	IDMXProtocolPtrWeak WeakDMXProtocol;
 };
 
 #undef LOCTEXT_NAMESPACE
\ No newline at end of file
diff --git a/Engine/Plugins/Experimental/DigitalMultiplex/DMXEngine/Source/DMXEditor/Private/Widgets/OutputFader/SDMXOutputFaderList.cpp b/Engine/Plugins/Experimental/DigitalMultiplex/DMXEngine/Source/DMXEditor/Private/Widgets/OutputFader/SDMXOutputFaderList.cpp
--- a/Engine/Plugins/Experimental/DigitalMultiplex/DMXEngine/Source/DMXEditor/Private/Widgets/OutputFader/SDMXOutputFaderList.cpp
+++ b/Engine/Plugins/Experimental/DigitalMultiplex/DMXEngine/Source/DMXEditor/Private/Widgets/OutputFader/SDMXOutputFaderList.cpp
@@ -200,7 +200,7 @@ void SDMXOutputFaderList::SetFaderProperties(const TSharedPtr<SDMXFader>& InFade
 		TransferCreatedFaderObjectProperties(InFaderObject, WeakFaderTemplate);
 	}
 
-	TSharedPtr<IDMXProtocol> DMXProtocol = InFaderObject->DeviceProtocol;
+	IDMXProtocolPtr DMXProtocol = InFaderObject->DeviceProtocol;
 
 	InFaderWidget->SetFaderLabel(InFaderObject->GetDisplayName());
 	for (const FDMXUniverse& Universe : InFaderObject->Universes)
diff --git a/Engine/Plugins/Experimental/DigitalMultiplex/DMXEngine/Source/DMXEditor/Private/Widgets/SDMXInputInfo.cpp b/Engine/Plugins/Experimental/DigitalMultiplex/DMXEngine/Source/DMXEditor/Private/Widgets/SDMXInputInfo.cpp
--- a/Engine/Plugins/Experimental/DigitalMultiplex/DMXEngine/Source/DMXEditor/Private/Widgets/SDMXInputInfo.cpp
+++ b/Engine/Plugins/Experimental/DigitalMultiplex/DMXEngine/Source/DMXEditor/Private/Widgets/SDMXInputInfo.cpp
@@ -90,21 +90,24 @@ void SDMXInputInfo::UpdateChannelsValues()
 {
 	if (TSharedPtr<SDMXInputInfoSelecter> InfoSelecterPtr = WeakInfoSelecter.Pin())
 	{
-		if (TSharedPtr<IDMXProtocol> DMXProtocolPtr = IDMXProtocol::Get(InfoSelecterPtr->GetCurrentProtocolName()))
+		if (IDMXProtocolPtr DMXProtocolPtr = IDMXProtocol::Get(InfoSelecterPtr->GetCurrentProtocolName()))
 		{
 			if (IDMXProtocolUniversePtr Universe = DMXProtocolPtr->GetUniverseById(InfoSelecterPtr->GetCurrentUniverseID()))
 			{
 				if (TSharedPtr<FDMXBuffer> DMXBuffer = Universe->GetInputDMXBuffer())
 				{
 					// check the sequence ID 
-					uint32 BufferSequanceID = DMXBuffer->GetSequanceID();
-					if (BufferSequanceID != UISequanceID)
+					uint32 BufferSequenceID = DMXBuffer->GetSequenceID();
+					if (BufferSequenceID != UISequenceID)
 					{
-						FMemory::Memcpy(ChannelsValues.GetData(), DMXBuffer->GetDMXData().GetData(), DMX_UNIVERSE_SIZE);
+						DMXBuffer->AccessDMXData([this](TArray<uint8>& InData)
+							{
+								FMemory::Memcpy(ChannelsValues.GetData(), InData.GetData(), DMX_UNIVERSE_SIZE);
+							});
 						UpdateChannelWidgetsValues(ChannelsValues);
 					}
 
-					UISequanceID = BufferSequanceID;
+					UISequenceID = BufferSequenceID;
 				}
 			}
 		}
diff --git a/Engine/Plugins/Experimental/DigitalMultiplex/DMXEngine/Source/DMXEditor/Private/Widgets/SDMXInputInfo.h b/Engine/Plugins/Experimental/DigitalMultiplex/DMXEngine/Source/DMXEditor/Private/Widgets/SDMXInputInfo.h
--- a/Engine/Plugins/Experimental/DigitalMultiplex/DMXEngine/Source/DMXEditor/Private/Widgets/SDMXInputInfo.h
+++ b/Engine/Plugins/Experimental/DigitalMultiplex/DMXEngine/Source/DMXEditor/Private/Widgets/SDMXInputInfo.h
@@ -29,7 +29,7 @@ public:
 
 	const TArray<uint8>& GetChannelsValues() const { return ChannelsValues; }
 
-	void ResetUISequanceID() { UISequanceID = 0; }
+	void ResetUISequanceID() { UISequenceID = 0; }
 	
 protected:
 	/** Container widget for all the channels' values */
@@ -53,5 +53,5 @@ protected:
 
 private:
 	/** ID of the sequence on input info widget */
-	uint32 UISequanceID;
+	uint32 UISequenceID;
 }; 
\ No newline at end of file
diff --git a/Engine/Plugins/Experimental/DigitalMultiplex/DMXEngine/Source/DMXEditor/Private/Widgets/SDMXInputInfoSelecter.cpp b/Engine/Plugins/Experimental/DigitalMultiplex/DMXEngine/Source/DMXEditor/Private/Widgets/SDMXInputInfoSelecter.cpp
--- a/Engine/Plugins/Experimental/DigitalMultiplex/DMXEngine/Source/DMXEditor/Private/Widgets/SDMXInputInfoSelecter.cpp
+++ b/Engine/Plugins/Experimental/DigitalMultiplex/DMXEngine/Source/DMXEditor/Private/Widgets/SDMXInputInfoSelecter.cpp
@@ -94,7 +94,7 @@ void SDMXInputInfoSelecter::Construct(const FArguments& InArgs)
 			.Padding(PaddingKeyVal, 0, 0, 0)
 			.AutoWidth()
 			[
-				SAssignNew(UniverseIDField, SSpinBox<uint16>)
+				SAssignNew(UniverseIDField, SSpinBox<uint32>)
 				.Value(this, &SDMXInputInfoSelecter::GetCurrentUniverseID)
 				.OnValueChanged(this, &SDMXInputInfoSelecter::HandleUniverseIDChanged)
 				.OnValueCommitted(this, &SDMXInputInfoSelecter::HandleUniverseIDValueCommitted)
@@ -153,19 +153,19 @@ void SDMXInputInfoSelecter::HandleProtocolChanged(FName SelectedProtocol)
 	ProtocolSettings->SaveConfig();
 }
 
-void SDMXInputInfoSelecter::HandleUniverseIDChanged(uint16 NewValue)
+void SDMXInputInfoSelecter::HandleUniverseIDChanged(uint32 NewValue)
 {
-	if (CurrentUniverseID != NewValue)
+	if (CurrentUniverseID != (uint16)NewValue)
 	{
-		CurrentUniverseID = NewValue;
+		CurrentUniverseID = FMath::Clamp(NewValue, (uint32)0u, (uint32)MAX_uint16);
 	}
 }
 
-void SDMXInputInfoSelecter::HandleUniverseIDValueCommitted(uint16 NewValue, ETextCommit::Type CommitType)
+void SDMXInputInfoSelecter::HandleUniverseIDValueCommitted(uint32 NewValue, ETextCommit::Type CommitType)
 {
-	if (CurrentUniverseID != NewValue)
+	if (CurrentUniverseID != (uint16)NewValue)
 	{
-		CurrentUniverseID = NewValue;
+		CurrentUniverseID = FMath::Clamp(NewValue, (uint32)0u, (uint32)MAX_uint16);;
 	}
 
 	// Update stored settings
diff --git a/Engine/Plugins/Experimental/DigitalMultiplex/DMXEngine/Source/DMXEditor/Private/Widgets/SDMXInputInfoSelecter.h b/Engine/Plugins/Experimental/DigitalMultiplex/DMXEngine/Source/DMXEditor/Private/Widgets/SDMXInputInfoSelecter.h
--- a/Engine/Plugins/Experimental/DigitalMultiplex/DMXEngine/Source/DMXEditor/Private/Widgets/SDMXInputInfoSelecter.h
+++ b/Engine/Plugins/Experimental/DigitalMultiplex/DMXEngine/Source/DMXEditor/Private/Widgets/SDMXInputInfoSelecter.h
@@ -28,14 +28,14 @@ public:
 	/** Constructs the widget */
 	void Construct(const FArguments& InArgs);
 
-	const TSharedRef<SSpinBox<uint16>> GetUniverseField() const { return UniverseIDField.ToSharedRef(); }
+	const TSharedRef<SSpinBox<uint32>> GetUniverseField() const { return UniverseIDField.ToSharedRef(); }
 	const void SetProtocol(const FName& Name) { CurrentProtocol = FDMXProtocolName(Name); };
 
 protected:
 	/** Universe ID value computed using Net, Subnet and Universe values */
 	uint16 CurrentUniverseID;
 	/** Final Universe field widget */
-	TSharedPtr<SSpinBox<uint16>> UniverseIDField;
+	TSharedPtr<SSpinBox<uint32>> UniverseIDField;
 
 	/** The user-selected protocol */
 	FDMXProtocolName CurrentProtocol;
@@ -45,7 +45,7 @@ public:
 	//const FText& GetCurrentDirectionality() const { return CurrentDirectionality; }
 
 	/** Returns the UniverseID value computed from Net, Subnet and Universe values */
-	uint16 GetCurrentUniverseID() const { return CurrentUniverseID; }
+	uint32 GetCurrentUniverseID() const { return CurrentUniverseID; }
 
 	/** Returns the user-selected DMX protocol */
 	FName GetCurrentProtocolName() const { return CurrentProtocol; }
@@ -57,10 +57,10 @@ protected:
 	void HandleProtocolChanged(FName SelectedProtocol);
 
 	/** Handles when the user changes the universe value, including while spinning the value */
-	void HandleUniverseIDChanged(uint16 NewValue);
+	void HandleUniverseIDChanged(uint32 NewValue);
 
 	/** Handles when the user commit the Universe value. Doesn't fire while spinning the value */
-	void HandleUniverseIDValueCommitted(uint16 NewValue, ETextCommit::Type CommitType);
+	void HandleUniverseIDValueCommitted(uint32 NewValue, ETextCommit::Type CommitType);
 
 protected:
 	FOnUniverseSelectionChanged OnUniverseSelectionChanged;
diff --git a/Engine/Plugins/Experimental/DigitalMultiplex/DMXEngine/Source/DMXEditor/Private/Widgets/SNameListPicker.cpp b/Engine/Plugins/Experimental/DigitalMultiplex/DMXEngine/Source/DMXEditor/Private/Widgets/SNameListPicker.cpp
--- a/Engine/Plugins/Experimental/DigitalMultiplex/DMXEngine/Source/DMXEditor/Private/Widgets/SNameListPicker.cpp
+++ b/Engine/Plugins/Experimental/DigitalMultiplex/DMXEngine/Source/DMXEditor/Private/Widgets/SNameListPicker.cpp
@@ -23,18 +23,6 @@ void SNameListPicker::Construct(const FArguments& InArgs)
 		UpdateOptionsHandle = UpdateOptionsDelegate->Add(FSimpleDelegate::CreateSP(this, &SNameListPicker::UpdateOptionsSource));
 	}
 
-	TSharedPtr<FName> InitiallySelected = nullptr;
-	for (TSharedPtr<FName> NameItem : OptionsSource)
-	{
-		if (NameItem)
-		{
-			if (ValueAttribute.Get() == *NameItem)
-			{
-				InitiallySelected = NameItem;
-				break;
-			}
-		}
-	}
 
 	ChildSlot
 	[
@@ -42,7 +30,8 @@ void SNameListPicker::Construct(const FArguments& InArgs)
 		.OptionsSource(&OptionsSource)
 		.OnGenerateWidget(this, &SNameListPicker::GenerateNameItemWidget)
 		.OnSelectionChanged(this, &SNameListPicker::HandleSelectionChanged)
-		.InitiallySelectedItem(InitiallySelected)
+		.OnComboBoxOpening(this, &SNameListPicker::UpdateSelectedOption)
+		.InitiallySelectedItem(GetSelectedItemFromCurrentValue())
 		[
 			SNew(STextBlock)
 			.Text(this, &SNameListPicker::GetCurrentNameLabel)
@@ -105,6 +94,41 @@ void SNameListPicker::HandleSelectionChanged(const TSharedPtr<FName> Item, ESele
 	}
 }
 
+TSharedPtr<FName> SNameListPicker::GetSelectedItemFromCurrentValue() const
+{
+	TSharedPtr<FName> InitiallySelected = nullptr;
+
+	const bool bHasMultipleValues = HasMultipleValuesAttribute.Get();
+	if (bHasMultipleValues)
+	{
+		return InitiallySelected;
+	}
+
+	const FName CurrentValue = ValueAttribute.Get();
+
+	for (TSharedPtr<FName> NameItem : OptionsSource)
+	{
+		if (NameItem != nullptr)
+		{
+			if (CurrentValue.IsEqual(*NameItem))
+			{
+				InitiallySelected = NameItem;
+				break;
+			}
+		}
+	}
+
+	return InitiallySelected;
+}
+
+void SNameListPicker::UpdateSelectedOption()
+{
+	if (TSharedPtr<SComboBox<TSharedPtr<FName>>> PinnedButton = PickerComboButton.Pin())
+	{
+		PinnedButton->SetSelectedItem(GetSelectedItemFromCurrentValue());
+	}
+}
+
 FText SNameListPicker::GetCurrentNameLabel() const
 {
 	const bool bHasMultipleValues = HasMultipleValuesAttribute.Get();
diff --git a/Engine/Plugins/Experimental/DigitalMultiplex/DMXEngine/Source/DMXEditor/Public/Widgets/SNameListPicker.h b/Engine/Plugins/Experimental/DigitalMultiplex/DMXEngine/Source/DMXEditor/Public/Widgets/SNameListPicker.h
--- a/Engine/Plugins/Experimental/DigitalMultiplex/DMXEngine/Source/DMXEditor/Public/Widgets/SNameListPicker.h
+++ b/Engine/Plugins/Experimental/DigitalMultiplex/DMXEngine/Source/DMXEditor/Public/Widgets/SNameListPicker.h
@@ -61,6 +61,7 @@ public:
 
 	/**  Slate widget construction method */
 	void Construct(const FArguments& InArgs);
+
 	virtual ~SNameListPicker();
 
 private:
@@ -73,6 +74,16 @@ private:
 	/** Handles a selection change from the combo box */
 	void HandleSelectionChanged(const TSharedPtr<FName> Item, ESelectInfo::Type SelectInfo);
 
+	TSharedPtr<FName> GetSelectedItemFromCurrentValue() const;
+
+	/**
+	 * Workaround to keep the correct option highlighted in the dropdown menu.
+	 * When code changes the current value of the property this button represents, it's possible
+	 * that the button will keep the previous value highlighted. So we set the currently highlighted
+	 * option every time the menu is opened.
+	 */
+	void UpdateSelectedOption();
+
 private:
 	TWeakPtr<SComboBox<TSharedPtr<FName>>> PickerComboButton;
 	TAttribute<TArray<FName>> OptionsSourceAttr;
diff --git a/Engine/Plugins/Experimental/DigitalMultiplex/DMXEngine/Source/DMXRuntime/Private/DMXSubsystem.cpp b/Engine/Plugins/Experimental/DigitalMultiplex/DMXEngine/Source/DMXRuntime/Private/DMXSubsystem.cpp
--- a/Engine/Plugins/Experimental/DigitalMultiplex/DMXEngine/Source/DMXRuntime/Private/DMXSubsystem.cpp
+++ b/Engine/Plugins/Experimental/DigitalMultiplex/DMXEngine/Source/DMXRuntime/Private/DMXSubsystem.cpp
@@ -13,6 +13,10 @@
 
 #include "EngineUtils.h"
 #include "UObject/UObjectIterator.h"
+#include "Async/Async.h"
+#include "Engine/Engine.h"
+
+DECLARE_LOG_CATEGORY_CLASS(DMXSubsystemLog, Log, All);
 
 void UDMXSubsystem::SendDMX(FDMXProtocolName SelectedProtocol, UDMXEntityFixturePatch* FixturePatch, TMap<FName, int32> FunctionMap, EDMXSendResult& OutResult)
 {
@@ -23,25 +27,38 @@ void UDMXSubsystem::SendDMX(FDMXProtocolName SelectedProtocol, UDMXEntityFixture
 		IDMXFragmentMap DMXFragmentMap;
 		for (const TPair<FName, int32>& Elem : FunctionMap)
 		{
-			const int32& ActiveMode = FixturePatch->ActiveMode;
-			const FDMXFixtureMode& RelevantMode = FixturePatch->ParentFixtureTypeTemplate->Modes[ActiveMode];
-			for (const FDMXFixtureFunction& Function : RelevantMode.Functions)
+			if (const UDMXEntityFixtureType* ParentType = FixturePatch->ParentFixtureTypeTemplate)
 			{
-				const FName FunctionName(*Function.FunctionName);
-				if (FunctionName == Elem.Key && Function.Channel > 0)
+				if (ParentType->Modes.Num() < 1)
 				{
-					const uint8 ChannelsToAdd = UDMXEntityFixtureType::NumChannelsToOccupy(Function.DataType);
-					const int32& Channel = Function.Channel + (FixturePatch->GetStartingChannel() - 1);
-					const uint32 FinalValue = UDMXEntityFixtureType::ClampValueToDataType(Function.DataType, Elem.Value);
+					UE_LOG(DMXSubsystemLog, Error, TEXT("%S: Tried to use Fixture Patch which Parent Fixture Type has no Modes set up."));
+					return;
+				}
 
-					for (uint8 ChannelIt = 0; ChannelIt < ChannelsToAdd; ChannelIt++)
+				const int32& ActiveMode = FMath::Min(FixturePatch->ActiveMode, ParentType->Modes.Num() - 1);
+				const FDMXFixtureMode& RelevantMode = ParentType->Modes[ActiveMode];
+				for (const FDMXFixtureFunction& Function : RelevantMode.Functions)
+				{
+					const FName FunctionName(*Function.FunctionName);
+					if (FunctionName == Elem.Key)
 					{
-						const uint8 BytesOffset = (ChannelIt) * 8;
-						const uint8 RetVal = FinalValue >> BytesOffset & 0xff;
-
-						DMXFragmentMap.Add(Channel + ChannelIt, RetVal);
+						if (!UDMXEntityFixtureType::IsFunctionInModeRange(Function, RelevantMode, FixturePatch->GetStartingChannel() - 1))
+						{
+							continue;
+						}
+	
+						const int32 Channel = Function.Channel + FixturePatch->GetStartingChannel() - 1;
+	
+						uint32 ChannelValue = 0;
+						uint8* ChannelValueBytes = reinterpret_cast<uint8*>(&ChannelValue);
+						UDMXEntityFixtureType::FunctionValueToBytes(Function, Elem.Value, ChannelValueBytes);
+						
+						const uint8 NumBytesInSignalFormat = UDMXEntityFixtureType::NumChannelsToOccupy(Function.DataType);
+						for (uint8 ChannelIt = 0; ChannelIt < NumBytesInSignalFormat; ++ChannelIt)
+						{
+							DMXFragmentMap.Add(Channel + ChannelIt, ChannelValueBytes[ChannelIt]);
+						}
 					}
-
 				}
 			}
 		}
@@ -62,8 +79,12 @@ void UDMXSubsystem::SendDMX(FDMXProtocolName SelectedProtocol, UDMXEntityFixture
 				const uint32 RemoteUniverse = Universe + Controller->RemoteOffset;
 				if (!UniversesUsed.Contains(RemoteUniverse))
 				{
-					Results.Add(SelectedProtocol.GetProtocol()->SendDMXFragment(Universe + Controller->RemoteOffset, DMXFragmentMap));
-					UniversesUsed.Add(RemoteUniverse); // Avoid setting values in the same Universe more than once
+					IDMXProtocolPtr Protocol = SelectedProtocol.GetProtocol();
+					if (Protocol.IsValid())
+					{
+						Results.Add(Protocol->SendDMXFragment(Universe + Controller->RemoteOffset, DMXFragmentMap));
+						UniversesUsed.Add(RemoteUniverse); // Avoid setting values in the same Universe more than once
+					}
 				}
 			}
 
@@ -95,19 +116,23 @@ void UDMXSubsystem::SendDMXRaw(FDMXProtocolName SelectedProtocol, int32 Universe
 				DMXFragmentMap.Add(Elem.Key, Elem.Value);
 			}
 		}
-		OutResult = SelectedProtocol.GetProtocol()->SendDMXFragment(UniverseIndex, DMXFragmentMap);
+		IDMXProtocolPtr Protocol = SelectedProtocol.GetProtocol();
+		if (Protocol.IsValid())
+		{
+			OutResult = Protocol->SendDMXFragmentCreate(UniverseIndex, DMXFragmentMap);
+		}
 	}
 }
 
-void UDMXSubsystem::GetAllFixturesOfType(const UDMXLibrary* DMXLibrary, const FName& FixtureType, TArray<UDMXEntityFixturePatch*>& OutResult)
+void UDMXSubsystem::GetAllFixturesOfType(const FDMXEntityFixtureTypeRef& FixtureType, TArray<UDMXEntityFixturePatch*>& OutResult)
 {
 	OutResult.Reset();
 
-	if (DMXLibrary != nullptr)
+	if (const UDMXEntityFixtureType* FixtureTypeObj = FixtureType.GetFixtureType())
 	{
-		DMXLibrary->ForEachEntityOfType<UDMXEntityFixturePatch>([&](UDMXEntityFixturePatch* Fixture)
+		FixtureTypeObj->GetParentLibrary()->ForEachEntityOfType<UDMXEntityFixturePatch>([&](UDMXEntityFixturePatch* Fixture)
 			{
-				if (Fixture->ParentFixtureTypeTemplate != nullptr && Fixture->ParentFixtureTypeTemplate->GetDisplayName() == FixtureType.ToString())
+				if (Fixture->ParentFixtureTypeTemplate == FixtureTypeObj)
 				{
 					OutResult.Add(Fixture);
 				}
@@ -174,13 +199,20 @@ void UDMXSubsystem::GetRawBuffer(FDMXProtocolName SelectedProtocol, int32 Univer
 	DMXBuffer.Reset();
 	if (SelectedProtocol)
 	{
-		TSharedPtr<IDMXProtocolUniverse, ESPMode::ThreadSafe> ProtocolUniverse = SelectedProtocol.GetProtocol()->GetUniverseById(UniverseIndex);
-		if (ProtocolUniverse.IsValid())
+		IDMXProtocolPtr Protocol = SelectedProtocol.GetProtocol();
+		if (Protocol.IsValid())
 		{
-			TSharedPtr<FDMXBuffer> Buffer = ProtocolUniverse.Get()->GetInputDMXBuffer();
-			if (Buffer.IsValid())
+			TSharedPtr<IDMXProtocolUniverse, ESPMode::ThreadSafe> ProtocolUniverse = Protocol->GetUniverseById(UniverseIndex);
+			if (ProtocolUniverse.IsValid())
 			{
-				DMXBuffer = Buffer->GetDMXData();
+				TSharedPtr<FDMXBuffer> Buffer = ProtocolUniverse.Get()->GetInputDMXBuffer();
+				if (Buffer.IsValid())
+				{
+					Buffer->AccessDMXData([&DMXBuffer](TArray<uint8>& InData)
+						{
+							DMXBuffer = InData;
+						});
+				}
 			}
 		}
 	}
@@ -192,33 +224,34 @@ void UDMXSubsystem::GetFixtureFunctions(const UDMXEntityFixturePatch* InFixtureP
 
 	if (InFixturePatch != nullptr)
 	{
-		const int32& ActiveMode = InFixturePatch->ActiveMode;
-		const int32 StartingAddress = InFixturePatch->GetStartingChannel();
-
-		const UDMXEntityFixtureType* FixtureType = InFixturePatch->ParentFixtureTypeTemplate;
-		if (FixtureType != nullptr)
+		if (const UDMXEntityFixtureType* FixtureType = InFixturePatch->ParentFixtureTypeTemplate)
 		{
+			const int32 StartingAddress = InFixturePatch->GetStartingChannel() - 1;
+
+			if (FixtureType->Modes.Num() < 1)
+			{
+				UE_LOG(DMXSubsystemLog, Error, TEXT("%S: Tried to use Fixture Patch which Parent Fixture Type has no Modes set up."));
+				return;
+			}
+			const int32& ActiveMode = FMath::Min(InFixturePatch->ActiveMode, FixtureType->Modes.Num() - 1);
 			const FDMXFixtureMode& CurrentMode = FixtureType->Modes[ActiveMode];
 
 			for (const FDMXFixtureFunction& Function : CurrentMode.Functions)
 			{
-				const int32 IndexValue = (Function.Channel - 1) + (StartingAddress - 1);
-
-				const uint8 ChannelsToAdd = UDMXEntityFixtureType::NumChannelsToOccupy(Function.DataType);
-
-				TArray<uint8> Bytes;
-				for (uint8 ChannelIt = 0; ChannelIt < ChannelsToAdd; ChannelIt++)
+				if (!UDMXEntityFixtureType::IsFunctionInModeRange(Function, CurrentMode, StartingAddress))
 				{
-					const uint8& ByteValue = DMXBuffer[IndexValue + ChannelIt];
-					Bytes.Add(ByteValue);
+					// This function and the following ones are outside the Universe's range.
+					break;
 				}
 
-				uint32 IntVal = 0;
-				for (uint8 ByteIndex = 0; ByteIndex < Bytes.Num() && ByteIndex < 4; ++ByteIndex)
+				const int32 ChannelIndex = Function.Channel - 1 + StartingAddress;
+				if (ChannelIndex >= DMXBuffer.Num())
 				{
-					IntVal += Bytes[ByteIndex] << (ByteIndex * 8);
+					continue;
 				}
-				OutResult.Add(FName(*Function.FunctionName), IntVal);
+				const uint32 ChannelVal = UDMXEntityFixtureType::BytesToFunctionValue(Function, DMXBuffer.GetData() + ChannelIndex);
+
+				OutResult.Add(FName(*Function.FunctionName), ChannelVal);
 			}
 		}
 	}
@@ -235,89 +268,98 @@ bool UDMXSubsystem::GetFunctionsMap(UDMXEntityFixturePatch* InFixturePatch, cons
 
 	if (InFixturePatch == nullptr)
 	{
-		UE_LOG(LogTemp, Warning, TEXT("No FixturePatch"));
+		UE_LOG(DMXSubsystemLog, Warning, TEXT("%S: FixturePatch is null"), __FUNCTION__);
 
 		return false;
 	}
 
-	TSharedPtr<IDMXProtocol> Protocol = SelectedProtocol.GetProtocol();
-	if (!Protocol.IsValid())
+	UDMXEntityFixtureType* TypeTemplate = InFixturePatch->ParentFixtureTypeTemplate;
+	if (TypeTemplate == nullptr)
 	{
-		UE_LOG(LogTemp, Warning, TEXT("Protocol Not Valid"));
+		UE_LOG(DMXSubsystemLog, Warning, TEXT("%S: InFixturePatch '%s' ParentFixtureTypeTemplate is null."), __FUNCTION__, *InFixturePatch->GetDisplayName());
 
 		return false;
 	}
 
-	TSharedPtr<IDMXProtocolUniverse, ESPMode::ThreadSafe> Universe = Protocol->GetUniverseById(InFixturePatch->UniverseID);
-	if (!Universe.IsValid())
+	if (TypeTemplate->Modes.Num() < InFixturePatch->ActiveMode + 1)
 	{
-		UE_LOG(LogTemp, Warning, TEXT("Universe Not Valid"));
+		UE_LOG(DMXSubsystemLog, Warning, TEXT("Wrong ActiveMode %d, Num of modes %d"), InFixturePatch->ActiveMode, TypeTemplate->Modes.Num());
 
 		return false;
 	}
 
-	TSharedPtr<FDMXBuffer> InputDMXBuffer = Universe->GetInputDMXBuffer();
-	if (!InputDMXBuffer.IsValid())
+	IDMXProtocolPtr Protocol = SelectedProtocol.GetProtocol();
+	if (!Protocol.IsValid())
 	{
-		UE_LOG(LogTemp, Warning, TEXT("InputDMXBuffer Not Valid"));
+		UE_LOG(DMXSubsystemLog, Warning, TEXT("Protocol Not Valid"));
 
 		return false;
 	}
 
-	if (InFixturePatch == nullptr)
+	// Search for a controller that is assigned to the selected protocol and matches the Fixture Patch Universe ID
+	const UDMXEntityController* SelectedController = nullptr;
+	for (const UDMXEntityController* Controller : InFixturePatch->GetRelevantControllers())
+	{
+		if (Controller->DeviceProtocol == SelectedProtocol)
+		{
+			SelectedController = Controller;
+			break;
+		}
+	}
+
+	if (SelectedController == nullptr)
 	{
-		UE_LOG(LogTemp, Warning, TEXT("InFixturePatch == nullptr"));
+		UE_LOG(DMXSubsystemLog, Warning, TEXT("%S: The Fixture Patch '%s' is not assigned to any existing Controller's Universe under the '%s' protocol."), __FUNCTION__, *InFixturePatch->GetDisplayName(), *SelectedProtocol.Name.ToString());
+
 		return false;
 	}
 
-	UDMXEntityFixtureType* TypeTemplate = InFixturePatch->ParentFixtureTypeTemplate;
-	if (TypeTemplate == nullptr)
+	TSharedPtr<IDMXProtocolUniverse, ESPMode::ThreadSafe> Universe = Protocol->GetUniverseById(InFixturePatch->UniverseID + SelectedController->RemoteOffset);
+	if (!Universe.IsValid())
 	{
-		UE_LOG(LogTemp, Warning, TEXT("TypeTemplate == nullptr"));
+		UE_LOG(DMXSubsystemLog, Warning, TEXT("Universe Not Valid"));
 
 		return false;
 	}
 
-	if (TypeTemplate->Modes.Num() < InFixturePatch->ActiveMode + 1)
+	TSharedPtr<FDMXBuffer> InputDMXBuffer = Universe->GetInputDMXBuffer();
+	if (!InputDMXBuffer.IsValid())
 	{
-		UE_LOG(LogTemp, Warning, TEXT("Wrong ActiveMode %d, Num of modes %d"), InFixturePatch->ActiveMode, TypeTemplate->Modes.Num());
+		UE_LOG(DMXSubsystemLog, Warning, TEXT("InputDMXBuffer Not Valid"));
 
 		return false;
 	}
 
-	// Read the correct amount of channels for each function and store each in the correct byte of an int32
 	const FDMXFixtureMode& Mode = TypeTemplate->Modes[InFixturePatch->ActiveMode];
-	const TArray<uint8>& DMXData = InputDMXBuffer->GetDMXData();
-	
-	// Unsigned to make sure there won't be any system that changes how the bytes sum is done with int32
-	uint32 ChannelValue = 0;
+	const int32 FixtureChannelStart = InFixturePatch->GetStartingChannel() - 1;
 
 	for (const FDMXFixtureFunction& Function : Mode.Functions)
 	{
 		if (Function.Channel > DMX_MAX_ADDRESS)
 		{
-			UE_LOG(LogTemp, Warning, TEXT("Channel %d is more then %d"), Function.Channel, DMX_MAX_ADDRESS);
+			UE_LOG(DMXSubsystemLog, Warning, TEXT("%S: Function Channel %d is higher than %d"), Function.Channel, DMX_MAX_ADDRESS);
 
 			return false;
 		}
 
-		if (UDMXEntityFixtureType::GetFunctionLastChannel(Function) > Mode.ChannelSpan)
+		if (!UDMXEntityFixtureType::IsFunctionInModeRange(Function, Mode, FixtureChannelStart))
 		{
 			// We reached the functions outside the valid channels for this mode
 			break;
 		}
 
-		// Always reset all bytes to zero on our int, in case the new ones don't use all bytes.
-		// Otherwise, the bytes not used now could have values from the previous function's channels.
-		ChannelValue = 0;
-
-		const int32 NumChannelsToRead = UDMXEntityFixtureType::NumChannelsToOccupy(Function.DataType);
-		for (int32 ChannelIndex = 0; ChannelIndex < NumChannelsToRead && ChannelIndex + Function.Channel < DMXData.Num() && ChannelIndex < 4; ++ChannelIndex)
-		{
-			ChannelValue += DMXData[Function.Channel + ChannelIndex] << ChannelIndex * 8;
-		}
+		InputDMXBuffer->AccessDMXData([&](TArray<uint8>& DMXData)
+			{
+				const int32 FunctionStartIndex = Function.Channel - 1 + FixtureChannelStart;
+				const uint8 FunctionLastIndex = FunctionStartIndex + UDMXEntityFixtureType::NumChannelsToOccupy(Function.DataType) - 1;
+				if (FunctionLastIndex >= DMXData.Num())
+				{
+					return;
+				}
 
-		OutFunctionsMap.Add(*Function.FunctionName, ChannelValue);
+				const uint32 ChannelValue = UDMXEntityFixtureType::BytesToFunctionValue(Function, DMXData.GetData() + FunctionStartIndex);
+				OutFunctionsMap.Add(*Function.FunctionName, ChannelValue);
+			});
 	}
 
 	return true;
@@ -334,6 +376,16 @@ int32 UDMXSubsystem::GetFunctionsValue(const FName& InName, const TMap<FName, in
 	return 0;
 }
 
+/*static*/ UDMXSubsystem* UDMXSubsystem::GetDMXSubsystem_Pure()
+{
+	return GEngine->GetEngineSubsystem<UDMXSubsystem>();
+}
+
+/*static*/ UDMXSubsystem* UDMXSubsystem::GetDMXSubsystem_Callable()
+{
+	return UDMXSubsystem::GetDMXSubsystem_Pure();
+}
+
 TArray<UDMXEntityFixturePatch*> UDMXSubsystem::GetAllFixturesWithTag(const UDMXLibrary* DMXLibrary, FName CustomTag)
 {
 	TArray<UDMXEntityFixturePatch*> FoundPatches;
@@ -446,28 +498,159 @@ TArray<UDMXLibrary*> UDMXSubsystem::GetAllDMXLibraries()
 	return LibrariesFound;
 }
 
-int32 UDMXSubsystem::BytesToInt(const TArray<uint8>& Bytes)
+FORCEINLINE EDMXFixtureSignalFormat SignalFormatFromBytesNum(uint32 InBytesNum)
+{
+	switch (InBytesNum)
+	{
+	case 0:
+		UE_LOG(DMXSubsystemLog, Error, TEXT("%S called with InBytesNum = 0"), __FUNCTION__);
+		return EDMXFixtureSignalFormat::E8Bit;
+	case 1:
+		return EDMXFixtureSignalFormat::E8Bit;
+	case 2:
+		return EDMXFixtureSignalFormat::E16Bit;
+	case 3:
+		return EDMXFixtureSignalFormat::E24Bit;
+	case 4:
+		return EDMXFixtureSignalFormat::E32Bit;
+	default: // InBytesNum is 4 or higher
+		UE_LOG(DMXSubsystemLog, Warning, TEXT("%S called with InBytesNum > 4. Only 4 bytes will be used."), __FUNCTION__);
+		return EDMXFixtureSignalFormat::E32Bit;
+	}
+}
+
+int32 UDMXSubsystem::BytesToInt(const TArray<uint8>& Bytes, bool bUseLSB /*= false*/) const
+{
+	if (Bytes.Num() == 0)
+	{
+		return 0;
+	}
+
+	const EDMXFixtureSignalFormat SignalFormat = SignalFormatFromBytesNum(Bytes.Num());
+	return UDMXEntityFixtureType::BytesToInt(SignalFormat, bUseLSB, Bytes.GetData());
+}
+
+float UDMXSubsystem::BytesToNormalizedValue(const TArray<uint8>& Bytes, bool bUseLSB /*= false*/) const
 {
-	uint32 IntVal = 0;
-	for (uint8 ByteIndex = 0; ByteIndex < Bytes.Num() && ByteIndex < 4; ++ByteIndex)
+	if (Bytes.Num() == 0)
 	{
-		IntVal += Bytes[ByteIndex] << (ByteIndex * 8);
+		return 0;
 	}
-	return IntVal;
+
+	const EDMXFixtureSignalFormat SignalFormat = SignalFormatFromBytesNum(Bytes.Num());
+	return UDMXEntityFixtureType::BytesToNormalizedValue(SignalFormat, bUseLSB, Bytes.GetData());
+}
+
+void UDMXSubsystem::NormalizedValueToBytes(float InValue, EDMXFixtureSignalFormat InSignalFormat, TArray<uint8>& Bytes, bool bUseLSB /*= false*/) const
+{
+	const uint8 NumBytes = UDMXEntityFixtureType::NumChannelsToOccupy(InSignalFormat);
+	// Make sure the array will fit the correct number of bytes
+	Bytes.Reset(NumBytes);
+	Bytes.AddZeroed(NumBytes);
+
+	UDMXEntityFixtureType::NormalizedValueToBytes(InSignalFormat, bUseLSB, InValue, Bytes.GetData());
+}
+
+void UDMXSubsystem::IntValueToBytes(int32 InValue, EDMXFixtureSignalFormat InSignalFormat, TArray<uint8>& Bytes, bool bUseLSB /*= false*/) const
+{
+	const uint8 NumBytes = UDMXEntityFixtureType::NumChannelsToOccupy(InSignalFormat);
+	// Make sure the array will fit the correct number of bytes
+	Bytes.Reset(NumBytes);
+	Bytes.AddZeroed(NumBytes);
+
+	UDMXEntityFixtureType::IntToBytes(InSignalFormat, bUseLSB, InValue, Bytes.GetData());
+}
+
+float UDMXSubsystem::IntToNormalizedValue(int32 InValue, EDMXFixtureSignalFormat InSignalFormat) const
+{
+	return (float)(uint32)(InValue) / UDMXEntityFixtureType::GetDataTypeMaxValue(InSignalFormat);
+}
+
+float UDMXSubsystem::GetNormalizedFunctionValue(UDMXEntityFixturePatch* InFixturePatch, FName InFunctionName, int32 InValue) const
+{
+	if (InFixturePatch == nullptr)
+	{
+		UE_LOG(DMXSubsystemLog, Error, TEXT("%S: InFixturePatch is null!"), __FUNCTION__);
+		return 0.0f;
+	}
+
+	UDMXEntityFixtureType* ParentType = InFixturePatch->ParentFixtureTypeTemplate;
+	if (ParentType == nullptr)
+	{
+		UE_LOG(DMXSubsystemLog, Error, TEXT("%S: InFixturePatch->ParentFixtureTypeTemplate is null!"), __FUNCTION__);
+		return 0.0f;
+	}
+
+	if (ParentType->Modes.Num() == 0)
+	{
+		UE_LOG(DMXSubsystemLog, Error, TEXT("%S: InFixturePatch's Fixture Type has no Modes!"), __FUNCTION__);
+		return 0.0f;
+	}
+
+	if (InFixturePatch->ActiveMode >= InFixturePatch->ParentFixtureTypeTemplate->Modes.Num())
+	{
+		UE_LOG(DMXSubsystemLog, Error, TEXT("%S: InFixturePatch' ActiveMode is not an existing mode from its Fixture Type!"), __FUNCTION__);
+		return 0.0f;
+	}
+
+	const FDMXFixtureMode& Mode = ParentType->Modes[InFixturePatch->ActiveMode];
+
+	// Search for a Function named InFunctionName in the Fixture Type current mode
+	for (const FDMXFixtureFunction& Function : Mode.Functions)
+	{
+		if (Function.FunctionName.Equals(InFunctionName.ToString()))
+		{
+			return IntToNormalizedValue(InValue, Function.DataType);
+		}
+	}
+
+	return -1.0f;
 }
 
 void UDMXSubsystem::Initialize(FSubsystemCollectionBase& Collection)
 {
 	for (const FName& ProtocolName : IDMXProtocol::GetProtocolNames())
 	{
-		if (TSharedPtr<IDMXProtocol> Protocol = IDMXProtocol::Get(ProtocolName))
+		if (IDMXProtocolPtr Protocol = IDMXProtocol::Get(ProtocolName))
 		{
-			Protocol->GetOnUniverseInputUpdate().AddUObject(this, &UDMXSubsystem::BufferReceivedBroadcast);
+			// Using AddLambda instead of AddUObject because the UObject reference is stored
+			// as a TWeakObjectPtr, which gets Garbage Collected when outside PIE and would break
+			// Utility Blueprints that rely on this event and any subsequent PIE sections.
+			FDelegateHandle UniverseUpdateHandle = Protocol->GetOnUniverseInputUpdate().AddLambda([](FName InProtocol, uint16 InUniverseID, const TArray<uint8>& InValues)
+				{
+					// Call Broadcast on GameThread to make sure a listener BP won't change
+					// an Actor's properties outside of it
+					AsyncTask(ENamedThreads::GameThread, [InProtocol, InUniverseID, InValues]()
+						{
+							// If this gets called after FEngineLoop::Exit(), GetEngineSubsystem() can crash
+							if (IsEngineExitRequested())
+							{
+								return;
+							}
+
+							// The subsystem could be invalid by the time this code gets called
+							UDMXSubsystem* DMXSubsystem = GEngine->GetEngineSubsystem<UDMXSubsystem>();
+							if (DMXSubsystem != nullptr && DMXSubsystem->IsValidLowLevelFast())
+							{
+								DMXSubsystem->OnProtocolReceived.Broadcast(FDMXProtocolName(InProtocol), InUniverseID, InValues);
+							}
+						});
+				});
+
+			// Store handles to unbind from the event when this subsystem deinitializes
+			UniverseInputUpdateHandles.Add(ProtocolName, UniverseUpdateHandle);
 		}
 	}
 }
 
-void UDMXSubsystem::BufferReceivedBroadcast(FName Protocol, uint16 UniverseID, const TArray<uint8>& Values)
+void UDMXSubsystem::Deinitialize()
 {
-	OnProtocolReceived.Broadcast(FDMXProtocolName(Protocol), UniverseID, Values);
+	// Unbind from the protocols' universe update event
+	for (const TPair<FName, FDelegateHandle>& It : UniverseInputUpdateHandles)
+	{
+		if (IDMXProtocolPtr Protocol = IDMXProtocol::Get(It.Key))
+		{
+			Protocol->GetOnUniverseInputUpdate().Remove(It.Value);
+		}
+	}
 }
diff --git a/Engine/Plugins/Experimental/DigitalMultiplex/DMXEngine/Source/DMXRuntime/Private/Library/DMXEntity.cpp b/Engine/Plugins/Experimental/DigitalMultiplex/DMXEngine/Source/DMXRuntime/Private/Library/DMXEntity.cpp
--- a/Engine/Plugins/Experimental/DigitalMultiplex/DMXEngine/Source/DMXRuntime/Private/Library/DMXEntity.cpp
+++ b/Engine/Plugins/Experimental/DigitalMultiplex/DMXEngine/Source/DMXRuntime/Private/Library/DMXEntity.cpp
@@ -19,7 +19,6 @@ FString UDMXEntity::GetDisplayName() const
 
 void UDMXEntity::SetName(const FString& InNewName)
 {
-	// TODO check name for validity
 	Name = InNewName;
 }
 
@@ -38,15 +37,6 @@ void UDMXEntity::ReplicateID(UDMXEntity* Other)
 	Id = Other->Id;
 }
 
-#if WITH_EDITOR
-void UDMXEntity::PostDuplicate(EDuplicateMode::Type DuplicateMode)
-{
-	Super::PostDuplicate(DuplicateMode);
-
-	RefreshID();
-}
-#endif // WITH_EDITOR
-
 UDMXEntityUniverseManaged::UDMXEntityUniverseManaged()
 {
 	if (!HasAnyFlags(RF_ClassDefaultObject | RF_ArchetypeObject))
@@ -62,12 +52,11 @@ void UDMXEntityUniverseManaged::PostLoad()
 }
 
 #if WITH_EDITOR
-
-void UDMXEntityUniverseManaged::PostEditChangeChainProperty(FPropertyChangedChainEvent& PropertyChangedEvent)
+void UDMXEntityUniverseManaged::PostEditChangeProperty(FPropertyChangedEvent& PropertyChangedEvent)
 {
 	static const FName NAME_Universes = GET_MEMBER_NAME_CHECKED(UDMXEntityUniverseManaged, Universes);
 
-	Super::PostEditChangeChainProperty(PropertyChangedEvent);
+	Super::PostEditChangeProperty(PropertyChangedEvent);
 
 	if (PropertyChangedEvent.GetPropertyName() == GET_MEMBER_NAME_CHECKED(FDMXUniverse, UniverseNumber) 
 		|| PropertyChangedEvent.GetPropertyName() == NAME_Universes
@@ -76,23 +65,13 @@ void UDMXEntityUniverseManaged::PostEditChangeChainProperty(FPropertyChangedChai
 		// Keep the Universe ID values within the valid range for the current protocol
 		if (DeviceProtocol.IsValid())
 		{
-			const TSharedPtr<IDMXProtocol> Protocol = DeviceProtocol.GetProtocol();
-			const int32 ChangedUniverseIndex = PropertyChangedEvent.GetArrayIndex(NAME_Universes.ToString());
+			const IDMXProtocolPtr Protocol = DeviceProtocol.GetProtocol();
+			const uint32 MinUniverseID = Protocol->GetMinUniverseID();
+			const uint32 MaxUniverseID = Protocol->GetMaxUniverses();
 
-			if (ChangedUniverseIndex != INDEX_NONE)
+			for (FDMXUniverse& Universe : Universes)
 			{
-				FDMXUniverse& ChangedUniverse = Universes[ChangedUniverseIndex];
-				ChangedUniverse.UniverseNumber = FMath::Clamp(ChangedUniverse.UniverseNumber, (uint32)Protocol->GetMinUniverseID(), (uint32)Protocol->GetMaxUniverses());
-			}
-			else
-			{
-				const uint32 MinUniverseID = Protocol->GetMinUniverseID();
-				const uint32 MaxUniverseID = Protocol->GetMaxUniverses();
-
-				for (FDMXUniverse& Universe : Universes)
-				{
-					Universe.UniverseNumber = FMath::Clamp(Universe.UniverseNumber, MinUniverseID, MaxUniverseID);
-				}
+				Universe.UniverseNumber = FMath::Clamp(Universe.UniverseNumber, MinUniverseID, MaxUniverseID);
 			}
 		}
 
@@ -113,7 +92,7 @@ void UDMXEntityUniverseManaged::UpdateProtocolUniverses() const
 {
 	if (DeviceProtocol.IsValid())
 	{
-		TSharedPtr<IDMXProtocol> DMXProtocol = DeviceProtocol;
+		IDMXProtocolPtr DMXProtocol = DeviceProtocol;
 		DMXProtocol->CollectUniverses(Universes);
 	}
 }
diff --git a/Engine/Plugins/Experimental/DigitalMultiplex/DMXEngine/Source/DMXRuntime/Private/Library/DMXEntityController.cpp b/Engine/Plugins/Experimental/DigitalMultiplex/DMXEngine/Source/DMXRuntime/Private/Library/DMXEntityController.cpp
--- a/Engine/Plugins/Experimental/DigitalMultiplex/DMXEngine/Source/DMXRuntime/Private/Library/DMXEntityController.cpp
+++ b/Engine/Plugins/Experimental/DigitalMultiplex/DMXEngine/Source/DMXRuntime/Private/Library/DMXEntityController.cpp
@@ -20,13 +20,19 @@ void UDMXEntityController::PostEditChangeProperty(FPropertyChangedEvent& Propert
 	Super::PostEditChangeProperty(PropertyChangedEvent);
 }
 
+void UDMXEntityController::PostLoad()
+{
+	Super::PostLoad();
+	ValidateRangeValues();
+}
+
 #endif // WITH_EDITOR
 
 void UDMXEntityController::ValidateRangeValues()
 {
 	if (DeviceProtocol.IsValid())
 	{
-		const TSharedPtr<IDMXProtocol>& Protocol = DeviceProtocol;
+		const IDMXProtocolPtr& Protocol = DeviceProtocol;
 		const int32 MinUniverseID = Protocol->GetMinUniverseID();
 		const int32 MaxUniverseID = Protocol->GetMaxUniverses();
 
diff --git a/Engine/Plugins/Experimental/DigitalMultiplex/DMXEngine/Source/DMXRuntime/Private/Library/DMXEntityFixturePatch.cpp b/Engine/Plugins/Experimental/DigitalMultiplex/DMXEngine/Source/DMXRuntime/Private/Library/DMXEntityFixturePatch.cpp
--- a/Engine/Plugins/Experimental/DigitalMultiplex/DMXEngine/Source/DMXRuntime/Private/Library/DMXEntityFixturePatch.cpp
+++ b/Engine/Plugins/Experimental/DigitalMultiplex/DMXEngine/Source/DMXRuntime/Private/Library/DMXEntityFixturePatch.cpp
@@ -330,7 +330,8 @@ bool UDMXEntityFixturePatch::IsInControllersRange(const TArray<UDMXEntityControl
 }
 
 UDMXEntityFixturePatch::UDMXEntityFixturePatch()
-	: bAutoAssignAddress(true)
+	: UniverseID(1)
+	, bAutoAssignAddress(true)
 	, ManualStartingAddress(1)
 	, AutoStartingAddress(1)
 	, ActiveMode(0)
diff --git a/Engine/Plugins/Experimental/DigitalMultiplex/DMXEngine/Source/DMXRuntime/Private/Library/DMXEntityFixtureType.cpp b/Engine/Plugins/Experimental/DigitalMultiplex/DMXEngine/Source/DMXRuntime/Private/Library/DMXEntityFixtureType.cpp
--- a/Engine/Plugins/Experimental/DigitalMultiplex/DMXEngine/Source/DMXRuntime/Private/Library/DMXEntityFixtureType.cpp
+++ b/Engine/Plugins/Experimental/DigitalMultiplex/DMXEngine/Source/DMXRuntime/Private/Library/DMXEntityFixtureType.cpp
@@ -6,6 +6,7 @@
 #include "Library/DMXImport.h"
 #include "Library/DMXImportGDTF.h"
 
+#if WITH_EDITOR
 void UDMXEntityFixtureType::SetModesFromDMXImport(UDMXImport* DMXImportAsset)
 {
 	if (DMXImportAsset == nullptr || !DMXImportAsset->IsValidLowLevelFast())
@@ -29,18 +30,70 @@ void UDMXEntityFixtureType::SetModesFromDMXImport(UDMXImport* DMXImportAsset)
 			FDMXFixtureMode& Mode = Modes[Modes.Emplace()];
 			Mode.ModeName = AssetMode.Name.ToString();
 
+			// We'll keep the functions addresses from the GDTF file.
+			// For that we need to keep track of the latest occupied address after adding each function.
+			int32 LastOccupiedAddress = 0;
+
 			for (const FDMXImportGDTFDMXChannel& ModeChannel : AssetMode.DMXChannels)
 			{
 				FDMXFixtureFunction& Function = Mode.Functions[Mode.Functions.Emplace()];
 				Function.FunctionName = ModeChannel.LogicalChannel.ChannelFunction.Name.ToString();
-				SetFunctionSize(Mode, Function, ModeChannel.Offset.Num());
 				Function.DefaultValue = ModeChannel.Default.Value;
+
+				if (ModeChannel.Offset.Num() > 0)
+				{
+					// Compute number of used addresses in the function as the interval
+					// between the lowest and highest addresses (inclusive)
+					int32 AddressMin = DMX_MAX_ADDRESS;
+					int32 AddressMax = 0;
+					for (const int32& Address : ModeChannel.Offset)
+					{
+						AddressMin = FMath::Min(AddressMin, Address);
+						AddressMax = FMath::Max(AddressMax, Address);
+					}
+					const int32 NumUsedAddresses = FMath::Clamp(AddressMax - AddressMin + 1, 1, DMX_MAX_FUNCTION_SIZE);
+
+					SetFunctionSize(Function, NumUsedAddresses);
+
+					// AddressMin is the first address this function occupies. If it's not 1 after the
+					// latest occupied channel, this function is offset, skipping some addresses.
+					if (AddressMin > LastOccupiedAddress + 1)
+					{
+						Function.ChannelOffset = AddressMin - LastOccupiedAddress - 1;
+					}
+
+					// Offsets represent the value bytes in MSB format. If they are in reverse order,
+					// it means this Function uses LSB format.
+					// We need at least 2 offsets to compare. Otherwise, we leave the function as MSB,
+					// which is most Fixtures' standard bit format.
+					if (ModeChannel.Offset.Num() > 1)
+					{
+						Function.bUseLSBMode = ModeChannel.Offset[0] > ModeChannel.Offset[1];
+					}
+					else
+					{
+						Function.bUseLSBMode = false;
+					}
+
+					// Update occupied addresses
+					LastOccupiedAddress += Function.ChannelOffset + NumChannelsToOccupy(Function.DataType);
+				}
+				else
+				{
+					SetFunctionSize(Function, 1);
+
+					// Update occupied addresses
+					++LastOccupiedAddress;
+				}
 			}
+
+			// Compute mode channel span from functions' addresses and sizes
+			UpdateModeChannelProperties(Mode);
 		}
 	}
 }
 
-void UDMXEntityFixtureType::SetFunctionSize(FDMXFixtureMode& InMode, FDMXFixtureFunction& InFunction, uint8 Size)
+void UDMXEntityFixtureType::SetFunctionSize(FDMXFixtureFunction& InFunction, uint8 Size)
 {
 	//Get New Data Type
 	EDMXFixtureSignalFormat NewDataType;
@@ -62,41 +115,26 @@ void UDMXEntityFixtureType::SetFunctionSize(FDMXFixtureMode& InMode, FDMXFixture
 		break;
 	}
 
-#if WITH_EDITOR
-	//Update UI
 	InFunction.DataType = NewDataType;
-	UpdateModeChannelProperties(InMode);
 	ClampDefaultValue(InFunction);
-#endif // WITH_EDITOR
 }
 
+#endif // WITH_EDITOR
+
 uint8 UDMXEntityFixtureType::GetFunctionLastChannel(const FDMXFixtureFunction& Function)
 {
 	return Function.Channel + UDMXEntityFixtureType::NumChannelsToOccupy(Function.DataType) - 1;
 }
 
+bool UDMXEntityFixtureType::IsFunctionInModeRange(const FDMXFixtureFunction& InFunction, const FDMXFixtureMode& InMode, int32 ChannelOffset /*= 0*/)
+{
+	const int32 LastChannel = GetFunctionLastChannel(InFunction);
+	return LastChannel <= InMode.ChannelSpan && LastChannel + ChannelOffset <= DMX_MAX_ADDRESS;
+}
+
 void UDMXEntityFixtureType::ClampDefaultValue(FDMXFixtureFunction& InFunction)
 {
-	int64& DefaultValue = InFunction.DefaultValue;
-	switch (InFunction.DataType)
-	{
-	case EDMXFixtureSignalFormat::E8BitSubFunctions:
-	case EDMXFixtureSignalFormat::E8Bit:
-		DefaultValue = FMath::Clamp(DefaultValue, (int64)0, (int64)MAX_uint8);
-		break;
-	case EDMXFixtureSignalFormat::E16Bit:
-		DefaultValue = FMath::Clamp(DefaultValue, (int64)0, (int64)MAX_uint16);
-		break;
-	case EDMXFixtureSignalFormat::E24Bit:
-		DefaultValue = FMath::Clamp(DefaultValue, (int64)0, (int64)0xFFFFFF);
-		break;
-	case EDMXFixtureSignalFormat::E32Bit:
-		DefaultValue = FMath::Clamp(DefaultValue, (int64)0, (int64)MAX_uint32);
-		break;
-	default:
-		checkNoEntry();
-		break;
-	}
+	InFunction.DefaultValue = ClampValueToDataType(InFunction.DataType, FMath::Min(InFunction.DefaultValue, (int64)MAX_uint32));
 }
 
 uint8 UDMXEntityFixtureType::NumChannelsToOccupy(EDMXFixtureSignalFormat DataType)
@@ -145,6 +183,116 @@ uint32 UDMXEntityFixtureType::ClampValueToDataType(EDMXFixtureSignalFormat DataT
 	return InValue;
 }
 
+uint32 UDMXEntityFixtureType::GetDataTypeMaxValue(EDMXFixtureSignalFormat DataType)
+{
+	switch (DataType)
+	{
+	case EDMXFixtureSignalFormat::E8BitSubFunctions:
+	case EDMXFixtureSignalFormat::E8Bit:
+		return MAX_uint8;
+	case EDMXFixtureSignalFormat::E16Bit:
+		return MAX_uint16;
+	case EDMXFixtureSignalFormat::E24Bit:
+		return 0xFFFFFF;
+	case EDMXFixtureSignalFormat::E32Bit:
+		return MAX_uint32;
+	default:
+		checkNoEntry();
+		return 0;
+	}
+}
+
+void UDMXEntityFixtureType::FunctionValueToBytes(const FDMXFixtureFunction& InFunction, uint32 InValue, uint8* OutBytes)
+{
+	IntToBytes(InFunction.DataType, InFunction.bUseLSBMode, InValue, OutBytes);
+}
+
+void UDMXEntityFixtureType::IntToBytes(EDMXFixtureSignalFormat InSignalFormat, bool bUseLSB, uint32 InValue, uint8* OutBytes)
+{
+	// Make sure the input value is in the valid range for the data type
+	InValue = ClampValueToDataType(InSignalFormat, InValue);
+
+	// Number of bytes we'll have to manipulate
+	const uint8 NumBytes = NumChannelsToOccupy(InSignalFormat);
+
+	if (NumBytes == 1)
+	{
+		OutBytes[0] = (uint8)InValue;
+		return;
+	}
+
+	// To avoid branching in the loop, we'll decide before it on which byte to start
+	// and which direction to go, depending on the Function's bit endianness.
+	const int8 ByteIndexStep = bUseLSB ? 1 : -1;
+	int8 OutByteIndex = bUseLSB ? 0 : NumBytes - 1;
+
+	for (uint8 ValueByte = 0; ValueByte < NumBytes; ++ValueByte)
+	{
+		OutBytes[OutByteIndex] = InValue >> 8 * ValueByte & 0xFF;
+		OutByteIndex += ByteIndexStep;
+	}
+}
+
+uint32 UDMXEntityFixtureType::BytesToFunctionValue(const FDMXFixtureFunction& InFunction, const uint8* InBytes)
+{
+	return BytesToInt(InFunction.DataType, InFunction.bUseLSBMode, InBytes);
+}
+
+uint32 UDMXEntityFixtureType::BytesToInt(EDMXFixtureSignalFormat InSignalFormat, bool bUseLSB, const uint8* InBytes)
+{
+	// Number of bytes we'll read
+	const uint8 NumBytes = NumChannelsToOccupy(InSignalFormat);
+
+	if (NumBytes == 1)
+	{
+		return *InBytes;
+	}
+
+	// To avoid branching in the loop, we'll decide before it on which byte to start
+	// and which direction to go, depending on the Function's bit endianness.
+	const int8 ByteIndexStep = bUseLSB ? 1 : -1;
+	int8 InByteIndex = bUseLSB ? 0 : NumBytes - 1;
+
+	uint32 Result = 0;
+	for (uint8 ValueByte = 0; ValueByte < NumBytes; ++ValueByte)
+	{
+		Result += InBytes[InByteIndex] << 8 * ValueByte;
+		InByteIndex += ByteIndexStep;
+	}
+
+	return Result;
+}
+
+void UDMXEntityFixtureType::FunctionNormalizedValueToBytes(const FDMXFixtureFunction& InFunction, float InValue, uint8* OutBytes)
+{
+	NormalizedValueToBytes(InFunction.DataType, InFunction.bUseLSBMode, InValue, OutBytes);
+}
+
+void UDMXEntityFixtureType::NormalizedValueToBytes(EDMXFixtureSignalFormat InSignalFormat, bool bUseLSB, float InValue, uint8* OutBytes)
+{
+	// Make sure InValue is in the range [0.0 ... 1.0]
+	InValue = FMath::Clamp(InValue, 0.0f, 1.0f);
+
+	const uint32 IntValue = GetDataTypeMaxValue(InSignalFormat) * InValue;
+
+	// Get the individual bytes from the computed IntValue
+	IntToBytes(InSignalFormat, bUseLSB, IntValue, OutBytes);
+}
+
+float UDMXEntityFixtureType::BytesToFunctionNormalizedValue(const FDMXFixtureFunction& InFunction, const uint8* InBytes)
+{
+	return BytesToNormalizedValue(InFunction.DataType, InFunction.bUseLSBMode, InBytes);
+}
+
+float UDMXEntityFixtureType::BytesToNormalizedValue(EDMXFixtureSignalFormat InSignalFormat, bool bUseLSB, const uint8* InBytes)
+{
+	// Get the value represented by the individual bytes
+	const float Value = BytesToInt(InSignalFormat, bUseLSB, InBytes);
+
+	// Normalize it
+	return Value / GetDataTypeMaxValue(InSignalFormat);
+}
+
 #if WITH_EDITOR
 
 void UDMXEntityFixtureType::PostEditChangeChainProperty(FPropertyChangedChainEvent& PropertyChangedEvent)
diff --git a/Engine/Plugins/Experimental/DigitalMultiplex/DMXEngine/Source/DMXRuntime/Public/DMXSubsystem.h b/Engine/Plugins/Experimental/DigitalMultiplex/DMXEngine/Source/DMXRuntime/Public/DMXSubsystem.h
--- a/Engine/Plugins/Experimental/DigitalMultiplex/DMXEngine/Source/DMXRuntime/Public/DMXSubsystem.h
+++ b/Engine/Plugins/Experimental/DigitalMultiplex/DMXEngine/Source/DMXRuntime/Public/DMXSubsystem.h
@@ -36,8 +36,8 @@ public:
 	void SendDMXRaw(FDMXProtocolName SelectedProtocol, int32 UniverseIndex, TMap<int32, uint8> ChannelValuesMap, EDMXSendResult& OutResult);
 
 	/**  Return reference to array of Fixture Patch objects of a given type. */
-	UFUNCTION(BlueprintCallable, Category = "DMX", meta = (BlueprintInternalUseOnly = "true", AutoCreateRefTerm = "FixtureType"))
-	void GetAllFixturesOfType(const UDMXLibrary* DMXLibrary, const FName& FixtureType, TArray<UDMXEntityFixturePatch*>& OutResult);
+	UFUNCTION(BlueprintCallable, Category = "DMX", meta = (AutoCreateRefTerm = "FixtureType"))
+	void GetAllFixturesOfType(const FDMXEntityFixtureTypeRef& FixtureType, TArray<UDMXEntityFixturePatch*>& OutResult);
 
 	/**  Return reference to array of Fixture Patch objects of a given category. */
 	UFUNCTION(BlueprintCallable, Category = "DMX")
@@ -91,9 +91,58 @@ public:
 	UFUNCTION(BlueprintCallable, Category = "DMX")
 	TArray<UDMXLibrary*> GetAllDMXLibraries();
 
-	/**  Return integer given an array of bytes. The first 4 bytes in the array will be used for the conversion. */
-	UFUNCTION(BlueprintCallable, Category = "DMX")
-	int32 BytesToInt(const TArray<uint8>& Bytes);
+	/**
+	 * Return integer given an array of bytes. Up to the first 4 bytes in the array will be used for the conversion.
+	 * @param bUseLSB	Least Significant Byte mode makes the individual bytes (channels) of the function be
+	 *					interpreted with the first bytes being the lowest part of the number.
+	 *					Most Fixtures use MSB (Most Significant Byte).
+	 */
+	UFUNCTION(BlueprintPure, Category = "DMX")
+	int32 BytesToInt(const TArray<uint8>& Bytes, bool bUseLSB = false) const;
+
+	/**
+	 * Return normalized value given an array of bytes. Up to the first 4 bytes in the array will be used for the conversion.
+	 * @param bUseLSB	Least Significant Byte mode makes the individual bytes (channels) of the function be
+	 *					interpreted with the first bytes being the lowest part of the number.
+	 *					Most Fixtures use MSB (Most Significant Byte).
+	 */
+	UFUNCTION(BlueprintPure, Category = "DMX", meta = (Keywords = "float"))
+	float BytesToNormalizedValue(const TArray<uint8>& Bytes, bool bUseLSB = false) const;
+
+	/**
+	 * Return the Bytes format of Value in the desired Signal Format.
+	 * @param bUseLSB	Least Significant Byte mode makes the individual bytes (channels) of the function be
+	 *					interpreted with the first bytes being the lowest part of the number.
+	 *					Most Fixtures use MSB (Most Significant Byte).
+	 */
+	UFUNCTION(BlueprintPure, Category = "DMX", meta = (Keywords = "float"))
+	void NormalizedValueToBytes(float InValue, EDMXFixtureSignalFormat InSignalFormat, TArray<uint8>& Bytes, bool bUseLSB = false) const;
+
+	/**
+	 * Return the Bytes format of Value in the desired Signal Format.
+	 * @param bUseLSB	Least Significant Byte mode makes the individual bytes (channels) of the function be
+	 *					interpreted with the first bytes being the lowest part of the number.
+	 *					Most Fixtures use MSB (Most Significant Byte).
+	 */
+	UFUNCTION(BlueprintPure, Category = "DMX")
+	void IntValueToBytes(int32 InValue, EDMXFixtureSignalFormat InSignalFormat, TArray<uint8>& Bytes, bool bUseLSB = false) const;
+
+	/**
+	 * Return the normalized value of an Int value from the specified Signal Format.
+	 * @param bUseLSB	Least Significant Byte mode makes the individual bytes (channels) of the function be
+	 *					interpreted with the first bytes being the lowest part of the number.
+	 *					Most Fixtures use MSB (Most Significant Byte).
+	 */
+	UFUNCTION(BlueprintPure, Category = "DMX", meta = (Keywords = "float"))
+	float IntToNormalizedValue(int32 InValue, EDMXFixtureSignalFormat InSignalFormat) const;
+
+	/**
+	 * Return the normalized value of an Int value from a Fixture Patch function.
+	 * @return	The normalized value of the passed in Int using the Function's signal format.
+	 *			-1.0 if the Function is not found in the Fixture Patch.
+	 */
+	UFUNCTION(BlueprintPure, Category = "DMX", meta = (Keywords = "float"))
+	float GetNormalizedFunctionValue(UDMXEntityFixturePatch* InFixturePatch, FName InFunctionName, int32 InValue) const;
 
 	/**
 	 * Creates a literal UDMXEntityFixturePatch reference
@@ -122,18 +171,27 @@ public:
 	UFUNCTION(BlueprintCallable, meta = (BlueprintInternalUseOnly = "TRUE", AutoCreateRefTerm = "InName, InFunctionsMap"), Category = "DMX")
 	int32 GetFunctionsValue(const FName& InName, const TMap<FName, int32>& InFunctionsMap);
 
+	/** Get a DMX Subsystem, pure version */
+	UFUNCTION(BlueprintPure, Category = "DMX Subsystem", meta = (BlueprintInternalUseOnly = "true"))
+	static UDMXSubsystem* GetDMXSubsystem_Pure();
+
+	/** Get a DMX Subsystem, callable version */
+	UFUNCTION(BlueprintCallable, Category = "DMX Subsystem", meta = (BlueprintInternalUseOnly = "true"))
+	static UDMXSubsystem* GetDMXSubsystem_Callable();
+
 	UPROPERTY(BlueprintAssignable, Category = "DMX")
 	FProtocolReceivedDelegate OnProtocolReceived;
 
 public:
 	//~ USubsystem interface begin
 	virtual void Initialize(FSubsystemCollectionBase& Collection) override;
+	virtual void Deinitialize() override;
 	//~ USubsystem interface end
 
 private:
-
-	// This function is a delegate for when protocols have input updates
-	UFUNCTION()
-	void BufferReceivedBroadcast(FName Protocol, uint16 UniverseID, const TArray<uint8>& Values);
-
+	/**
+	 * Stores DelegateHandles for each Protocol's UniverseInputUpdate event.
+	 * That way we can unbind them when this subsystem is being destroyed and prevent crashes.
+	 */
+	TMap<FName, FDelegateHandle> UniverseInputUpdateHandles;
 };
diff --git a/Engine/Plugins/Experimental/DigitalMultiplex/DMXEngine/Source/DMXRuntime/Public/Library/DMXEntity.h b/Engine/Plugins/Experimental/DigitalMultiplex/DMXEngine/Source/DMXRuntime/Public/Library/DMXEntity.h
--- a/Engine/Plugins/Experimental/DigitalMultiplex/DMXEngine/Source/DMXRuntime/Public/Library/DMXEntity.h
+++ b/Engine/Plugins/Experimental/DigitalMultiplex/DMXEngine/Source/DMXRuntime/Public/Library/DMXEntity.h
@@ -58,16 +58,12 @@ public:
 	/** Copy another Entity's ID. Used when copying, to not lose the original Entity's reference  */
 	void ReplicateID(UDMXEntity* Other);
 
-#if WITH_EDITOR
-	virtual void PostDuplicate(EDuplicateMode::Type DuplicateMode) override;
-#endif // WITH_EDITOR
-
 protected:
 	UPROPERTY()
 	TWeakObjectPtr<UDMXLibrary> ParentLibrary;
 	
 	/** Uniquely identifies the parameter, used for fixing up Blueprints that reference this Entity when renaming. */
-	UPROPERTY()
+	UPROPERTY(DuplicateTransient)
 	FGuid Id;
 };
 
@@ -77,11 +73,8 @@ class DMXRUNTIME_API UDMXEntityUniverseManaged
 	: public UDMXEntity
 {
 	GENERATED_BODY()
-
 public:
 
-	UDMXEntityUniverseManaged();
-
 	UPROPERTY(EditAnywhere, Category = "Communication Properties", meta = (DisplayName = "Protocol", DisplayPriority = "1"))
 	FDMXProtocolName DeviceProtocol;
 
@@ -89,10 +82,12 @@ public:
 	TArray<FDMXUniverse> Universes;
 
 public:
+	UDMXEntityUniverseManaged();
+
 	virtual void PostLoad() override;
 
 #if WITH_EDITOR
-	virtual void PostEditChangeChainProperty(FPropertyChangedChainEvent& PropertyChangedEvent) override;
+	virtual void PostEditChangeProperty(FPropertyChangedEvent& PropertyChangedEvent) override;
 #endif // WITH_EDITOR
 
 
diff --git a/Engine/Plugins/Experimental/DigitalMultiplex/DMXEngine/Source/DMXRuntime/Public/Library/DMXEntityController.h b/Engine/Plugins/Experimental/DigitalMultiplex/DMXEngine/Source/DMXRuntime/Public/Library/DMXEntityController.h
--- a/Engine/Plugins/Experimental/DigitalMultiplex/DMXEngine/Source/DMXRuntime/Public/Library/DMXEntityController.h
+++ b/Engine/Plugins/Experimental/DigitalMultiplex/DMXEngine/Source/DMXRuntime/Public/Library/DMXEntityController.h
@@ -56,9 +56,7 @@ public:
 public:
 	UDMXEntityController()
 		: UniverseLocalNum(1)
-	{
-		ValidateRangeValues();
-	}
+	{}
 
 	/** Returns the currently assigned protocol for this controller */
 	UFUNCTION(BlueprintPure, Category = "DMX")
@@ -66,6 +64,7 @@ public:
 
 #if WITH_EDITOR
 	virtual void PostEditChangeProperty(FPropertyChangedEvent& PropertyChangedEvent) override;
+	virtual void PostLoad() override;
 #endif // WITH_EDITOR
 
 protected:
diff --git a/Engine/Plugins/Experimental/DigitalMultiplex/DMXEngine/Source/DMXRuntime/Public/Library/DMXEntityFixtureType.h b/Engine/Plugins/Experimental/DigitalMultiplex/DMXEngine/Source/DMXRuntime/Public/Library/DMXEntityFixtureType.h
--- a/Engine/Plugins/Experimental/DigitalMultiplex/DMXEngine/Source/DMXRuntime/Public/Library/DMXEntityFixtureType.h
+++ b/Engine/Plugins/Experimental/DigitalMultiplex/DMXEngine/Source/DMXRuntime/Public/Library/DMXEntityFixtureType.h
@@ -62,6 +62,21 @@ struct DMXRUNTIME_API FDMXFixtureFunction
 	UPROPERTY(EditAnywhere, BlueprintReadOnly, meta = (DisplayPriority = "5"), Category = "DMX")
 	EDMXFixtureSignalFormat DataType;
 
+	/**
+	 * Least Significant Byte mode makes the individual bytes (channels) of the function be
+	 * interpreted with the first bytes being the lowest part of the number.
+	 * 
+	 * E.g., given a 16 bit function with two channel values set to [0, 1],
+	 * they would be interpreted as the binary number 00000001 00000000, which means 256.
+	 * The first byte (0) became the lowest part in binary form and the following byte (1), the highest.
+	 * 
+	 * Most Fixtures use MSB (Most Significant Byte) mode, which interprets bytes as highest first.
+	 * In MSB mode, the example above would be interpreted in binary as 00000000 00000001, which means 1.
+	 * The first byte (0) became the highest part in binary form and the following byte (1), the lowest.
+	 */
+	UPROPERTY(EditAnywhere, BlueprintReadOnly, meta = (DisplayName = "Use LSB Mode", DisplayPriority = "29"), Category = "DMX")
+	bool bUseLSBMode;
+
 	FDMXFixtureFunction()
 		: FunctionName()
 		, Description()
@@ -119,20 +134,43 @@ public:
 	TArray<FDMXFixtureMode> Modes;
 
 public:
+#if WITH_EDITOR
 	UFUNCTION(BlueprintCallable, Category = "Fixture Settings")
 	void SetModesFromDMXImport(UDMXImport* DMXImportAsset);
 
-	static void SetFunctionSize(FDMXFixtureMode& InMode, FDMXFixtureFunction& InFunction, uint8 Size);
+	static void SetFunctionSize(FDMXFixtureFunction& InFunction, uint8 Size);
+#endif // WITH_EDITOR
 
 	/** Gets the last channel occupied by the Function */
 	static uint8 GetFunctionLastChannel(const FDMXFixtureFunction& Function);
 
+	/**
+	 * Return true if a Function's occupied channels are within the Mode's Channel Span.
+	 * Optionally add an offset to the function address
+	 */
+	static bool IsFunctionInModeRange(const FDMXFixtureFunction& InFunction, const FDMXFixtureMode& InMode, int32 ChannelOffset = 0);
+
 	static void ClampDefaultValue(FDMXFixtureFunction& InFunction);
 
 	static uint8 NumChannelsToOccupy(EDMXFixtureSignalFormat DataType);
 
 	static uint32 ClampValueToDataType(EDMXFixtureSignalFormat DataType, uint32 InValue);
 
+	static uint32 GetDataTypeMaxValue(EDMXFixtureSignalFormat DataType);
+
+	//~ Conversions to/from Bytes, Int and Normalized Float values.
+	static void FunctionValueToBytes(const FDMXFixtureFunction& InFunction, uint32 InValue, uint8* OutBytes);
+	static void IntToBytes(EDMXFixtureSignalFormat InSignalFormat, bool bUseLSB, uint32 InValue, uint8* OutBytes);
+
+	static uint32 BytesToFunctionValue(const FDMXFixtureFunction& InFunction, const uint8* InBytes);
+	static uint32 BytesToInt(EDMXFixtureSignalFormat InSignalFormat, bool bUseLSB, const uint8* InBytes);
+
+	static void FunctionNormalizedValueToBytes(const FDMXFixtureFunction& InFunction, float InValue, uint8* OutBytes);
+	static void NormalizedValueToBytes(EDMXFixtureSignalFormat InSignalFormat, bool bUseLSB, float InValue, uint8* OutBytes);
+
+	static float BytesToFunctionNormalizedValue(const FDMXFixtureFunction& InFunction, const uint8* InBytes);
+	static float BytesToNormalizedValue(EDMXFixtureSignalFormat InSignalFormat, bool bUseLSB, const uint8* InBytes);
+
 #if WITH_EDITOR
 	virtual void PostEditChangeChainProperty(FPropertyChangedChainEvent& PropertyChangedEvent) override;
 	virtual void PostEditUndo() override;
diff --git a/Engine/Plugins/Experimental/DigitalMultiplex/DMXProtocol/DMXProtocol.uplugin b/Engine/Plugins/Experimental/DigitalMultiplex/DMXProtocol/DMXProtocol.uplugin
--- a/Engine/Plugins/Experimental/DigitalMultiplex/DMXProtocol/DMXProtocol.uplugin
+++ b/Engine/Plugins/Experimental/DigitalMultiplex/DMXProtocol/DMXProtocol.uplugin
@@ -16,38 +16,17 @@
     {
       "Name": "DMXProtocol",
       "Type": "Runtime",
-      "LoadingPhase": "PreDefault",
-      "WhitelistPlatforms": 
-      [
-        "Win64",
-        "Win32",
-        "Mac",
-        "Linux"
-      ]
+      "LoadingPhase": "PreDefault"
     },
     {
       "Name": "DMXProtocolArtNet",
       "Type": "Runtime",
-      "LoadingPhase": "Default",
-      "WhitelistPlatforms": 
-      [
-        "Win64",
-        "Win32",
-        "Mac",
-        "Linux"
-      ]
+      "LoadingPhase": "Default"
     },
     {
       "Name": "DMXProtocolSACN",
       "Type": "Runtime",
-      "LoadingPhase": "Default",
-      "WhitelistPlatforms": 
-      [
-        "Win64",
-        "Win32",
-        "Mac",
-        "Linux"
-      ]
+      "LoadingPhase": "Default"
     }
   ]
 }
diff --git a/Engine/Plugins/Experimental/DigitalMultiplex/DMXProtocol/Source/DMXProtocol/Private/DMXProtocolModule.cpp b/Engine/Plugins/Experimental/DigitalMultiplex/DMXProtocol/Source/DMXProtocol/Private/DMXProtocolModule.cpp
--- a/Engine/Plugins/Experimental/DigitalMultiplex/DMXProtocol/Source/DMXProtocol/Private/DMXProtocolModule.cpp
+++ b/Engine/Plugins/Experimental/DigitalMultiplex/DMXProtocol/Source/DMXProtocol/Private/DMXProtocolModule.cpp
@@ -19,7 +19,7 @@ FOnNetworkInterfaceChanged IDMXProtocol::OnNetworkInterfaceChanged;
 
 const TCHAR* FDMXProtocolModule::BaseModuleName = TEXT("DMXProtocol");
 
-const TMap<FName, TSharedPtr<IDMXProtocol>>& FDMXProtocolModule::GetProtocols() const
+const TMap<FName, IDMXProtocolPtr>& FDMXProtocolModule::GetProtocols() const
 {
 	return DMXProtocols;
 }
@@ -112,9 +112,9 @@ FDMXProtocolModule& FDMXProtocolModule::Get()
 }
 
 
-TSharedPtr<IDMXProtocol> FDMXProtocolModule::GetProtocol(const FName InProtocolName)
+IDMXProtocolPtr FDMXProtocolModule::GetProtocol(const FName InProtocolName)
 {
-	TSharedPtr<IDMXProtocol>* DMXProtocolPtr = nullptr;
+	IDMXProtocolPtr* DMXProtocolPtr = nullptr;
 	if (!InProtocolName.IsNone())
 	{
 		DMXProtocolPtr = DMXProtocols.Find(InProtocolName);
@@ -126,7 +126,7 @@ TSharedPtr<IDMXProtocol> FDMXProtocolModule::GetProtocol(const FName InProtocolN
 			{
 				UE_LOG_DMXPROTOCOL(Log, TEXT("Creating protocol instance for: %s"), *InProtocolName.ToString());
 
-				TSharedPtr<IDMXProtocol> NewProtocol = (*DMXProtocolFactory)->CreateProtocol(InProtocolName);
+				IDMXProtocolPtr NewProtocol = (*DMXProtocolFactory)->CreateProtocol(InProtocolName);
 				if (NewProtocol.IsValid())
 				{
 					DMXProtocols.Add(InProtocolName, NewProtocol);
@@ -157,7 +157,7 @@ void FDMXProtocolModule::ShutdownDMXProtocol(const FName& ProtocolName)
 {
 	if (!ProtocolName.IsNone())
 	{
-		TSharedPtr<IDMXProtocol> DMXProtocol;
+		IDMXProtocolPtr DMXProtocol;
 		DMXProtocols.RemoveAndCopyValue(ProtocolName, DMXProtocol);
 		if (DMXProtocol.IsValid())
 		{
@@ -172,7 +172,7 @@ void FDMXProtocolModule::ShutdownDMXProtocol(const FName& ProtocolName)
 
 void FDMXProtocolModule::ShutdownAllDMXProtocols()
 {
-	for (TMap<FName, TSharedPtr<IDMXProtocol>>::TIterator It = DMXProtocols.CreateIterator(); It; ++It)
+	for (TMap<FName, IDMXProtocolPtr>::TIterator It = DMXProtocols.CreateIterator(); It; ++It)
 	{
 		It->Value->Shutdown();
 	}
diff --git a/Engine/Plugins/Experimental/DigitalMultiplex/DMXProtocol/Source/DMXProtocol/Private/DMXProtocolTypes.cpp b/Engine/Plugins/Experimental/DigitalMultiplex/DMXProtocol/Source/DMXProtocol/Private/DMXProtocolTypes.cpp
--- a/Engine/Plugins/Experimental/DigitalMultiplex/DMXProtocol/Source/DMXProtocol/Private/DMXProtocolTypes.cpp
+++ b/Engine/Plugins/Experimental/DigitalMultiplex/DMXProtocol/Source/DMXProtocol/Private/DMXProtocolTypes.cpp
@@ -9,7 +9,7 @@ TArray<FName> FDMXProtocolName::GetPossibleValues()
 	return IDMXProtocol::GetProtocolNames();
 }
 
-FDMXProtocolName::FDMXProtocolName(TSharedPtr<IDMXProtocol> InProtocol)
+FDMXProtocolName::FDMXProtocolName(IDMXProtocolPtr InProtocol)
 {
 	if (InProtocol.IsValid())
 	{
@@ -21,7 +21,19 @@ FDMXProtocolName::FDMXProtocolName(const FName& InName)
 	: Name(InName)
 {}
 
-TSharedPtr<IDMXProtocol> FDMXProtocolName::GetProtocol() const
+FDMXProtocolName::FDMXProtocolName()
+{
+	// GetFirstProtocolName depends on the FDMXProtocolModule.
+	// This can be called on CDO creation, when the module might not be available yet.
+	// So we first check if it is available.
+	const IModuleInterface* DMXProtocolModule = FModuleManager::Get().GetModule(FDMXProtocolModule::BaseModuleName);
+	if (DMXProtocolModule != nullptr)
+	{
+		Name = IDMXProtocol::GetFirstProtocolName();
+	}
+}
+
+IDMXProtocolPtr FDMXProtocolName::GetProtocol() const
 {
 	if (Name.IsNone())
 	{
@@ -83,8 +95,22 @@ FName UDMXNameContainersConversions::Conv_DMXFixtureCategoryToName(const FDMXFix
 	return InFixtureCategory.Name;
 }
 
+uint8 FDMXBuffer::GetDMXDataAddress(uint32 InAddress) const
+{
+	FScopeLock BufferLock(&BufferCritSec);
+	return DMXData[InAddress];
+}
+
+void FDMXBuffer::AccessDMXData(TFunctionRef<void(TArray<uint8>&)> InFunction)
+{
+	FScopeLock BufferLock(&BufferCritSec);
+	InFunction(DMXData);
+}
+
 bool FDMXBuffer::SetDMXFragment(const IDMXFragmentMap & InDMXFragment)
 {
+	FScopeLock BufferLock(&BufferCritSec);
+
 	for (const TPair<uint32, uint8>& It : InDMXFragment)
 	{
 		if (It.Key <= (DMX_UNIVERSE_SIZE) && It.Key > 0)
@@ -98,13 +124,15 @@ bool FDMXBuffer::SetDMXFragment(const IDMXFragmentMap & InDMXFragment)
 	}
 
 	// Increase Sequence
-	SequanceID++;
+	SequenceID++;
 
 	return true;
 }
 
 bool FDMXBuffer::SetDMXBuffer(const uint8* InBuffer, uint32 InSize)
 {
+	FScopeLock BufferLock(&BufferCritSec);
+
 	if (InSize <= (DMX_UNIVERSE_SIZE) && InSize > 0)
 	{
 		FMemory::Memcpy(DMXData.GetData(), InBuffer, InSize);
@@ -115,7 +143,7 @@ bool FDMXBuffer::SetDMXBuffer(const uint8* InBuffer, uint32 InSize)
 	}
 
 	// Increase Sequence
-	SequanceID++;
+	SequenceID++;
 
 	return true;
 }
diff --git a/Engine/Plugins/Experimental/DigitalMultiplex/DMXProtocol/Source/DMXProtocol/Private/Tests/DMXProtocolTests.cpp b/Engine/Plugins/Experimental/DigitalMultiplex/DMXProtocol/Source/DMXProtocol/Private/Tests/DMXProtocolTests.cpp
--- a/Engine/Plugins/Experimental/DigitalMultiplex/DMXProtocol/Source/DMXProtocol/Private/Tests/DMXProtocolTests.cpp
+++ b/Engine/Plugins/Experimental/DigitalMultiplex/DMXProtocol/Source/DMXProtocol/Private/Tests/DMXProtocolTests.cpp
@@ -63,10 +63,10 @@ private:
 class FDMXProtocolFactoryTestFactory : public IDMXProtocolFactory
 {
 public:
-	virtual TSharedPtr<IDMXProtocol> CreateProtocol(const FName& ProtocolName) override
+	virtual IDMXProtocolPtr CreateProtocol(const FName& ProtocolName) override
 	{
 		FJsonObject ProtocolSettings;
-		TSharedPtr<IDMXProtocol> ProtocolArtNetPtr = MakeShared<FDMXProtocolTest>(ProtocolName, ProtocolSettings);
+		IDMXProtocolPtr ProtocolArtNetPtr = MakeShared<FDMXProtocolTest, ESPMode::ThreadSafe>(ProtocolName, ProtocolSettings);
 		if (ProtocolArtNetPtr->IsEnabled())
 		{
 			if (!ProtocolArtNetPtr->Init())
@@ -139,7 +139,7 @@ bool FDMXProtocolFactoryTest::RunTest(const FString& Parameters)
 	FDMXProtocolModule& DMXProtocolModule = FModuleManager::GetModuleChecked<FDMXProtocolModule>("DMXProtocol");
 
 	// Store the protocol pointer
-	TSharedPtr<IDMXProtocol> CachedProtocol = nullptr;
+	IDMXProtocolPtr CachedProtocol = nullptr;
 
 	// Try to register 3 times
 	TArray<TUniquePtr<IDMXProtocolFactory>> Factories;
diff --git a/Engine/Plugins/Experimental/DigitalMultiplex/DMXProtocol/Source/DMXProtocol/Public/DMXProtocolCommon.h b/Engine/Plugins/Experimental/DigitalMultiplex/DMXProtocol/Source/DMXProtocol/Public/DMXProtocolCommon.h
--- a/Engine/Plugins/Experimental/DigitalMultiplex/DMXProtocol/Source/DMXProtocol/Public/DMXProtocolCommon.h
+++ b/Engine/Plugins/Experimental/DigitalMultiplex/DMXProtocol/Source/DMXProtocol/Public/DMXProtocolCommon.h
@@ -29,5 +29,7 @@ using IDMXFragmentMap = TMap<uint32, uint8>;
 using FArrayReaderPtr = TSharedPtr<FArrayReader, ESPMode::ThreadSafe>;
 using FDMXPacketPtr = TSharedPtr<FDMXPacket, ESPMode::ThreadSafe>;
 using IDMXProtocolUniversePtr = TSharedPtr<IDMXProtocolUniverse, ESPMode::ThreadSafe>;
+using IDMXProtocolPtr = TSharedPtr<IDMXProtocol, ESPMode::ThreadSafe>;
+using IDMXProtocolPtrWeak = TWeakPtr<IDMXProtocol, ESPMode::ThreadSafe>;
 
 enum class EDMXSendResult : uint8;
diff --git a/Engine/Plugins/Experimental/DigitalMultiplex/DMXProtocol/Source/DMXProtocol/Public/DMXProtocolConstants.h b/Engine/Plugins/Experimental/DigitalMultiplex/DMXProtocol/Source/DMXProtocol/Public/DMXProtocolConstants.h
--- a/Engine/Plugins/Experimental/DigitalMultiplex/DMXProtocol/Source/DMXProtocol/Public/DMXProtocolConstants.h
+++ b/Engine/Plugins/Experimental/DigitalMultiplex/DMXProtocol/Source/DMXProtocol/Public/DMXProtocolConstants.h
@@ -32,5 +32,6 @@ enum EDMUniverseDirection
 #define DMX_MAX_FINALUNIVERSE 32768
 #define DMX_MAX_ADDRESS 512
 #define DMX_MAX_VALUE 255
+#define DMX_MAX_FUNCTION_SIZE ((int32)sizeof(int32))
 
 #define DMX_K2_CATEGORY_NAME "DMX"
\ No newline at end of file
diff --git a/Engine/Plugins/Experimental/DigitalMultiplex/DMXProtocol/Source/DMXProtocol/Public/DMXProtocolModule.h b/Engine/Plugins/Experimental/DigitalMultiplex/DMXProtocol/Source/DMXProtocol/Public/DMXProtocolModule.h
--- a/Engine/Plugins/Experimental/DigitalMultiplex/DMXProtocol/Source/DMXProtocol/Public/DMXProtocolModule.h
+++ b/Engine/Plugins/Experimental/DigitalMultiplex/DMXProtocol/Source/DMXProtocol/Public/DMXProtocolModule.h
@@ -19,13 +19,13 @@ public:
 	 * @param  InProtocolName Name of the requested protocol
 	 * @return Return the pointer to protocol.
 	 */
-	virtual TSharedPtr<IDMXProtocol> GetProtocol(const FName InProtocolName = NAME_None);
+	virtual IDMXProtocolPtr GetProtocol(const FName InProtocolName = NAME_None);
 	
 	/**  Get the reference to all protocol factories map */
 	const TMap<FName, IDMXProtocolFactory*>& GetProtocolFactories() const;
 
 	/**  Get the reference to all protocols map */
-	const TMap<FName, TSharedPtr<IDMXProtocol>>& GetProtocols() const;
+	const TMap<FName, IDMXProtocolPtr>& GetProtocols() const;
 
 	//~ Begin IModuleInterface implementation
 	virtual void StartupModule() override;
@@ -44,6 +44,6 @@ public:
 
 private:
 	TMap<FName, IDMXProtocolFactory*> DMXFactories;
-	TMap<FName, TSharedPtr<IDMXProtocol>> DMXProtocols;
+	TMap<FName, IDMXProtocolPtr> DMXProtocols;
 	TMap<FName, bool> DMXProtocolFailureNotes;
 };
diff --git a/Engine/Plugins/Experimental/DigitalMultiplex/DMXProtocol/Source/DMXProtocol/Public/DMXProtocolTypes.h b/Engine/Plugins/Experimental/DigitalMultiplex/DMXProtocol/Source/DMXProtocol/Public/DMXProtocolTypes.h
--- a/Engine/Plugins/Experimental/DigitalMultiplex/DMXProtocol/Source/DMXProtocol/Public/DMXProtocolTypes.h
+++ b/Engine/Plugins/Experimental/DigitalMultiplex/DMXProtocol/Source/DMXProtocol/Public/DMXProtocolTypes.h
@@ -5,9 +5,11 @@
 #include "DMXProtocolCommon.h"
 #include "DMXProtocolConstants.h"
 #include "DMXProtocolMacros.h"
-#include "Dom/JsonObject.h"
 
+#include "Dom/JsonObject.h"
 #include "Kismet/BlueprintFunctionLibrary.h"
+#include "HAL/CriticalSection.h"
+#include "Templates/Atomic.h"
 
 #include "DMXProtocolTypes.generated.h"
 
@@ -35,7 +37,7 @@ enum class EDMXCommunicationTypes : uint8
 	Broadcast UMETA(DisplayName = "Broadcast")
 };
 
-UENUM()
+UENUM(BlueprintType)
 enum class EDMXFixtureSignalFormat : uint8
 {
 	/** Uses 1 channel (byte) and allows subdivision into sub functions */
@@ -62,14 +64,14 @@ public:
 
 	static TArray<FName> GetPossibleValues();
 
-	FDMXProtocolName() = default;
+	FDMXProtocolName();
 	/** Construct from a protocol name */
 	explicit FDMXProtocolName(const FName& InName);
 	/** Construct from a protocol */
-	FDMXProtocolName(TSharedPtr<IDMXProtocol> InProtocol);
+	FDMXProtocolName(IDMXProtocolPtr InProtocol);
 
 	/** Returns the Protocol this name represents */
-	TSharedPtr<IDMXProtocol> GetProtocol() const;
+	IDMXProtocolPtr GetProtocol() const;
 	/** True if it has a valid Protocol name */
 	bool IsValid() const;
 
@@ -78,15 +80,15 @@ public:
 	operator const FName&() const { return Name; }
 
 	//~ DMXProtocol operators
-	operator TSharedPtr<IDMXProtocol>() { return GetProtocol(); }
-	operator const TSharedPtr<IDMXProtocol>() const { return GetProtocol(); }
+	operator IDMXProtocolPtr() { return GetProtocol(); }
+	operator const IDMXProtocolPtr() const { return GetProtocol(); }
 
 	/** Bool (is valid) operator */
 	operator bool() const { return IsValid(); }
 
 	//~ Comparison operators
 	bool operator==(const FDMXProtocolName& Other) const { return Name == Other.Name; }
-	bool operator==(const TSharedPtr<IDMXProtocol>& Other) const { return GetProtocol().Get() == Other.Get(); }
+	bool operator==(const IDMXProtocolPtr& Other) const { return GetProtocol().Get() == Other.Get(); }
 	bool operator==(const FName& Other) const { return Name == Other; }
 };
 
@@ -150,7 +152,7 @@ struct DMXPROTOCOL_API FDMXUniverse
 	UPROPERTY(EditAnywhere, Category = "DMX")
 	uint32 UniverseNumber;
 
-	UPROPERTY(EditAnywhere, meta = (DisplayName = "Address", ClampMin = 0, ClampMax = 512, UIMin = 0, UIMax = 512), Category = "DMX")
+	UPROPERTY(EditAnywhere, meta = (DisplayName = "Address", ClampMin = 1, ClampMax = 512, UIMin = 1, UIMax = 512), Category = "DMX")
 	uint32 Channel;
 
 	UPROPERTY()
@@ -167,16 +169,23 @@ struct DMXPROTOCOL_API FDMXBuffer
 {
 public:
 	FDMXBuffer()
-		: SequanceID(0)
+		: SequenceID(0)
 	{
 		DMXData.AddZeroed(DMX_UNIVERSE_SIZE);
 	}
 
 	/** @return Gets actual SequanceID  */
-	uint32 GetSequanceID() const { return SequanceID; }
+	uint32 GetSequenceID() const { return SequenceID; }
 
-	/** @return editable DMX buffer */
-	TArray<uint8>& GetDMXData() { return DMXData; }
+	/** @return DMX buffer address value */
+	uint8 GetDMXDataAddress(uint32 InAddress) const;
+
+	/**
+	 * Calls InFunction with a reference to the DMX data buffer passed in on a thread-safe manner.
+	 * This method locks execution and calls InFunction as soon as the buffer can be accessed,
+	 * so it's safe to reference locally scoped variables inside InFunction.
+	 */
+	void AccessDMXData(TFunctionRef<void(TArray<uint8>&)> InFunction);
 
 	/**
 	 * Updates the fragment in the DMX buffer
@@ -201,7 +210,10 @@ private:
 	TArray<uint8> DMXData;
 
 	/** Synchronizations sequence id */
-	uint32 SequanceID;
+	TAtomic<uint32> SequenceID;
+
+	/** Mutex to make sure no two threads write to the buffer concurrently */
+	mutable FCriticalSection BufferCritSec;
 };
 
 struct DMXPROTOCOL_API FDMXPacket
diff --git a/Engine/Plugins/Experimental/DigitalMultiplex/DMXProtocol/Source/DMXProtocol/Public/Interfaces/IDMXProtocol.h b/Engine/Plugins/Experimental/DigitalMultiplex/DMXProtocol/Source/DMXProtocol/Public/Interfaces/IDMXProtocol.h
--- a/Engine/Plugins/Experimental/DigitalMultiplex/DMXProtocol/Source/DMXProtocol/Public/Interfaces/IDMXProtocol.h
+++ b/Engine/Plugins/Experimental/DigitalMultiplex/DMXProtocol/Source/DMXProtocol/Public/Interfaces/IDMXProtocol.h
@@ -37,7 +37,7 @@ public:
 		return DMXProtocolModule.GetProtocolFactories();
 	}
 
-	static const TMap<FName, TSharedPtr<IDMXProtocol>>& GetProtocols()
+	static const TMap<FName, IDMXProtocolPtr>& GetProtocols()
 	{
 		static const FName DMXProtocolModuleName = TEXT("DMXProtocol");
 		FDMXProtocolModule& DMXProtocolModule = FModuleManager::GetModuleChecked<FDMXProtocolModule>(DMXProtocolModuleName);
@@ -69,7 +69,7 @@ public:
 	 * @param  InProtocolName Name of the requested protocol
 	 * @return Return the pointer to protocol
 	 */
-	static TSharedPtr<IDMXProtocol> Get(const FName& ProtocolName = NAME_None)
+	static IDMXProtocolPtr Get(const FName& ProtocolName = NAME_None)
 	{
 		static const FName DMXProtocolModuleName = TEXT("DMXProtocol");
 		FDMXProtocolModule& DMXProtocolModule = FModuleManager::GetModuleChecked<FDMXProtocolModule>(DMXProtocolModuleName);
diff --git a/Engine/Plugins/Experimental/DigitalMultiplex/DMXProtocol/Source/DMXProtocol/Public/Interfaces/IDMXProtocolBase.h b/Engine/Plugins/Experimental/DigitalMultiplex/DMXProtocol/Source/DMXProtocol/Public/Interfaces/IDMXProtocolBase.h
--- a/Engine/Plugins/Experimental/DigitalMultiplex/DMXProtocol/Source/DMXProtocol/Public/Interfaces/IDMXProtocolBase.h
+++ b/Engine/Plugins/Experimental/DigitalMultiplex/DMXProtocol/Source/DMXProtocol/Public/Interfaces/IDMXProtocolBase.h
@@ -5,7 +5,7 @@
 #include "CoreMinimal.h"
 
 class DMXPROTOCOL_API IDMXProtocolBase
-	: public TSharedFromThis<IDMXProtocolBase>
+	: public TSharedFromThis<IDMXProtocolBase, ESPMode::ThreadSafe>
 {
 public:
 	virtual ~IDMXProtocolBase() {}
diff --git a/Engine/Plugins/Experimental/DigitalMultiplex/DMXProtocol/Source/DMXProtocol/Public/Interfaces/IDMXProtocolFactory.h b/Engine/Plugins/Experimental/DigitalMultiplex/DMXProtocol/Source/DMXProtocol/Public/Interfaces/IDMXProtocolFactory.h
--- a/Engine/Plugins/Experimental/DigitalMultiplex/DMXProtocol/Source/DMXProtocol/Public/Interfaces/IDMXProtocolFactory.h
+++ b/Engine/Plugins/Experimental/DigitalMultiplex/DMXProtocol/Source/DMXProtocol/Public/Interfaces/IDMXProtocolFactory.h
@@ -13,5 +13,5 @@ public:
 	 * It creates only one instance of Protocol. There is the condition to check in DMXProtocolModule
 	 * @return Return the pointer to protocol
 	 */
-	virtual TSharedPtr<IDMXProtocol> CreateProtocol(const FName& ProtocolName) = 0;
+	virtual IDMXProtocolPtr CreateProtocol(const FName& ProtocolName) = 0;
 };
diff --git a/Engine/Plugins/Experimental/DigitalMultiplex/DMXProtocol/Source/DMXProtocol/Public/Interfaces/IDMXProtocolUniverse.h b/Engine/Plugins/Experimental/DigitalMultiplex/DMXProtocol/Source/DMXProtocol/Public/Interfaces/IDMXProtocolUniverse.h
--- a/Engine/Plugins/Experimental/DigitalMultiplex/DMXProtocol/Source/DMXProtocol/Public/Interfaces/IDMXProtocolUniverse.h
+++ b/Engine/Plugins/Experimental/DigitalMultiplex/DMXProtocol/Source/DMXProtocol/Public/Interfaces/IDMXProtocolUniverse.h
@@ -9,7 +9,7 @@ class IDMXProtocolUniverse
 public:
 	virtual ~IDMXProtocolUniverse() {}
 
-	virtual TSharedPtr<IDMXProtocol> GetProtocol() const = 0;
+	virtual IDMXProtocolPtr GetProtocol() const = 0;
 	virtual TSharedPtr<FDMXBuffer> GetInputDMXBuffer() const = 0;
 	virtual TSharedPtr<FDMXBuffer> GetOutputDMXBuffer() const = 0;
 	virtual bool SetDMXFragment(const IDMXFragmentMap& DMXFragment) = 0;
diff --git a/Engine/Plugins/Experimental/DigitalMultiplex/DMXProtocol/Source/DMXProtocolArtNet/Private/DMXProtocolArtNet.cpp b/Engine/Plugins/Experimental/DigitalMultiplex/DMXProtocol/Source/DMXProtocolArtNet/Private/DMXProtocolArtNet.cpp
--- a/Engine/Plugins/Experimental/DigitalMultiplex/DMXProtocol/Source/DMXProtocolArtNet/Private/DMXProtocolArtNet.cpp
+++ b/Engine/Plugins/Experimental/DigitalMultiplex/DMXProtocol/Source/DMXProtocolArtNet/Private/DMXProtocolArtNet.cpp
@@ -310,11 +310,20 @@ EDMXSendResult FDMXProtocolArtNet::SendDMXInternal(uint16 UniverseID, uint8 Port
 
 	// ArtNet DMX packet
 	FDMXProtocolArtNetDMXPacket ArtNetDMXPacket;
-	if (DMXBuffer->GetDMXData().Num() == ARTNET_DMX_LENGTH)
-	{
-		FMemory::Memcpy(ArtNetDMXPacket.Data, DMXBuffer->GetDMXData().GetData(), ARTNET_DMX_LENGTH);
-	}
-	else
+	bool bBufferSizeIsWrong = false;
+	DMXBuffer->AccessDMXData([&ArtNetDMXPacket, &bBufferSizeIsWrong](TArray<uint8>& InData)
+		{
+			if (InData.Num() == ARTNET_DMX_LENGTH)
+			{
+				FMemory::Memcpy(ArtNetDMXPacket.Data, InData.GetData(), ARTNET_DMX_LENGTH);
+			}
+			else
+			{
+				bBufferSizeIsWrong = true;
+			}
+		});
+
+	if (bBufferSizeIsWrong)
 	{
 		return EDMXSendResult::ErrorSizeBuffer;
 	}
@@ -503,8 +512,6 @@ bool FDMXProtocolArtNet::TransmitRDM(uint32 InUniverseID, const TArray<uint8>& D
 
 void FDMXProtocolArtNet::OnDataReceived(const FArrayReaderPtr& Buffer)
 {
-	FScopeLock Lock(&OnDataReceivedCS);
-
 	// It will be more handlers
 	switch (ArtNet::GetPacketType(Buffer))
 	{
@@ -558,24 +565,30 @@ bool FDMXProtocolArtNet::HandleDataPacket(const FArrayReaderPtr & Buffer)
 	// Write data to input DMX buffer universe if exists
 	TSharedPtr<FDMXProtocolUniverseArtNet, ESPMode::ThreadSafe> Universe = UniverseManager->GetUniverseById(ArtNetDMXPacket.Universe);
 
+	bool bSetDataSuccessful = false;
 	if (Universe.IsValid() && Universe->GetInputDMXBuffer().IsValid())
 	{
-		// Make sure we copy same amount of data
-		if (Universe->GetInputDMXBuffer()->GetDMXData().Num() == ARTNET_DMX_LENGTH)
-		{
-			Universe->GetInputDMXBuffer()->SetDMXBuffer(ArtNetDMXPacket.Data, ARTNET_DMX_LENGTH);
-
-			OnUniverseInputUpdateEvent.Broadcast(GetProtocolName(), ArtNetDMXPacket.Universe, Universe->GetInputDMXBuffer()->GetDMXData());
-			return true;
-		}
-		else
-		{
-			UE_LOG_DMXPROTOCOL(Error, TEXT("%s:Size of incoming DMX buffer is wrong"), NetworkErrorMessagePrefix);
-			return false;
-		}
+		Universe->GetInputDMXBuffer()->AccessDMXData([this, &ArtNetDMXPacket, &Universe, &bSetDataSuccessful](TArray<uint8>& InData)
+			{
+				// Make sure we copy same amount of data
+				if (InData.Num() == ARTNET_DMX_LENGTH)
+				{
+					Universe->GetInputDMXBuffer()->SetDMXBuffer(ArtNetDMXPacket.Data, ARTNET_DMX_LENGTH);
+					OnUniverseInputUpdateEvent.Broadcast(GetProtocolName(), ArtNetDMXPacket.Universe, InData);
+					bSetDataSuccessful = true;
+				}
+				else
+				{
+					UE_LOG_DMXPROTOCOL(Error, TEXT("%s: Size of incoming DMX buffer is wrong. Expected size: %d. Current: %d")
+						, NetworkErrorMessagePrefix
+						, ARTNET_DMX_LENGTH
+						, InData.Num());
+					bSetDataSuccessful = false;
+				}
+			});
 	}
-
-	return false;
+	
+	return bSetDataSuccessful;
 }
 
 bool FDMXProtocolArtNet::HandleTodRequest(const FArrayReaderPtr & Buffer)
diff --git a/Engine/Plugins/Experimental/DigitalMultiplex/DMXProtocol/Source/DMXProtocolArtNet/Private/DMXProtocolArtNet.h b/Engine/Plugins/Experimental/DigitalMultiplex/DMXProtocol/Source/DMXProtocolArtNet/Private/DMXProtocolArtNet.h
--- a/Engine/Plugins/Experimental/DigitalMultiplex/DMXProtocol/Source/DMXProtocolArtNet/Private/DMXProtocolArtNet.h
+++ b/Engine/Plugins/Experimental/DigitalMultiplex/DMXProtocol/Source/DMXProtocolArtNet/Private/DMXProtocolArtNet.h
@@ -127,8 +127,6 @@ private:
 	TSharedPtr<IDMXProtocolSender> ArtNetSender;
 	TSharedPtr<IDMXProtocolReceiver> ArtNetReceiver;
 
-	mutable FCriticalSection OnDataReceivedCS;
-
 	/** Holds the network socket used to transport packages. */
 	FSocket* BroadcastSocket;
 
diff --git a/Engine/Plugins/Experimental/DigitalMultiplex/DMXProtocol/Source/DMXProtocolArtNet/Private/DMXProtocolArtNetModule.cpp b/Engine/Plugins/Experimental/DigitalMultiplex/DMXProtocol/Source/DMXProtocolArtNet/Private/DMXProtocolArtNetModule.cpp
--- a/Engine/Plugins/Experimental/DigitalMultiplex/DMXProtocol/Source/DMXProtocolArtNet/Private/DMXProtocolArtNetModule.cpp
+++ b/Engine/Plugins/Experimental/DigitalMultiplex/DMXProtocol/Source/DMXProtocolArtNet/Private/DMXProtocolArtNetModule.cpp
@@ -20,10 +20,10 @@ FAutoConsoleCommand FDMXProtocolArtNetModule::SendDMXCommand(
 	FConsoleCommandWithArgsDelegate::CreateStatic(&FDMXProtocolArtNetModule::SendDMXCommandHandler)
 );
 
-TSharedPtr<IDMXProtocol> FDMXProtocolFactoryArtNet::CreateProtocol(const FName & ProtocolName)
+IDMXProtocolPtr FDMXProtocolFactoryArtNet::CreateProtocol(const FName & ProtocolName)
 {
 	FJsonObject ProtocolSettings;
-	TSharedPtr<IDMXProtocol> ProtocolArtNetPtr = MakeShared<FDMXProtocolArtNet>(ProtocolName, ProtocolSettings);
+	IDMXProtocolPtr ProtocolArtNetPtr = MakeShared<FDMXProtocolArtNet, ESPMode::ThreadSafe>(ProtocolName, ProtocolSettings);
 	if (ProtocolArtNetPtr->IsEnabled())
 	{
 		if (!ProtocolArtNetPtr->Init())
@@ -119,7 +119,7 @@ void FDMXProtocolArtNetModule::SendDMXCommandHandler(const TArray<FString>& Args
 		DMXFragment.Add(Key, Value);
 	}
 
-	TSharedPtr<IDMXProtocol> DMXProtocol = IDMXProtocol::Get(FDMXProtocolArtNetModule::NAME_Artnet);
+	IDMXProtocolPtr DMXProtocol = IDMXProtocol::Get(FDMXProtocolArtNetModule::NAME_Artnet);
 	if (DMXProtocol.IsValid())
 	{
 		DMXProtocol->SendDMXFragmentCreate(UniverseID, DMXFragment);
diff --git a/Engine/Plugins/Experimental/DigitalMultiplex/DMXProtocol/Source/DMXProtocolArtNet/Private/DMXProtocolUniverseArtNet.cpp b/Engine/Plugins/Experimental/DigitalMultiplex/DMXProtocol/Source/DMXProtocolArtNet/Private/DMXProtocolUniverseArtNet.cpp
--- a/Engine/Plugins/Experimental/DigitalMultiplex/DMXProtocol/Source/DMXProtocolArtNet/Private/DMXProtocolUniverseArtNet.cpp
+++ b/Engine/Plugins/Experimental/DigitalMultiplex/DMXProtocol/Source/DMXProtocolArtNet/Private/DMXProtocolUniverseArtNet.cpp
@@ -12,7 +12,7 @@ static const uint8 HIGH_NIBBLE = 0xF0;
 static const uint8 LOW_NIBBLE = 0x0F;
 static const uint8 MAX_NET = 0x7f;
 
-FDMXProtocolUniverseArtNet::FDMXProtocolUniverseArtNet(TSharedPtr<IDMXProtocol> InDMXProtocol, const FJsonObject& InSettings)
+FDMXProtocolUniverseArtNet::FDMXProtocolUniverseArtNet(IDMXProtocolPtr InDMXProtocol, const FJsonObject& InSettings)
 	: WeakDMXProtocol(InDMXProtocol)
 {
 	checkf(WeakDMXProtocol.IsValid(), TEXT("DMXProtocol pointer is nullptr"));
@@ -40,7 +40,7 @@ FDMXProtocolUniverseArtNet::FDMXProtocolUniverseArtNet(TSharedPtr<IDMXProtocol>
 }
 
 
-TSharedPtr<IDMXProtocol> FDMXProtocolUniverseArtNet::GetProtocol() const
+IDMXProtocolPtr FDMXProtocolUniverseArtNet::GetProtocol() const
 {
 	return WeakDMXProtocol.Pin();
 }
diff --git a/Engine/Plugins/Experimental/DigitalMultiplex/DMXProtocol/Source/DMXProtocolArtNet/Private/DMXProtocolUniverseArtNet.h b/Engine/Plugins/Experimental/DigitalMultiplex/DMXProtocol/Source/DMXProtocolArtNet/Private/DMXProtocolUniverseArtNet.h
--- a/Engine/Plugins/Experimental/DigitalMultiplex/DMXProtocol/Source/DMXProtocolArtNet/Private/DMXProtocolUniverseArtNet.h
+++ b/Engine/Plugins/Experimental/DigitalMultiplex/DMXProtocol/Source/DMXProtocolArtNet/Private/DMXProtocolUniverseArtNet.h
@@ -11,10 +11,10 @@ class DMXPROTOCOLARTNET_API FDMXProtocolUniverseArtNet
 	: public IDMXProtocolUniverse
 {
 public:
-	FDMXProtocolUniverseArtNet(TSharedPtr<IDMXProtocol> InDMXProtocol, const FJsonObject& InSettings);
+	FDMXProtocolUniverseArtNet(IDMXProtocolPtr InDMXProtocol, const FJsonObject& InSettings);
 
 	//~ Begin IDMXProtocolDevice implementation
-	virtual TSharedPtr<IDMXProtocol> GetProtocol() const override;
+	virtual IDMXProtocolPtr GetProtocol() const override;
 	virtual TSharedPtr<FDMXBuffer> GetInputDMXBuffer() const override;
 	virtual TSharedPtr<FDMXBuffer> GetOutputDMXBuffer() const override;
 	virtual bool SetDMXFragment(const IDMXFragmentMap& DMXFragment) override;
@@ -68,7 +68,7 @@ private:
 	bool SetUniverse(uint8 InUniverse);
 
 private:
-	TWeakPtr<IDMXProtocol> WeakDMXProtocol;
+	IDMXProtocolPtrWeak WeakDMXProtocol;
 	TSharedPtr<FDMXBuffer> OutputDMXBuffer;
 	TSharedPtr<FDMXBuffer> InputDMXBuffer;
 	uint8 Priority;
diff --git a/Engine/Plugins/Experimental/DigitalMultiplex/DMXProtocol/Source/DMXProtocolArtNet/Private/Tests/DMXPrtotocolArtNetFlowTest.cpp b/Engine/Plugins/Experimental/DigitalMultiplex/DMXProtocol/Source/DMXProtocolArtNet/Private/Tests/DMXPrtotocolArtNetFlowTest.cpp
--- a/Engine/Plugins/Experimental/DigitalMultiplex/DMXProtocol/Source/DMXProtocolArtNet/Private/Tests/DMXPrtotocolArtNetFlowTest.cpp
+++ b/Engine/Plugins/Experimental/DigitalMultiplex/DMXProtocol/Source/DMXProtocolArtNet/Private/Tests/DMXPrtotocolArtNetFlowTest.cpp
@@ -44,9 +44,9 @@ bool FDMXPrtotocolArtNetBasicFlowTest::RunTest(const FString& Parameters)
 	FPlatformProcess::Sleep(0.2f);
 	TSharedPtr<FDMXBuffer> InputDMXBuffer = UniverseArtNet->GetInputDMXBuffer();
 
-	TestEqual(TEXT("Incoming buffer should be same"), FixtureValues[0], InputDMXBuffer->GetDMXData()[FixtureChannels[0]]);
-	TestEqual(TEXT("Incoming buffer should be same"), FixtureValues[1], InputDMXBuffer->GetDMXData()[FixtureChannels[1]]);
-	TestEqual(TEXT("Incoming buffer should be same"), FixtureValues[2], InputDMXBuffer->GetDMXData()[FixtureChannels[2]]);
+	TestEqual(TEXT("Incoming buffer should be same"), FixtureValues[0], InputDMXBuffer->GetDMXDataAddress(FixtureChannels[0] - 1));
+	TestEqual(TEXT("Incoming buffer should be same"), FixtureValues[1], InputDMXBuffer->GetDMXDataAddress(FixtureChannels[1] - 1));
+	TestEqual(TEXT("Incoming buffer should be same"), FixtureValues[2], InputDMXBuffer->GetDMXDataAddress(FixtureChannels[2] - 1));
 
 	// Send DMX fragment 2
 	DMXFragmentMap.Add(FixtureChannels[3], FixtureValues[3]);
@@ -56,9 +56,9 @@ bool FDMXPrtotocolArtNetBasicFlowTest::RunTest(const FString& Parameters)
 
 	FPlatformProcess::Sleep(0.2f);
 
-	TestEqual(TEXT("Old value DMX input should be the same"), FixtureValues[0], InputDMXBuffer->GetDMXData()[FixtureChannels[0]]);
-	TestEqual(TEXT("Incoming buffer should be same"), FixtureValues[4], InputDMXBuffer->GetDMXData()[FixtureChannels[4]]);
-	TestEqual(TEXT("Incoming buffer should be same"), FixtureValues[5], InputDMXBuffer->GetDMXData()[FixtureChannels[5]]);
+	TestEqual(TEXT("Old value DMX input should be the same"), FixtureValues[0], InputDMXBuffer->GetDMXDataAddress(FixtureChannels[0] - 1));
+	TestEqual(TEXT("Incoming buffer should be same"), FixtureValues[4], InputDMXBuffer->GetDMXDataAddress(FixtureChannels[4] - 1));
+	TestEqual(TEXT("Incoming buffer should be same"), FixtureValues[5], InputDMXBuffer->GetDMXDataAddress(FixtureChannels[5] - 1));
 
 	// @TODO Implement test cases for Pool and RDM
 	// Send DMX pull
@@ -142,9 +142,9 @@ bool FDMXPrtotocolArtNetConsoleCommandsTest::RunTest(const FString& Parameters)
 	FPlatformProcess::Sleep(0.2f);
 
 	TSharedPtr<FDMXBuffer> InputDMXBuffer = Universe->GetInputDMXBuffer();
-	TestEqual(TEXT("Incoming buffer should be same"), Values[0][1], InputDMXBuffer->GetDMXData()[Values[0][0]]);
-	TestEqual(TEXT("Incoming buffer should be same"), Values[1][1], InputDMXBuffer->GetDMXData()[Values[1][0]]);
-	TestEqual(TEXT("Incoming buffer should be same"), Values[2][1], InputDMXBuffer->GetDMXData()[Values[2][0]]);
+	TestEqual(TEXT("Incoming buffer should be same"), Values[0][1], InputDMXBuffer->GetDMXDataAddress(Values[0][0] - 1));
+	TestEqual(TEXT("Incoming buffer should be same"), Values[1][1], InputDMXBuffer->GetDMXDataAddress(Values[1][0] - 1));
+	TestEqual(TEXT("Incoming buffer should be same"), Values[2][1], InputDMXBuffer->GetDMXDataAddress(Values[2][0] - 1));
 
 	return true;
 }
\ No newline at end of file
diff --git a/Engine/Plugins/Experimental/DigitalMultiplex/DMXProtocol/Source/DMXProtocolArtNet/Public/DMXProtocolArtNetModule.h b/Engine/Plugins/Experimental/DigitalMultiplex/DMXProtocol/Source/DMXProtocolArtNet/Public/DMXProtocolArtNetModule.h
--- a/Engine/Plugins/Experimental/DigitalMultiplex/DMXProtocol/Source/DMXProtocolArtNet/Public/DMXProtocolArtNetModule.h
+++ b/Engine/Plugins/Experimental/DigitalMultiplex/DMXProtocol/Source/DMXProtocolArtNet/Public/DMXProtocolArtNetModule.h
@@ -12,7 +12,7 @@
 class FDMXProtocolFactoryArtNet : public IDMXProtocolFactory
 {
 public:
-	virtual TSharedPtr<IDMXProtocol> CreateProtocol(const FName& ProtocolName) override;
+	virtual IDMXProtocolPtr CreateProtocol(const FName& ProtocolName) override;
 };
 
 
diff --git a/Engine/Plugins/Experimental/DigitalMultiplex/DMXProtocol/Source/DMXProtocolSACN/Private/DMXProtocolSACN.cpp b/Engine/Plugins/Experimental/DigitalMultiplex/DMXProtocol/Source/DMXProtocolSACN/Private/DMXProtocolSACN.cpp
--- a/Engine/Plugins/Experimental/DigitalMultiplex/DMXProtocol/Source/DMXProtocolSACN/Private/DMXProtocolSACN.cpp
+++ b/Engine/Plugins/Experimental/DigitalMultiplex/DMXProtocol/Source/DMXProtocolSACN/Private/DMXProtocolSACN.cpp
@@ -182,7 +182,7 @@ EDMXSendResult FDMXProtocolSACN::SendDMXFragmentCreate(uint16 InUniverseID, cons
 uint16 FDMXProtocolSACN::GetFinalSendUniverseID(uint16 InUniverseID) const
 {
 	// GlobalSACNUniverseOffset clamp between 0 and 65535(max uint16) 
-	return InUniverseID + GetDefault<UDMXProtocolSettings>()->GlobalArtNetUniverseOffset;
+	return InUniverseID + GetDefault<UDMXProtocolSettings>()->GlobalSACNUniverseOffset;
 }
 
 uint32 FDMXProtocolSACN::GetUniversesNum() const
@@ -219,11 +219,20 @@ EDMXSendResult FDMXProtocolSACN::SendDMXInternal(uint16 InUniverseID, const TSha
 	FDMXProtocolE131DMPLayerPacket DMPLayer;
 	DMPLayer.AddressIncrement = ACN_ADDRESS_INC;
 	DMPLayer.PropertyValueCount = ACN_DMX_SIZE + 1;
-	if (DMXBuffer->GetDMXData().Num() == ACN_DMX_SIZE)
-	{
-		FMemory::Memcpy(DMPLayer.DMX, DMXBuffer->GetDMXData().GetData(), ACN_DMX_SIZE);
-	}
-	else
+	bool bBufferSizeIsWrong = false;
+	DMXBuffer->AccessDMXData([&DMPLayer, &bBufferSizeIsWrong](TArray<uint8>& InData)
+		{
+			if (InData.Num() == ACN_DMX_SIZE)
+			{
+				FMemory::Memcpy(DMPLayer.DMX, InData.GetData(), ACN_DMX_SIZE);
+			}
+			else
+			{
+				bBufferSizeIsWrong = true;
+			}
+		});
+
+	if (bBufferSizeIsWrong)
 	{
 		return EDMXSendResult::ErrorSizeBuffer;
 	}
diff --git a/Engine/Plugins/Experimental/DigitalMultiplex/DMXProtocol/Source/DMXProtocolSACN/Private/DMXProtocolSACNModule.cpp b/Engine/Plugins/Experimental/DigitalMultiplex/DMXProtocol/Source/DMXProtocolSACN/Private/DMXProtocolSACNModule.cpp
--- a/Engine/Plugins/Experimental/DigitalMultiplex/DMXProtocol/Source/DMXProtocolSACN/Private/DMXProtocolSACNModule.cpp
+++ b/Engine/Plugins/Experimental/DigitalMultiplex/DMXProtocol/Source/DMXProtocolSACN/Private/DMXProtocolSACNModule.cpp
@@ -24,10 +24,10 @@ FAutoConsoleCommand FDMXProtocolSACNModule::SendDMXCommand(
 	FConsoleCommandWithArgsDelegate::CreateStatic(&FDMXProtocolSACNModule::SendDMXCommandHandler)
 );
 
-TSharedPtr<IDMXProtocol> FDMXProtocolFactorySACN::CreateProtocol(const FName & ProtocolName)
+IDMXProtocolPtr FDMXProtocolFactorySACN::CreateProtocol(const FName & ProtocolName)
 {
 	FJsonObject ProtocolSettings;
-	TSharedPtr<IDMXProtocol> ProtocolSACNPtr = MakeShared<FDMXProtocolSACN>(ProtocolName, ProtocolSettings);
+	IDMXProtocolPtr ProtocolSACNPtr = MakeShared<FDMXProtocolSACN, ESPMode::ThreadSafe>(ProtocolName, ProtocolSettings);
 	if (ProtocolSACNPtr->IsEnabled())
 	{
 		if (!ProtocolSACNPtr->Init())
@@ -123,7 +123,7 @@ void FDMXProtocolSACNModule::SendDMXCommandHandler(const TArray<FString>& Args)
 		DMXFragment.Add(Key, Value);
 	}
 
-	TSharedPtr<IDMXProtocol> DMXProtocol = IDMXProtocol::Get(FDMXProtocolSACNModule::NAME_SACN);
+	IDMXProtocolPtr DMXProtocol = IDMXProtocol::Get(FDMXProtocolSACNModule::NAME_SACN);
 	if (DMXProtocol.IsValid())
 	{
 		DMXProtocol->SendDMXFragmentCreate(UniverseID, DMXFragment);
diff --git a/Engine/Plugins/Experimental/DigitalMultiplex/DMXProtocol/Source/DMXProtocolSACN/Private/DMXProtocolUniverseSACN.cpp b/Engine/Plugins/Experimental/DigitalMultiplex/DMXProtocol/Source/DMXProtocolSACN/Private/DMXProtocolUniverseSACN.cpp
--- a/Engine/Plugins/Experimental/DigitalMultiplex/DMXProtocol/Source/DMXProtocolSACN/Private/DMXProtocolUniverseSACN.cpp
+++ b/Engine/Plugins/Experimental/DigitalMultiplex/DMXProtocol/Source/DMXProtocolSACN/Private/DMXProtocolUniverseSACN.cpp
@@ -12,7 +12,7 @@
 
 #include "DMXProtocolSACNUtils.h"
 
-FDMXProtocolUniverseSACN::FDMXProtocolUniverseSACN(TSharedPtr<IDMXProtocol> InDMXProtocol, const FJsonObject& InSettings)
+FDMXProtocolUniverseSACN::FDMXProtocolUniverseSACN(IDMXProtocolPtr InDMXProtocol, const FJsonObject& InSettings)
 	: WeakDMXProtocol(InDMXProtocol)
 	, ListeningSocket(nullptr)
 {
@@ -46,7 +46,7 @@ FDMXProtocolUniverseSACN::~FDMXProtocolUniverseSACN()
 	IDMXProtocol::OnNetworkInterfaceChanged.Remove(NetworkInterfaceChangedHandle);
 }
 
-TSharedPtr<IDMXProtocol> FDMXProtocolUniverseSACN::GetProtocol() const
+IDMXProtocolPtr FDMXProtocolUniverseSACN::GetProtocol() const
 {
 	return WeakDMXProtocol.Pin();
 }
@@ -88,8 +88,6 @@ bool FDMXProtocolUniverseSACN::IsSupportRDM() const
 
 void FDMXProtocolUniverseSACN::OnDataReceived(const FArrayReaderPtr & Buffer)
 {
-	FScopeLock Lock(&OnDataReceivedCS);
-
 	// It will be more handlers
 	switch (SACN::GetRootPacketType(Buffer))
 	{
@@ -108,19 +106,37 @@ bool FDMXProtocolUniverseSACN::HandleReplyPacket(const FArrayReaderPtr & Buffer)
 	*Buffer << IncomingDMXFramingLayer;
 	*Buffer << IncomingDMXDMPLayer;
 
-	// Make sure we copy same amount of data
-	if (InputDMXBuffer->GetDMXData().Num() == ACN_DMX_SIZE)
-	{
-		InputDMXBuffer->SetDMXBuffer(IncomingDMXDMPLayer.DMX, ACN_DMX_SIZE);
-
-		GetProtocol()->GetOnUniverseInputUpdate().Broadcast(GetProtocol()->GetProtocolName(), GetUniverseID(), InputDMXBuffer->GetDMXData());
-		return true;
-	}
-	else
-	{
-		UE_LOG_DMXPROTOCOL(Error, TEXT("%s:Size of incoming DMX buffer is wrong"), NetworkErrorMessagePrefix);
-		return false;
-	}
+	bool bCopySuccessful = false;
+	// Access the buffer thread-safely
+	InputDMXBuffer->AccessDMXData([this, &bCopySuccessful](TArray<uint8>& InData)
+		{
+			// Make sure we copy same amount of data
+			if (InData.Num() == ACN_DMX_SIZE)
+			{
+				InputDMXBuffer->SetDMXBuffer(IncomingDMXDMPLayer.DMX, ACN_DMX_SIZE);
+				IDMXProtocolPtr Protocol = GetProtocol();
+				if (Protocol.IsValid())
+				{
+					Protocol->GetOnUniverseInputUpdate().Broadcast(Protocol->GetProtocolName(), GetUniverseID(), InData);
+					bCopySuccessful = true;
+				}
+				else
+				{
+					UE_LOG_DMXPROTOCOL(Error, TEXT("Invalid SACN Protocol"));
+					bCopySuccessful = false;
+				}
+			}
+			else
+			{
+				UE_LOG_DMXPROTOCOL(Error, TEXT("%s: Size of incoming DMX buffer is wrong! Expected: %d; Found: %d")
+					, NetworkErrorMessagePrefix
+					, ACN_DMX_SIZE
+					, InData.Num());
+				bCopySuccessful = false;
+			}
+		});
+
+	return bCopySuccessful;
 }
 
 void FDMXProtocolUniverseSACN::OnNetworkInterfaceChanged(const FString& InInterfaceIPAddress)
diff --git a/Engine/Plugins/Experimental/DigitalMultiplex/DMXProtocol/Source/DMXProtocolSACN/Private/DMXProtocolUniverseSACN.h b/Engine/Plugins/Experimental/DigitalMultiplex/DMXProtocol/Source/DMXProtocolSACN/Private/DMXProtocolUniverseSACN.h
--- a/Engine/Plugins/Experimental/DigitalMultiplex/DMXProtocol/Source/DMXProtocolSACN/Private/DMXProtocolUniverseSACN.h
+++ b/Engine/Plugins/Experimental/DigitalMultiplex/DMXProtocol/Source/DMXProtocolSACN/Private/DMXProtocolUniverseSACN.h
@@ -18,11 +18,11 @@ class DMXPROTOCOLSACN_API FDMXProtocolUniverseSACN
 	, public IDMXNetworkInterface
 {
 public:
-	FDMXProtocolUniverseSACN(TSharedPtr<IDMXProtocol> InDMXProtocol, const FJsonObject& InSettings);
+	FDMXProtocolUniverseSACN(IDMXProtocolPtr InDMXProtocol, const FJsonObject& InSettings);
 	virtual ~FDMXProtocolUniverseSACN();
 
 	//~ Begin IDMXProtocolDevice implementation
-	virtual TSharedPtr<IDMXProtocol> GetProtocol() const override;
+	virtual IDMXProtocolPtr GetProtocol() const override;
 	virtual TSharedPtr<FDMXBuffer> GetInputDMXBuffer() const override;
 	virtual TSharedPtr<FDMXBuffer> GetOutputDMXBuffer() const override;
 	virtual bool SetDMXFragment(const IDMXFragmentMap& DMXFragment) override;
@@ -49,7 +49,7 @@ private:
 	bool HandleReplyPacket(const FArrayReaderPtr & Buffer);
 
 private:
-	TWeakPtr<IDMXProtocol> WeakDMXProtocol;
+	IDMXProtocolPtrWeak WeakDMXProtocol;
 	TSharedPtr<FDMXBuffer> OutputDMXBuffer;
 	TSharedPtr<FDMXBuffer> InputDMXBuffer;
 	uint8 Priority;
@@ -57,8 +57,6 @@ private:
 	bool bIsRDMSupport;
 	TSharedPtr<FJsonObject> Settings;
 
-	mutable FCriticalSection OnDataReceivedCS;
-
 	FSocket* ListeningSocket;
 	TSharedPtr<IDMXProtocolReceiver> SACNReceiver;
 
diff --git a/Engine/Plugins/Experimental/DigitalMultiplex/DMXProtocol/Source/DMXProtocolSACN/Private/Tests/DMXPrtotocolSACNFlowTest.cpp b/Engine/Plugins/Experimental/DigitalMultiplex/DMXProtocol/Source/DMXProtocolSACN/Private/Tests/DMXPrtotocolSACNFlowTest.cpp
--- a/Engine/Plugins/Experimental/DigitalMultiplex/DMXProtocol/Source/DMXProtocolSACN/Private/Tests/DMXPrtotocolSACNFlowTest.cpp
+++ b/Engine/Plugins/Experimental/DigitalMultiplex/DMXProtocol/Source/DMXProtocolSACN/Private/Tests/DMXPrtotocolSACNFlowTest.cpp
@@ -49,9 +49,9 @@ bool DMXPrtotocolSACNBasicFlowTest::RunTest(const FString& Parameters)
 		TestTrue(TEXT("Incoming ACNPacketIdentifier should be the same"), ExpectedInentifier.Equals(IncomingInentifier));
 	}
 
-	TestEqual(TEXT("Incoming buffer should be same"), FixtureValues[0], InputDMXBuffer->GetDMXData()[FixtureChannels[0] - 1]);
-	TestEqual(TEXT("Incoming buffer should be same"), FixtureValues[1], InputDMXBuffer->GetDMXData()[FixtureChannels[1] - 1]);
-	TestEqual(TEXT("Incoming buffer should be same"), FixtureValues[2], InputDMXBuffer->GetDMXData()[FixtureChannels[2] - 1]);
+	TestEqual(TEXT("Incoming buffer should be same"), FixtureValues[0], InputDMXBuffer->GetDMXDataAddress(FixtureChannels[0] - 1));
+	TestEqual(TEXT("Incoming buffer should be same"), FixtureValues[1], InputDMXBuffer->GetDMXDataAddress(FixtureChannels[1] - 1));
+	TestEqual(TEXT("Incoming buffer should be same"), FixtureValues[2], InputDMXBuffer->GetDMXDataAddress(FixtureChannels[2] - 1));
 
 	// Send DMX fragment 2
 	DMXFragmentMap.Add(FixtureChannels[3], FixtureValues[3]);
@@ -60,9 +60,9 @@ bool DMXPrtotocolSACNBasicFlowTest::RunTest(const FString& Parameters)
 	DMXProtocol->SendDMXFragment(Universe->GetUniverseID(), DMXFragmentMap);
 	FPlatformProcess::Sleep(0.2f);
 
-	TestEqual(TEXT("Old value DMX input should be the same"), FixtureValues[0], InputDMXBuffer->GetDMXData()[FixtureChannels[0] - 1]);
-	TestEqual(TEXT("Incoming buffer should be same"), FixtureValues[4], InputDMXBuffer->GetDMXData()[FixtureChannels[4] - 1]);
-	TestEqual(TEXT("Incoming buffer should be same"), FixtureValues[5], InputDMXBuffer->GetDMXData()[FixtureChannels[5] - 1]);
+	TestEqual(TEXT("Old value DMX input should be the same"), FixtureValues[0], InputDMXBuffer->GetDMXDataAddress(FixtureChannels[0] - 1));
+	TestEqual(TEXT("Incoming buffer should be same"), FixtureValues[4], InputDMXBuffer->GetDMXDataAddress(FixtureChannels[4] - 1));
+	TestEqual(TEXT("Incoming buffer should be same"), FixtureValues[5], InputDMXBuffer->GetDMXDataAddress(FixtureChannels[5] - 1));
 
 	// @TODO Implement test cases for Discovery
 	TArray<uint16> Universes;
@@ -108,9 +108,9 @@ bool FDMXPrtotocolSACNConsoleCommandsTest::RunTest(const FString& Parameters)
 	FPlatformProcess::Sleep(0.2f);
 
 	TSharedPtr<FDMXBuffer> InputDMXBuffer = Universe->GetInputDMXBuffer();
-	TestEqual(TEXT("Incoming buffer should be same"), Values[0][1], InputDMXBuffer->GetDMXData()[Values[0][0] - 1]);
-	TestEqual(TEXT("Incoming buffer should be same"), Values[1][1], InputDMXBuffer->GetDMXData()[Values[1][0] - 1]);
-	TestEqual(TEXT("Incoming buffer should be same"), Values[2][1], InputDMXBuffer->GetDMXData()[Values[2][0] - 1]);
+	TestEqual(TEXT("Incoming buffer should be same"), Values[0][1], InputDMXBuffer->GetDMXDataAddress(Values[0][0] - 1));
+	TestEqual(TEXT("Incoming buffer should be same"), Values[1][1], InputDMXBuffer->GetDMXDataAddress(Values[1][0] - 1));
+	TestEqual(TEXT("Incoming buffer should be same"), Values[2][1], InputDMXBuffer->GetDMXDataAddress(Values[2][0] - 1));
 
 	return true;
 }
\ No newline at end of file
diff --git a/Engine/Plugins/Experimental/DigitalMultiplex/DMXProtocol/Source/DMXProtocolSACN/Public/DMXProtocolSACNModule.h b/Engine/Plugins/Experimental/DigitalMultiplex/DMXProtocol/Source/DMXProtocolSACN/Public/DMXProtocolSACNModule.h
--- a/Engine/Plugins/Experimental/DigitalMultiplex/DMXProtocol/Source/DMXProtocolSACN/Public/DMXProtocolSACNModule.h
+++ b/Engine/Plugins/Experimental/DigitalMultiplex/DMXProtocol/Source/DMXProtocolSACN/Public/DMXProtocolSACNModule.h
@@ -12,7 +12,7 @@
 class FDMXProtocolFactorySACN : public IDMXProtocolFactory
 {
 public:
-	virtual TSharedPtr<IDMXProtocol> CreateProtocol(const FName& ProtocolName) override;
+	virtual IDMXProtocolPtr CreateProtocol(const FName& ProtocolName) override;
 };
 
 /**
