Fixed CIS warning V609: Divide by zero. Denominator range.

diff --git a/Engine/Plugins/Experimental/DigitalMultiplex/DMXEngine/Source/DMXRuntime/Private/Library/DMXEntityFixtureType.cpp b/Engine/Plugins/Experimental/DigitalMultiplex/DMXEngine/Source/DMXRuntime/Private/Library/DMXEntityFixtureType.cpp
--- a/Engine/Plugins/Experimental/DigitalMultiplex/DMXEngine/Source/DMXRuntime/Private/Library/DMXEntityFixtureType.cpp
+++ b/Engine/Plugins/Experimental/DigitalMultiplex/DMXEngine/Source/DMXRuntime/Private/Library/DMXEntityFixtureType.cpp
@@ -198,7 +198,7 @@ uint32 UDMXEntityFixtureType::GetDataTypeMaxValue(EDMXFixtureSignalFormat DataTy
 		return MAX_uint32;
 	default:
 		checkNoEntry();
-		return 0;
+		return 1;
 	}
 }
 
