Edigrated CL 12476081.

Fix of the FDMXEntityFixturePatchRef struct serialization

otherwise K2Node::FixupPinStringDataReferences() function called multiple times and it reset the right serialisation for UStruct

diff --git a/Engine/Plugins/Experimental/DigitalMultiplex/DMXEngine/Source/DMXBlueprintGraph/Private/K2Node_GetDMXFixturePatch.cpp b/Engine/Plugins/Experimental/DigitalMultiplex/DMXEngine/Source/DMXBlueprintGraph/Private/K2Node_GetDMXFixturePatch.cpp
--- a/Engine/Plugins/Experimental/DigitalMultiplex/DMXEngine/Source/DMXBlueprintGraph/Private/K2Node_GetDMXFixturePatch.cpp
+++ b/Engine/Plugins/Experimental/DigitalMultiplex/DMXEngine/Source/DMXBlueprintGraph/Private/K2Node_GetDMXFixturePatch.cpp
@@ -12,12 +12,43 @@
 #include "BlueprintNodeSpawner.h"
 #include "K2Node_CallFunction.h"
 #include "Kismet2/BlueprintEditorUtils.h"
+#include "UObject/FrameworkObjectVersion.h"
 
 #define LOCTEXT_NAMESPACE "UK2Node_GetDMXFixturePatch"
 
 const FName UK2Node_GetDMXFixturePatch::InputDMXFixturePatchPinName(TEXT("InFixturePatch"));
 const FName UK2Node_GetDMXFixturePatch::OutputDMXFixturePatchPinName(TEXT("OutFixturePatch"));
 
+void UK2Node_GetDMXFixturePatch::Serialize(FArchive& Ar)
+{
+	Ar.UsingCustomVersion(FFrameworkObjectVersion::GUID);
+
+	// Check if it not SavingPackage
+	if (Ar.IsSaving() && !GIsSavingPackage)
+	{
+		if (Ar.IsObjectReferenceCollector() || Ar.Tell() < 0)
+		{
+			// When this is a reference collector/modifier, serialize some pins as structs
+			FixupPinStringDataReferences(&Ar);
+		}
+	}
+
+	// Do not call parent, but call grandparent
+	UEdGraphNode::Serialize(Ar);
+
+	if (Ar.IsLoading() && ((Ar.GetPortFlags() & PPF_Duplicate) == 0))
+	{
+		// Fix up pin default values, must be done before post load
+		FixupPinDefaultValues();
+
+		if (GIsEditor)
+		{
+			// We need to serialize string data references on load in editor builds so the cooker knows about them
+			FixupPinStringDataReferences(nullptr);
+		}
+	}
+}
+
 void UK2Node_GetDMXFixturePatch::AllocateDefaultPins()
 {
 	const UEdGraphSchema_K2* K2Schema = GetDefault<UEdGraphSchema_K2>();
diff --git a/Engine/Plugins/Experimental/DigitalMultiplex/DMXEngine/Source/DMXBlueprintGraph/Public/K2Node_GetDMXFixturePatch.h b/Engine/Plugins/Experimental/DigitalMultiplex/DMXEngine/Source/DMXBlueprintGraph/Public/K2Node_GetDMXFixturePatch.h
--- a/Engine/Plugins/Experimental/DigitalMultiplex/DMXEngine/Source/DMXBlueprintGraph/Public/K2Node_GetDMXFixturePatch.h
+++ b/Engine/Plugins/Experimental/DigitalMultiplex/DMXEngine/Source/DMXBlueprintGraph/Public/K2Node_GetDMXFixturePatch.h
@@ -20,6 +20,10 @@ class DMXBLUEPRINTGRAPH_API UK2Node_GetDMXFixturePatch
 	GENERATED_BODY()
 
 public:
+	//~ Begin UObject interface
+	virtual void Serialize(FArchive& Ar) override;
+	//~ End UObject interface
+
 	//~ Begin UEdGraphNode Interface
 	virtual void AllocateDefaultPins() override;
 	virtual FText GetTooltipText() const override;
