Add Load DMX UObject blueprint function

diff --git a/Engine/Plugins/Experimental/DigitalMultiplex/DMXEngine/Source/DMXRuntime/Public/DMXSubsystem.h b/Engine/Plugins/Experimental/DigitalMultiplex/DMXEngine/Source/DMXRuntime/Public/DMXSubsystem.h
--- a/Engine/Plugins/Experimental/DigitalMultiplex/DMXEngine/Source/DMXRuntime/Public/DMXSubsystem.h
+++ b/Engine/Plugins/Experimental/DigitalMultiplex/DMXEngine/Source/DMXRuntime/Public/DMXSubsystem.h
@@ -38,6 +38,10 @@ public:
 	/**  Return reference to array of Fixture Patch objects of a given type. */
 	UFUNCTION(BlueprintCallable, Category = "DMX", meta = (AutoCreateRefTerm = "FixtureType"))
 	void GetAllFixturesOfType(const FDMXEntityFixtureTypeRef& FixtureType, TArray<UDMXEntityFixturePatch*>& OutResult);
+	
+	/** Load all referenced Objects in the given library. The DMX library should be Loaded at least one in any part of the Unreal Engine */
+	UFUNCTION(BlueprintCallable, Category = "DMX", meta = (DisplayName = "Load DMX Library"))
+	void LoadDMXLibrary(UDMXLibrary* DMXLibrary) {} // It could be empty. It automatically pre-loads all Objects if we have a reference in the blueprint.
 
 	/**  Return reference to array of Fixture Patch objects of a given category. */
 	UFUNCTION(BlueprintCallable, Category = "DMX")
