DMX fix: Controller updates valid universes when created, thus Patches on default Universe 1 display as assigned.

diff --git a/Engine/Plugins/Experimental/DigitalMultiplex/DMXEngine/Source/DMXRuntime/Private/Library/DMXEntityController.cpp b/Engine/Plugins/Experimental/DigitalMultiplex/DMXEngine/Source/DMXRuntime/Private/Library/DMXEntityController.cpp
--- a/Engine/Plugins/Experimental/DigitalMultiplex/DMXEngine/Source/DMXRuntime/Private/Library/DMXEntityController.cpp
+++ b/Engine/Plugins/Experimental/DigitalMultiplex/DMXEngine/Source/DMXRuntime/Private/Library/DMXEntityController.cpp
@@ -24,6 +24,14 @@ void UDMXEntityController::PostLoad()
 {
 	Super::PostLoad();
 	ValidateRangeValues();
+	UpdateUniversesFromRange();
+}
+
+void UDMXEntityController::PostInitProperties()
+{
+	Super::PostInitProperties();
+	ValidateRangeValues();
+	UpdateUniversesFromRange();
 }
 
 #endif // WITH_EDITOR
diff --git a/Engine/Plugins/Experimental/DigitalMultiplex/DMXEngine/Source/DMXRuntime/Public/Library/DMXEntityController.h b/Engine/Plugins/Experimental/DigitalMultiplex/DMXEngine/Source/DMXRuntime/Public/Library/DMXEntityController.h
--- a/Engine/Plugins/Experimental/DigitalMultiplex/DMXEngine/Source/DMXRuntime/Public/Library/DMXEntityController.h
+++ b/Engine/Plugins/Experimental/DigitalMultiplex/DMXEngine/Source/DMXRuntime/Public/Library/DMXEntityController.h
@@ -71,6 +71,7 @@ public:
 	virtual void PostEditChangeProperty(FPropertyChangedEvent& PropertyChangedEvent) override;
 	virtual void PostLoad() override;
 #endif // WITH_EDITOR
+	virtual void PostInitProperties() override;
 
 protected:
 	/** Keep range valid */
