DMX fix: references to Entities in a duplicated library are no longer lost after restarting the editor.

diff --git a/Engine/Plugins/Experimental/DigitalMultiplex/DMXEngine/Source/DMXRuntime/Private/Library/DMXLibrary.cpp b/Engine/Plugins/Experimental/DigitalMultiplex/DMXEngine/Source/DMXRuntime/Private/Library/DMXLibrary.cpp
--- a/Engine/Plugins/Experimental/DigitalMultiplex/DMXEngine/Source/DMXRuntime/Private/Library/DMXLibrary.cpp
+++ b/Engine/Plugins/Experimental/DigitalMultiplex/DMXEngine/Source/DMXRuntime/Private/Library/DMXLibrary.cpp
@@ -210,4 +210,22 @@ FOnEntitiesUpdated& UDMXLibrary::GetOnEntitiesUpdated()
 	return OnEntitiesUpdated;
 }
 
+void UDMXLibrary::PostDuplicate(EDuplicateMode::Type DuplicateMode)
+{
+	Super::PostDuplicate(DuplicateMode);
+
+	if (DuplicateMode != EDuplicateMode::Normal)
+	{
+		return;
+	}
+
+	// Make sure all Entity children have this library as their parent
+	// and refresh their ID
+	for (UDMXEntity* Entity : Entities)
+	{
+		Entity->SetParentLibrary(this);
+		Entity->RefreshID();
+	}
+}
+
 #undef LOCTEXT_NAMESPACE
diff --git a/Engine/Plugins/Experimental/DigitalMultiplex/DMXEngine/Source/DMXRuntime/Public/Library/DMXLibrary.h b/Engine/Plugins/Experimental/DigitalMultiplex/DMXEngine/Source/DMXRuntime/Public/Library/DMXLibrary.h
--- a/Engine/Plugins/Experimental/DigitalMultiplex/DMXEngine/Source/DMXRuntime/Public/Library/DMXLibrary.h
+++ b/Engine/Plugins/Experimental/DigitalMultiplex/DMXEngine/Source/DMXRuntime/Public/Library/DMXLibrary.h
@@ -138,6 +138,8 @@ public:
 	/** Called when the list of entities is changed by either adding or removing entities */
 	FOnEntitiesUpdated& GetOnEntitiesUpdated();
 
+	virtual void PostDuplicate(EDuplicateMode::Type DuplicateMode) override;
+
 
 private:
 	UPROPERTY()
