Build fix for a missing argument to NotifyPostChange.

#rnx

diff --git a/Engine/Plugins/Experimental/DigitalMultiplex/DMXEngine/Source/DMXEditor/Private/Sequencer/Customizations/TakeRecorderDMXLibrarySourceEditorCustomization.cpp b/Engine/Plugins/Experimental/DigitalMultiplex/DMXEngine/Source/DMXEditor/Private/Sequencer/Customizations/TakeRecorderDMXLibrarySourceEditorCustomization.cpp
--- a/Engine/Plugins/Experimental/DigitalMultiplex/DMXEngine/Source/DMXEditor/Private/Sequencer/Customizations/TakeRecorderDMXLibrarySourceEditorCustomization.cpp
+++ b/Engine/Plugins/Experimental/DigitalMultiplex/DMXEngine/Source/DMXEditor/Private/Sequencer/Customizations/TakeRecorderDMXLibrarySourceEditorCustomization.cpp
@@ -101,7 +101,7 @@ FReply FDMXLibraryRecorderAddAllPatchesButtonCustomization::HandleOnClicked()
 			}
 		}
 
-		StructHandle->NotifyPostChange();
+		StructHandle->NotifyPostChange(EPropertyChangeType::ValueSet);
 	}
 
 	StructHandle->NotifyFinishedChangingProperties();
