DMX - Recording DMX now directly sets the DMX library to the track, log warnings no longer appear. Fix a rare crash when async recording DMX would not complete by the time the data was requested.

diff --git a/Engine/Plugins/VirtualProduction/DMX/DMXEngine/Source/DMXEditor/Private/Sequencer/DMXAsyncDMXRecorder.cpp b/Engine/Plugins/VirtualProduction/DMX/DMXEngine/Source/DMXEditor/Private/Sequencer/DMXAsyncDMXRecorder.cpp
--- a/Engine/Plugins/VirtualProduction/DMX/DMXEngine/Source/DMXEditor/Private/Sequencer/DMXAsyncDMXRecorder.cpp
+++ b/Engine/Plugins/VirtualProduction/DMX/DMXEngine/Source/DMXEditor/Private/Sequencer/DMXAsyncDMXRecorder.cpp
@@ -303,6 +303,9 @@ void FDMXAsyncDMXRecorder::StartRecording(double InRecordStartTime, FFrameNumber
 
 	if (bListenersStopped)
 	{
+		// Reset the buffer, to make sure it doesn't contain data from the previous recording
+		SignalToLocalUniverseMap.Reset();
+
 		// Get relevant timing info
 		FTakeRecorderParameters Parameters;
 		Parameters.Project = GetDefault<UTakeRecorderProjectSettings>()->Settings;
@@ -358,7 +361,6 @@ int32 FDMXAsyncDMXRecorder::GetNumSignalsToProccess() const
 TArray<FDMXFunctionChannelData> FDMXAsyncDMXRecorder::GetRecordedData() const
 {
 	check(bListenersStopped);
-	check(NumSignalsToProccess == 0);
 
 	return FunctionChannelDataArr;
 }
diff --git a/Engine/Plugins/VirtualProduction/DMX/DMXEngine/Source/DMXEditor/Private/Sequencer/DMXAsyncDMXRecorder.h b/Engine/Plugins/VirtualProduction/DMX/DMXEngine/Source/DMXEditor/Private/Sequencer/DMXAsyncDMXRecorder.h
--- a/Engine/Plugins/VirtualProduction/DMX/DMXEngine/Source/DMXEditor/Private/Sequencer/DMXAsyncDMXRecorder.h
+++ b/Engine/Plugins/VirtualProduction/DMX/DMXEngine/Source/DMXEditor/Private/Sequencer/DMXAsyncDMXRecorder.h
@@ -109,16 +109,10 @@ public:
 	/** Stops the recording */
 	void StopRecording();
 
-	/**
-	 * Returns the number of signals that remain to be proccessed.
-	 * Should only be called when StopRecording was called (checked).
-	 */
+	/** Returns the number of signals that remain to be proccessed. */
 	int32 GetNumSignalsToProccess() const;
 
-	/**
-	 * Returns the recorded data.
-	 * Should only be called when StopRecording was called GetNumSignalsToProccess == 0 (checked).
-	 */
+	/** Returns the recorded data. */
 	TArray<FDMXFunctionChannelData> GetRecordedData() const;
 
 private:
diff --git a/Engine/Plugins/VirtualProduction/DMX/DMXEngine/Source/DMXRuntime/Private/Sequencer/MovieSceneDMXLibrarySection.cpp b/Engine/Plugins/VirtualProduction/DMX/DMXEngine/Source/DMXRuntime/Private/Sequencer/MovieSceneDMXLibrarySection.cpp
--- a/Engine/Plugins/VirtualProduction/DMX/DMXEngine/Source/DMXRuntime/Private/Sequencer/MovieSceneDMXLibrarySection.cpp
+++ b/Engine/Plugins/VirtualProduction/DMX/DMXEngine/Source/DMXRuntime/Private/Sequencer/MovieSceneDMXLibrarySection.cpp
@@ -25,6 +25,7 @@ void FDMXFixturePatchChannel::SetFixturePatch(UDMXEntityFixturePatch* InPatch)
 {
 	if (InPatch != nullptr && InPatch->IsValidLowLevelFast())
 	{
+		DMXLibrary = InPatch->GetParentLibrary();
 		ActiveMode = InPatch->GetActiveModeIndex();
 	}
 
@@ -335,7 +336,7 @@ void UMovieSceneDMXLibrarySection::Serialize(FArchive& Ar)
 
 		if (Ar.CustomVer(FDMXRuntimeObjectVersion::GUID) < FDMXRuntimeObjectVersion::ReplaceWeakWithStrongFixturePatchReferncesInLibrarySection)
 		{
-			// Add a library reference if possible. In cases where this is serialized before the section's library is loaded, this is not possible programaticalle.
+			// Add a library reference if possible. In cases where this is serialized before the section's library is loaded, this is not possible programmatically.
 			// This is expected to be a rare case (as it wouldn't prevent the sequence from playing before 4.27). For these cases, provide detailed steps to the user how to upgrade in logs.
 			for (FDMXFixturePatchChannel& Channel : FixturePatchChannels)
 			{
