DMX - Fix a crash when a library asset is present but cannot be loaded, e.g. because saved with a newer engine version

diff --git a/Engine/Plugins/VirtualProduction/DMX/DMXEngine/Source/DMXRuntime/Private/DMXSubsystem.cpp b/Engine/Plugins/VirtualProduction/DMX/DMXEngine/Source/DMXRuntime/Private/DMXSubsystem.cpp
--- a/Engine/Plugins/VirtualProduction/DMX/DMXEngine/Source/DMXRuntime/Private/DMXSubsystem.cpp
+++ b/Engine/Plugins/VirtualProduction/DMX/DMXEngine/Source/DMXRuntime/Private/DMXSubsystem.cpp
@@ -764,14 +764,16 @@ void UDMXSubsystem::Initialize(FSubsystemCollectionBase& Collection)
 	LibraryOfDMXLibraries->GetAssetDataList(AssetDatas);
 
 	for (const FAssetData& AssetData : AssetDatas)
-{
-		if (!AssetData.IsAssetLoaded())
 	{
-			LoadedDMXLibraries.Add(CastChecked<UDMXLibrary>(AssetData.ToSoftObjectPath().TryLoad()));
+		UDMXLibrary* DMXLibrary = Cast<UDMXLibrary>(AssetData.ToSoftObjectPath().TryLoad());
+
+		if (DMXLibrary)
+		{
+			LoadedDMXLibraries.Add(DMXLibrary);
 		}
 		else
 		{
-			LoadedDMXLibraries.Add(CastChecked<UDMXLibrary>(AssetData.ToSoftObjectPath().ResolveObject()));
+			UE_LOG(DMXSubsystemLog, Warning, TEXT("Failed to load DMXLibrary %s. See previous errors for causes."), *AssetData.AssetName.ToString());
 		}
 	}
 	OnAllDMXLibraryAssetsLoaded.Broadcast();
