DMX - Activity monitor triggers an ensure when user modifies the universe range

diff --git a/Engine/Plugins/VirtualProduction/DMX/DMXEngine/Source/DMXEditor/Private/Widgets/Monitors/SDMXActivityMonitor.cpp b/Engine/Plugins/VirtualProduction/DMX/DMXEngine/Source/DMXEditor/Private/Widgets/Monitors/SDMXActivityMonitor.cpp
--- a/Engine/Plugins/VirtualProduction/DMX/DMXEngine/Source/DMXEditor/Private/Widgets/Monitors/SDMXActivityMonitor.cpp
+++ b/Engine/Plugins/VirtualProduction/DMX/DMXEngine/Source/DMXEditor/Private/Widgets/Monitors/SDMXActivityMonitor.cpp
@@ -211,9 +211,9 @@ void SDMXActivityMonitor::ClearAllDMXBuffers()
 
 void SDMXActivityMonitor::ResizeDataMapToUniverseRange()
 {
-	TMap<int32, TArray<uint8, TFixedAllocator<DMX_UNIVERSE_SIZE>>> CachedUniverseToDataMap = UniverseToDataMap;
+	const TMap<int32, TArray<uint8, TFixedAllocator<DMX_UNIVERSE_SIZE>>> CachedUniverseToDataMap = UniverseToDataMap;
 	
-	for (TTuple<int32, TArray<uint8, TFixedAllocator<DMX_UNIVERSE_SIZE>>>& UniverseToDataKvp : UniverseToDataMap)
+	for (const TTuple<int32, TArray<uint8, TFixedAllocator<DMX_UNIVERSE_SIZE>>>& UniverseToDataKvp : CachedUniverseToDataMap)
 	{
 		if (UniverseToDataKvp.Key < MinUniverseID || UniverseToDataKvp.Key > MaxUniverseID)
 		{
