DMX - UI for Universe view stuttering with specific scale

diff --git a/Engine/Plugins/VirtualProduction/DMX/DMXEngine/Source/DMXEditor/Private/Widgets/FixturePatch/SDMXFixturePatcher.cpp b/Engine/Plugins/VirtualProduction/DMX/DMXEngine/Source/DMXEditor/Private/Widgets/FixturePatch/SDMXFixturePatcher.cpp
--- a/Engine/Plugins/VirtualProduction/DMX/DMXEngine/Source/DMXEditor/Private/Widgets/FixturePatch/SDMXFixturePatcher.cpp
+++ b/Engine/Plugins/VirtualProduction/DMX/DMXEngine/Source/DMXEditor/Private/Widgets/FixturePatch/SDMXFixturePatcher.cpp
@@ -133,7 +133,9 @@ void SDMXFixturePatcher::Construct(const FArguments& InArgs)
 					.VAlign(VAlign_Fill)
 					[	
 						SAssignNew(PatchedUniverseScrollBox, SScrollBox)					
-						.Orientation(EOrientation::Orient_Vertical)	
+						.Orientation(EOrientation::Orient_Vertical)
+						.ScrollBarAlwaysVisible(true)
+						.ScrollBarThickness(FVector2D(5.f, 0.f))
 					]
 				]
 			];
