DMX - Split EntityList classes and refactor into two independent types FixtureTypeTree and FixturePatchTree

diff --git a/Engine/Plugins/VirtualProduction/DMX/DMXEngine/Source/DMXEditor/Private/DragDrop/DMXEntityDragDropOp.cpp b/Engine/Plugins/VirtualProduction/DMX/DMXEngine/Source/DMXEditor/Private/DragDrop/DMXEntityDragDropOp.cpp
--- a/Engine/Plugins/VirtualProduction/DMX/DMXEngine/Source/DMXEditor/Private/DragDrop/DMXEntityDragDropOp.cpp
+++ b/Engine/Plugins/VirtualProduction/DMX/DMXEngine/Source/DMXEditor/Private/DragDrop/DMXEntityDragDropOp.cpp
@@ -6,7 +6,6 @@
 #include "Library/DMXEntity.h"
 #include "Library/DMXEntityFixturePatch.h"
 #include "Library/DMXLibrary.h"
-#include "Widgets/SDMXEntityList.h"
 
 #include "Dialogs/Dialogs.h"
 
@@ -45,202 +44,6 @@ void FDMXEntityDragDropOperation::AddReferencedObjects(FReferenceCollector& Coll
 	Collector.AddReferencedObject(DraggedFromLibrary);
 }
 
-void FDMXEntityDragDropOperation::HandleDragEnterCategoryRow(const TSharedPtr<SDMXCategoryRow>& CategoryRow)
-{
-	bool bShowFeedback = true;
-	TestCanDropOnCategoryRow(CategoryRow, bShowFeedback);
-}
-
-bool FDMXEntityDragDropOperation::HandleDropOnCategoryRow(const TSharedPtr<SDMXCategoryRow>& CategoryRow)
-{
-	if (TSharedPtr<SDMXEntityList> EntityListPinned = CategoryRow->GetEntityList().Pin())
-	{
-		bool bShowFeedback = false;
-		if (!TestCanDropOnCategoryRow(CategoryRow, bShowFeedback))
-		{
-			return false;
-		}
-
-		UDMXLibrary* Library = EntityListPinned->GetDMXLibrary();
-		check(Library && DraggedFromLibrary == Library);
-
-		TSharedPtr<FDMXCategoryTreeNode> CategoryNode = CategoryRow->GetNode();
-		check(CategoryNode.IsValid());
-
-		// Register transaction and current DMX library state for Undo
-		const FScopedTransaction ChangeCategoryTransaction = FScopedTransaction(
-			FText::Format(LOCTEXT("ChangeEntitiesCategory", "Change {0}|plural(one=Entity, other=Entities) category"), DraggedEntities.Num())
-		);
-		EntityListPinned->GetDMXLibrary()->Modify();
-
-		switch (CategoryNode->GetCategoryType())
-		{
-			case FDMXTreeNodeBase::ECategoryType::DMXCategory:
-			{
-				const FDMXFixtureCategory& FixtureCategory = CategoryNode->GetCategoryValue();
-				for (TWeakObjectPtr<UDMXEntity> Entity : DraggedEntities)
-				{
-					if (UDMXEntityFixtureType* FixtureType = Cast<UDMXEntityFixtureType>(Entity))
-					{
-						FixtureType->Modify();
-
-						FixtureType->PreEditChange(nullptr);
-						FixtureType->DMXCategory = FixtureCategory;
-						FixtureType->PostEditChange();
-					}
-				}
-			}
-			break;
-
-			case FDMXTreeNodeBase::ECategoryType::UniverseID:
-			case FDMXTreeNodeBase::ECategoryType::FixtureAssignmentState:
-			{
-				int32 UniverseID = CategoryNode->GetIntValue();
-				for (TWeakObjectPtr<UDMXEntity> Entity : DraggedEntities)
-				{
-					if (UDMXEntityFixturePatch* FixturePatch = Cast<UDMXEntityFixturePatch>(Entity))
-					{
-						FixturePatch->Modify();
-
-						FixturePatch->PreEditChange(nullptr);
-						if (FixturePatch->IsAutoAssignAddress())
-						{
-							FDMXEditorUtils::TryAutoAssignToUniverses(FixturePatch, TSet<int32>({ UniverseID }));
-						}
-						FixturePatch->SetUniverseID(UniverseID);
-						FixturePatch->PostEditChange();
-					}
-				}
-			}
-
-			default:
-				// The other category types don't change properties
-				break;
-		}
-
-		if (CategoryNode->GetChildren().Num() > 0)
-		{
-			// Index after last entity in hovered category
-			UDMXEntity* LastEntityInCategory = CategoryNode->GetChildren().Last(0)->GetEntity();
-			const int32 LastEntityIndex = Library->FindEntityIndex(LastEntityInCategory);
-			check(LastEntityIndex != INDEX_NONE);
-
-			// Move dragged entities after the last ones in the category
-			// Reverse for to keep dragged entities order.
-			int32 DesiredIndex = LastEntityIndex + 1;
-			for (int32 EntityIndex = DraggedEntities.Num() - 1; EntityIndex > -1; --EntityIndex)
-			{
-				if (UDMXEntity* Entity = DraggedEntities[EntityIndex].Get())
-				{
-					Library->SetEntityIndex(Entity, DesiredIndex);
-				}
-			}
-		}
-
-		return true;
-	}
-
-	return false;
-}
-
-void FDMXEntityDragDropOperation::HandleDragEnterEntityRow(const TSharedPtr<SDMXEntityRow>& EntityRow)
-{
-	bool bShowFeedback = true;
-	TestCanDropOnEntityRow(EntityRow, bShowFeedback);
-}
-
-bool FDMXEntityDragDropOperation::HandleDropOnEntityRow(const TSharedPtr<SDMXEntityRow>& EntityRow)
-{
-	if (TSharedPtr<SDMXEntityList> EntityListPinned = EntityRow->GetEntityList().Pin())
-	{
-		UDMXLibrary* Library = EntityListPinned->GetDMXLibrary();
-		check(Library && DraggedFromLibrary == Library);
-
-		TSharedPtr<FDMXEntityTreeNode> Node = EntityRow->GetNode();
-		check(Node);
-
-		UDMXEntity* HoveredEntity = Node->GetEntity();
-		check(HoveredEntity);
-
-		bool bShowFeedback = false;
-		if (!TestCanDropOnEntityRow(EntityRow, bShowFeedback))
-		{
-			return false;
-		}
-
-		// Register transaction and current DMX library state for Undo
-		const FScopedTransaction ReorderTransaction = FScopedTransaction(
-			FText::Format(LOCTEXT("ReorderEntities", "Reorder {0}|plural(one=Entity, other=Entities)"), GetDraggedEntities().Num())
-		);
-
-		Library->Modify();
-
-		// The index of the Entity, we're about to insert the dragged one before
-		const int32 InsertBeforeIndex = Library->FindEntityIndex(HoveredEntity);
-		check(InsertBeforeIndex != INDEX_NONE);
-
-		// Reverse for to keep dragged entities order
-		for (int32 EntityIndex = DraggedEntities.Num() - 1; EntityIndex > -1; --EntityIndex)
-		{
-			if (DraggedEntities[EntityIndex].IsValid())
-			{
-				UDMXEntity* Entity = DraggedEntities[EntityIndex].Get();
-				Library->SetEntityIndex(Entity, InsertBeforeIndex);
-			}
-		}
-
-		// Handle cases where it was was dropped on another category
-		UClass* DraggedType = nullptr;
-		if (GetDraggedEntityTypes().Num() == 1)
-		{
-			DraggedType = GetDraggedEntityTypes()[0];
-		}
-
-		TSharedPtr<FDMXCategoryTreeNode> TargetCategory = StaticCastSharedPtr<FDMXCategoryTreeNode>(Node->GetParent().Pin());
-		check(TargetCategory.IsValid());
-
-		switch (TargetCategory->GetCategoryType())
-		{
-			case FDMXTreeNodeBase::ECategoryType::DMXCategory:
-			{
-				const FDMXFixtureCategory& FixtureCategory = TargetCategory->GetCategoryValue();
-				for (TWeakObjectPtr<UDMXEntity> Entity : DraggedEntities)
-				{
-					if (UDMXEntityFixtureType* FixtureType = Cast<UDMXEntityFixtureType>(Entity))
-					{
-						FixtureType->Modify();
-						FixtureType->DMXCategory = FixtureCategory;
-					}
-				}
-			}
-			break;
-
-			case FDMXTreeNodeBase::ECategoryType::UniverseID:
-			case FDMXTreeNodeBase::ECategoryType::FixtureAssignmentState:
-			{
-				const uint32& UniverseID = TargetCategory->GetIntValue();
-
-				for (TWeakObjectPtr<UDMXEntity> Entity : DraggedEntities)
-				{
-					if (UDMXEntityFixturePatch* FixturePatch = Cast<UDMXEntityFixturePatch>(Entity))
-					{
-						FixturePatch->Modify();
-
-						FDMXEditorUtils::AutoAssignedAddresses(TArray<UDMXEntityFixturePatch*>({ FixturePatch }));
-						FixturePatch->SetUniverseID(UniverseID);
-					}
-				}
-			}
-
-			default:
-				// The other category types don't change properties
-				break;
-		}
-	}
-
-	return true;
-}
-
 TArray<UClass*> FDMXEntityDragDropOperation::GetDraggedEntityTypes() const
 {
 	TArray<UClass*> EntityClasses;
@@ -321,306 +124,4 @@ void FDMXEntityDragDropOperation::SetCustomFeedbackWidget(const TSharedRef<SWidg
 	);
 }
 
-TArray<TSharedPtr<FDMXTreeNodeBase>> FDMXEntityDragDropOperation::GetDraggedFromCategories(const TSharedPtr<SDMXEntityList>& EntityList) const
-{
-	TArray<TSharedPtr<FDMXTreeNodeBase>> TreeNodes;
-	for (TWeakObjectPtr<UDMXEntity> Entity : DraggedEntities)
-	{
-		if (Entity.IsValid())
-		{
-			TSharedPtr<FDMXTreeNodeBase> Node = EntityList->GetCategoryNode(Entity.Get());
-			check(Node.IsValid());
-
-			TreeNodes.Add(Node);
-		}
-	}
-
-	return TreeNodes;
-}
-
-bool FDMXEntityDragDropOperation::TestCanDropOnCategoryRow(const TSharedPtr<SDMXCategoryRow>& CategoryRow, bool bShowFeedback)
-{
-	if (TSharedPtr<SDMXEntityList> EntityList = CategoryRow->GetEntityList().Pin())
-	{
-		TSharedPtr<FDMXCategoryTreeNode> CategoryNode = CategoryRow->GetNode();
-		check(CategoryNode.IsValid());
-		
-		UDMXLibrary* Library = EntityList->GetDMXLibrary();
-		check(Library);
-
-		if (CategoryNode.IsValid())
-		{
-
-			if (!CategoryNode->CanDropOntoCategory())
-			{
-				if (bShowFeedback)
-				{
-					// There wouldn't be any change by dragging the items into their own category.
-					SetFeedbackMessageError(FText::Format(
-						LOCTEXT("DragIntoNoValueCategory", "Cannot drop {0} onto this category"),
-						GetDraggedEntitiesName()
-					));
-				}
-				return false;
-			}
-
-			TArray<TSharedPtr<FDMXTreeNodeBase>> DraggedFromCategories = GetDraggedFromCategories(EntityList);
-
-			if (DraggedFromCategories.Num() == 1 && DraggedFromCategories[0] == CategoryNode)
-			{
-				if (bShowFeedback)
-				{
-					// There wouldn't be any change by dragging the items into their own category.
-					SetFeedbackMessageError(FText::Format(
-						LOCTEXT("DragIntoSelfCategory", "The selected {0} {1}|plural(one=is, other=are) already in this category"),
-						FDMXEditorUtils::GetEntityTypeNameText(EntityList->GetListType(), DraggedEntities.Num() > 1),
-						DraggedEntities.Num()
-					));
-				}
-
-				return false;
-			}
-			else
-			{
-				if (bShowFeedback)
-				{
-					// Some (or all) items will have a property changed because they come from another category
-					FText PropertyChangeName;
-					FText PropertyNewValue;
-
-					TSharedPtr<FDMXCategoryTreeNode> TargetCategory = CategoryRow->GetNode();
-					check(TargetCategory.IsValid());
-
-					// Generate 'OK' feedback per category
-					FText ListTypeText;
-					FText CategoryText;
-
-					const TSubclassOf<UDMXEntity> ListType = EntityList->GetListType();
-					if (ListType == UDMXEntityFixtureType::StaticClass())
-					{
-						ListTypeText = LOCTEXT("Property_DMXCategory", "DMX Category");
-						CategoryText = TargetCategory->GetDisplayName();
-					}
-					else if (ListType == UDMXEntityFixturePatch::StaticClass())
-					{
-						ListTypeText = LOCTEXT("Property_Universe", "Universe");
-						const uint32 UniverseID = TargetCategory->GetIntValue();
-						CategoryText = FText::AsNumber(UniverseID);
-					}
-
-					SetFeedbackMessageOK(FText::Format(
-						LOCTEXT("ReorderBeforeAndSetProperty", "Reorder {0} before '{1}'\nSet {2} = '{3}'"),
-						GetDraggedEntitiesName(),
-						TargetCategory->GetDisplayName(),
-						ListTypeText,
-						CategoryText
-					));
-				}
-
-				return true;
-			}
-		}
-	}
-
-	return false;
-}
-
-bool FDMXEntityDragDropOperation::TestCanDropOnEntityRow(const TSharedPtr<SDMXEntityRow>& EntityRow, bool bShowFeedback)
-{
-	if (TSharedPtr<SDMXEntityList> EntityList = EntityRow->GetEntityList().Pin())
-	{
-		TSharedPtr<FDMXEntityTreeNode> Node = EntityRow->GetNode();
-		check(Node);
-
-		UDMXEntity* HoveredEntity = Node->GetEntity();
-		check(HoveredEntity);
-
-		int32 NumDraggedEntities = DraggedEntities.Num();
-
-		// Test the drag is comming from the same library
-		UDMXLibrary* Library = EntityList->GetDMXLibrary();
-		check(Library);
-		if (!TestLibraryEquals(Library, bShowFeedback))
-		{
-			return false;
-		}
-
-		// Test if the dragdrop op holds a valid object
-		if (!TestHasValidEntities(bShowFeedback))
-		{
-			return false;
-		}
-
-		// Test the drag is of same type
-		if (!TestAreDraggedEntitiesOfClass(HoveredEntity->GetClass()))
-		{
-			return false;
-		}
-
-		const TSubclassOf<UDMXEntity> ListType = EntityList->GetListType();
-		if (!TestAreDraggedEntitiesOfClass(ListType, bShowFeedback))
-		{
-			return false;
-		}
-
-		// Test the hovered entity is not the dragged entitiy
-		if (DraggedEntities.Contains(HoveredEntity))
-		{
-			if (bShowFeedback)
-			{
-				SetFeedbackMessageError(FText::Format(
-					LOCTEXT("ReorderBeforeItself", "Drop {0} on itself"),
-					GetDraggedEntitiesName()
-				));
-			}
-
-			return false;
-		}
-
-		// For Fixture Patches, test if the hovered entity resides in the same universe
-		if (UDMXEntityFixturePatch* HoveredFixturePatch = Cast<UDMXEntityFixturePatch>(HoveredEntity))
-		{
-			if (!TestCanDropOnFixturePatch(HoveredFixturePatch, bShowFeedback))
-			{
-				return false;
-			}
-		}
-
-		check(Node->GetParent().IsValid());
-		TSharedPtr<FDMXCategoryTreeNode> TargetCategory = StaticCastSharedPtr<FDMXCategoryTreeNode>(Node->GetParent().Pin());
-		check(TargetCategory.IsValid());
-
-		// Generate 'OK' feedback per category
-		FText ListTypeText;
-		FText CategoryText;
-
-		if (ListType == UDMXEntityFixtureType::StaticClass())
-		{
-			ListTypeText = LOCTEXT("Property_DMXCategory", "DMX Category");
-			CategoryText = TargetCategory->GetDisplayName();
-		}
-		else if (ListType == UDMXEntityFixturePatch::StaticClass())
-		{
-			ListTypeText = LOCTEXT("Property_Universe", "Universe");
-			const uint32 UniverseID = TargetCategory->GetIntValue();
-			CategoryText = UniverseID == MAX_uint32
-				? LOCTEXT("UnassignedUniverseIDValue", "Unassigned")
-				: FText::AsNumber(UniverseID);
-		}
-
-		SetFeedbackMessageOK(FText::Format(
-			LOCTEXT("ReorderBeforeAndSetProperty", "Reorder {0} before '{1}'\nSet {2} = '{3}'"),
-			GetDraggedEntitiesName(),
-			FText::FromString(HoveredEntity->GetDisplayName()),
-			ListTypeText,
-			CategoryText
-		));
-
-		return true;
-	}
-
-	return false;
-}
-
-bool FDMXEntityDragDropOperation::TestLibraryEquals(UDMXLibrary* DragToLibrary, bool bShowFeedback)
-{
-	check(DragToLibrary);
-	if (DragToLibrary && DragToLibrary == DraggedFromLibrary)
-	{
-		return true;
-	}
-
-	if (bShowFeedback)
-	{
-		// For now, we don't allow dragging entities from one library to the other
-		SetFeedbackMessageError(FText::Format(
-			LOCTEXT("CantDragToDifferentLibrary", "Cannot move {0} outside {1}|plural(one=its, other=their) library"),
-			GetDraggedEntitiesName(),
-			GetDraggedEntities().Num()
-		));
-	}
-
-	return false;
-}
-
-bool FDMXEntityDragDropOperation::TestHasValidEntities(bool bShowFeedback /** = true */)
-{
-	const TWeakObjectPtr<UDMXEntity>* ValidEntityPtr =
-		DraggedEntities.FindByPredicate([](TWeakObjectPtr<UDMXEntity> Entity) {
-			return Entity.IsValid();
-		});
-
-	if (ValidEntityPtr)
-	{
-		return true;
-	}
-
-	SetFeedbackMessageError(
-		LOCTEXT("DragNullAssetError", "Dragged object no longer available")
-	);
-
-	return false;
-}
-
-bool FDMXEntityDragDropOperation::TestAreDraggedEntitiesOfClass(TSubclassOf<UDMXEntity> EntityClass, bool bShowFeedback)
-{
-	for (TWeakObjectPtr<UDMXEntity> Entity : DraggedEntities)
-	{
-		if (!Entity.IsValid())
-		{
-			continue;
-		}
-
-		if (!Entity->IsA(EntityClass))
-		{
-			if (bShowFeedback)
-			{
-				// Don't allow dragging entities from a type onto a different type tab
-				SetFeedbackMessageError(FText::Format(
-					LOCTEXT("CantDragToDifferentType", "Cannot move {0} to {1} editor"),
-					GetDraggedEntitiesName(),
-					FDMXEditorUtils::GetEntityTypeNameText(EntityClass, GetDraggedEntities().Num() > 1)
-				));
-			}
-
-			return false;
-		}
-	}
-
-	return true;
-}
-
-bool FDMXEntityDragDropOperation::TestCanDropOnFixturePatch(UDMXEntityFixturePatch* HoveredFixturePatch, bool bShowFeedback)
-{
-	check(TestAreDraggedEntitiesOfClass(UDMXEntityFixturePatch::StaticClass(), false));
-
-	bool bDragBetweenUniverses = false;
-	for (TWeakObjectPtr<UDMXEntity> DraggedEntity : DraggedEntities)
-	{
-		if (UDMXEntityFixturePatch* DraggedFixturePatch = Cast<UDMXEntityFixturePatch>(DraggedEntity))
-		{
-			if (DraggedFixturePatch->GetUniverseID() != HoveredFixturePatch->GetUniverseID())
-			{
-				bDragBetweenUniverses = true;
-				break;
-			}
-		}
-	}
-
-	if (!bDragBetweenUniverses)
-	{
-		if (bShowFeedback)
-		{
-			SetFeedbackMessageError(FText::Format(
-				LOCTEXT("DragFixturePatchToItsOwnUniverse", "{0} already resides in this universe."),
-				GetDraggedEntitiesName()
-			));
-		}
-
-		return false;
-	}
-
-	return true;
-}
-
 #undef LOCTEXT_NAMESPACE
diff --git a/Engine/Plugins/VirtualProduction/DMX/DMXEngine/Source/DMXEditor/Private/DragDrop/DMXEntityDragDropOp.h b/Engine/Plugins/VirtualProduction/DMX/DMXEngine/Source/DMXEditor/Private/DragDrop/DMXEntityDragDropOp.h
--- a/Engine/Plugins/VirtualProduction/DMX/DMXEngine/Source/DMXEditor/Private/DragDrop/DMXEntityDragDropOp.h
+++ b/Engine/Plugins/VirtualProduction/DMX/DMXEngine/Source/DMXEditor/Private/DragDrop/DMXEntityDragDropOp.h
@@ -12,7 +12,6 @@ class FDMXTreeNodeBase;
 class SDMXCategoryRow;
 class FDMXEntityTreeNode;
 class SDMXEntityRow;
-class SDMXEntityList;
 class UDMXEntity;
 class UDMXEntityFixturePatch;
 class UDMXLibrary;
@@ -25,39 +24,18 @@ class FDMXEntityDragDropOperation
 public:
 	DRAG_DROP_OPERATOR_TYPE(FDMXEntityDragDropOperation, FDragDropOperation)
 
-	/**
-	 * Constructs the entity drag drop operation
-	 *
-	 *
-	 *
-	 *
-	 */
+	/**  Constructs the entity drag drop operation */
 	FDMXEntityDragDropOperation(UDMXLibrary* InLibrary, const TArray<TWeakObjectPtr<UDMXEntity>>& InEntities);
 
 protected:
 	/** Constructs the tooltip widget that follows the mouse */
 	virtual void Construct() override;
 
-	/** FGCObject interface */
+	//~ Begin FGCObject interface
 	virtual void AddReferencedObjects(FReferenceCollector& Collector) override;
-	virtual FString GetReferencerName() const override
-	{
-		return TEXT("FDMXEntityDragDropOperation");
-	}
+	//~ End FGCObject interface
 
 public:
-	/** Handles op entering a category row */
-	void HandleDragEnterCategoryRow(const TSharedPtr<SDMXCategoryRow>& CategoryRow);
-
-	/** Handles dropping the op on a category row, returns true if successfully dropped */
-	bool HandleDropOnCategoryRow(const TSharedPtr<SDMXCategoryRow>& CategoryRow);
-
-	/** Handles op entering an entity row */
-	void HandleDragEnterEntityRow(const TSharedPtr<SDMXEntityRow>& EntityRow);
-
-	/** Handles dropping the op on an entity row, returns true if successfully dropped */
-	bool HandleDropOnEntityRow(const TSharedPtr<SDMXEntityRow>& EntityRow);
-
 	/** Returns the entities dragged with this drag drop op */
 	const TArray<TWeakObjectPtr<UDMXEntity>>& GetDraggedEntities() const { return DraggedEntities; }
 
@@ -80,45 +58,6 @@ public:
 	void SetCustomFeedbackWidget(const TSharedRef<SWidget>& Widget);
 
 private:
-	/** Returns the categories the entities were dragged out of */
-	TArray<TSharedPtr<FDMXTreeNodeBase>> GetDraggedFromCategories(const TSharedPtr<SDMXEntityList>& EntityList) const;
-	
-	/**
-	 * Tests if the drag drop op can be dropped on specified category row.
-	 * Updates Feedback Message if bShowFeedback.
-	 */
-	bool TestCanDropOnCategoryRow(const TSharedPtr<SDMXCategoryRow>& CategoryRow, bool bShowFeedback = true);
-
-	/**
-	 * Tests if the drag drop op can be dropped on specified entity row.
-	 * Updates Feedback Message if bShowFeedback.
-	 */
-	bool TestCanDropOnEntityRow(const TSharedPtr<SDMXEntityRow>& EntityRow, bool bShowFeedback = true);
-
-	/**
-	 * Tests if the drag drop op can be dropped on specified library.
-	 * Updates Feedback Message if bShowFeedback.
-	 */
-	bool TestLibraryEquals(UDMXLibrary* DragToLibrary, bool bShowFeedback = true);
-
-	/**
-	 * Tests if the entity is valid.
-	 * Updates Feedback Message if bShowFeedback.
-	*/
-	bool TestHasValidEntities(bool bShowFeedback = true);
-
-	/**
-	 * Tests if the entity is valid.
-	 * Updates Feedback Message if bShowFeedback.
-	*/
-	bool TestAreDraggedEntitiesOfClass(TSubclassOf<UDMXEntity> EntityClass, bool bShowFeedback = true);
-
-	/**
-	 * Tests if the drag drop op can be dropped on the hovered fixture patch 
-	 * Updates Feedback Message if bShowFeedback.
-	 */
-	bool TestCanDropOnFixturePatch(UDMXEntityFixturePatch* HoveredFixturePatch, bool bShowFeedback = true);
-
 	/** The library that the entites were dragged out from */
 	UDMXLibrary* DraggedFromLibrary;
 
diff --git a/Engine/Plugins/VirtualProduction/DMX/DMXEngine/Source/DMXEditor/Private/Widgets/DMXEntityTreeNode.cpp b/Engine/Plugins/VirtualProduction/DMX/DMXEngine/Source/DMXEditor/Private/Widgets/DMXEntityTreeNode.cpp
new file mode 100644
--- /dev/null
+++ b/Engine/Plugins/VirtualProduction/DMX/DMXEngine/Source/DMXEditor/Private/Widgets/DMXEntityTreeNode.cpp
@@ -0,0 +1,318 @@
+// Copyright Epic Games, Inc. All Rights Reserved.
+
+#include "Widgets/DMXEntityTreeNode.h"
+
+#include "DMXRuntimeUtils.h"
+#include "Library/DMXEntity.h"
+
+
+#define LOCTEXT_NAMESPACE "FDMXEntityTreeNodeBase"
+
+FDMXEntityTreeNodeBase::FDMXEntityTreeNodeBase(FDMXEntityTreeNodeBase::ENodeType InNodeType)
+	: NodeType(InNodeType)
+{}
+
+FDMXEntityTreeNodeBase::ENodeType FDMXEntityTreeNodeBase::GetNodeType() const
+{
+	return NodeType;
+}
+
+void FDMXEntityTreeNodeBase::AddChild(TSharedPtr<FDMXEntityTreeNodeBase> InChildPtr)
+{
+	if (InChildPtr.IsValid())
+	{
+		// unparent from previous parent
+		if (InChildPtr->GetParent().IsValid())
+		{
+			InChildPtr->RemoveFromParent();
+		}
+
+		InChildPtr->ParentNodePtr = SharedThis(this);
+		Children.Add(InChildPtr);
+	}
+}
+
+void FDMXEntityTreeNodeBase::RemoveChild(TSharedPtr<FDMXEntityTreeNodeBase> InChildPtr)
+{
+	if (InChildPtr.IsValid())
+	{
+		if (Children.Contains(InChildPtr))
+		{
+			Children.Remove(InChildPtr);
+			InChildPtr->ParentNodePtr = nullptr;
+		}
+	}
+}
+
+void FDMXEntityTreeNodeBase::RemoveFromParent()
+{
+	if (GetParent().IsValid())
+	{
+		GetParent().Pin()->RemoveChild(SharedThis(this));
+	}
+}
+
+const TArray<TSharedPtr<FDMXEntityTreeNodeBase>>& FDMXEntityTreeNodeBase::GetChildren() const
+{
+	return Children;
+}
+
+void FDMXEntityTreeNodeBase::ClearChildren()
+{
+	Children.Empty(0);
+}
+
+void FDMXEntityTreeNodeBase::SortChildren()
+{
+	Children.HeapSort([](const TSharedPtr<FDMXEntityTreeNodeBase>& Left, const TSharedPtr<FDMXEntityTreeNodeBase>& Right)->bool
+		{
+			if (Left.IsValid() && Right.IsValid())
+			{
+				return *Left < *Right;
+			}
+			return false;
+		});
+}
+
+void FDMXEntityTreeNodeBase::SortChildren(TFunction<bool (const TSharedPtr<FDMXEntityTreeNodeBase>&, const TSharedPtr<FDMXEntityTreeNodeBase>&)> Predicate)
+{
+	Children.Sort(Predicate);
+}
+
+bool FDMXEntityTreeNodeBase::IsFlaggedForFiltration() const
+{
+	// if Unknown, throws an error message
+	return ensureMsgf(FilterFlags != (uint8)EFilteredState::Unknown, TEXT("Querying a bad filtration state.")) 
+		? (FilterFlags & (uint8)EFilteredState::FilteredInMask) == 0 
+		: false;
+}
+
+void FDMXEntityTreeNodeBase::UpdateCachedFilterState(bool bMatchesFilter, bool bUpdateParent)
+{
+	bool bFlagsChanged = false;
+	if ((FilterFlags & EFilteredState::Unknown) == EFilteredState::Unknown)
+	{
+		FilterFlags = EFilteredState::FilteredOut;
+		bFlagsChanged = true;
+	}
+
+	if (bMatchesFilter)
+	{
+		bFlagsChanged |= (FilterFlags & EFilteredState::MatchesFilter) == 0;
+		FilterFlags |= EFilteredState::MatchesFilter;
+	}
+	else
+	{
+		bFlagsChanged |= (FilterFlags & EFilteredState::MatchesFilter) != 0;
+		FilterFlags &= ~EFilteredState::MatchesFilter;
+	}
+
+	const bool bHadChildMatch = (FilterFlags & EFilteredState::ChildMatches) != 0;
+	// refresh the cached child state (don't update the parent, we'll do that below if it's needed)
+	RefreshCachedChildFilterState(/*bUpdateParent =*/false);
+
+	bFlagsChanged |= bHadChildMatch != ((FilterFlags & EFilteredState::ChildMatches) != 0);
+	if (bUpdateParent && bFlagsChanged)
+	{
+		ApplyFilteredStateToParent();
+	}
+}
+
+void FDMXEntityTreeNodeBase::SetExpansionState(bool bNewExpansionState)
+{
+	bShouldBeExpanded = bNewExpansionState;
+}
+
+void FDMXEntityTreeNodeBase::SetWarningStatus(const FText& InWarningToolTip)
+{
+	WarningToolTip = InWarningToolTip;
+}
+
+void FDMXEntityTreeNodeBase::SetErrorStatus(const FText& InErrorToolTip)
+{
+	ErrorToolTip = InErrorToolTip;
+}
+
+bool FDMXEntityTreeNodeBase::operator<(const FDMXEntityTreeNodeBase& Other) const
+{
+	const FString& ThisName = GetDisplayNameString();
+	const FString& OtherName = Other.GetDisplayNameString();
+	if (ThisName.IsNumeric() && OtherName.IsNumeric())
+	{
+		return FCString::Atoi(*ThisName) < FCString::Atoi(*OtherName);
+	}
+	
+	// If the names are strings with numbers at the end, separate them to compare name then number
+	FString NameOnlyThis;
+	FString NameOnlyOther;
+	int32 NumberThis = 0;
+	int32 NumberOther = 0;
+	if (FDMXRuntimeUtils::GetNameAndIndexFromString(ThisName, NameOnlyThis, NumberThis) && 
+		FDMXRuntimeUtils::GetNameAndIndexFromString(OtherName, NameOnlyOther, NumberOther) && 
+		NameOnlyThis.Equals(NameOnlyOther))
+	{
+		return NumberThis < NumberOther;
+	}
+
+	return ThisName < OtherName;
+}
+
+void FDMXEntityTreeNodeBase::RefreshCachedChildFilterState(bool bUpdateParent)
+{
+	const bool bContainedMatch = !IsFlaggedForFiltration();
+
+	FilterFlags &= ~EFilteredState::ChildMatches;
+	for (TSharedPtr<FDMXEntityTreeNodeBase> Child : Children)
+	{
+		if (!Child->IsFlaggedForFiltration())
+		{
+			FilterFlags |= EFilteredState::ChildMatches;
+			break;
+		}
+	}
+	const bool bContainsMatch = !IsFlaggedForFiltration();
+
+	const bool bStateChange = bContainedMatch != bContainsMatch;
+	if (bUpdateParent && bStateChange)
+	{
+		ApplyFilteredStateToParent();
+	}
+}
+
+void FDMXEntityTreeNodeBase::ApplyFilteredStateToParent()
+{
+	FDMXEntityTreeNodeBase* Child = this;
+
+	while (Child->ParentNodePtr.IsValid())
+	{
+		FDMXEntityTreeNodeBase* Parent = Child->ParentNodePtr.Pin().Get();
+
+		if (!IsFlaggedForFiltration())
+		{
+			if ((Parent->FilterFlags & EFilteredState::ChildMatches) == 0)
+			{
+				Parent->FilterFlags |= EFilteredState::ChildMatches;
+			}
+			else
+			{
+				// all parents from here on up should have the flag
+				break;
+			}
+		}
+		// have to see if this was the only child contributing to this flag
+		else if (Parent->FilterFlags & EFilteredState::ChildMatches)
+		{
+			Parent->FilterFlags &= ~EFilteredState::ChildMatches;
+			for (const TSharedPtr<FDMXEntityTreeNodeBase>& Sibling : Parent->Children)
+			{
+				if (Sibling.Get() == Child)
+				{
+					continue;
+				}
+
+				if (Sibling->FilterFlags & EFilteredState::FilteredInMask)
+				{
+					Parent->FilterFlags |= EFilteredState::ChildMatches;
+					break;
+				}
+			}
+
+			if (Parent->FilterFlags & EFilteredState::ChildMatches)
+			{
+				// another child added the flag back
+				break;
+			}
+		}
+		Child = Parent;
+	}
+}
+
+FDMXEntityTreeRootNode::FDMXEntityTreeRootNode()
+	: FDMXEntityTreeNodeBase(FDMXEntityTreeNodeBase::ENodeType::RootNode)
+{}
+
+FString FDMXEntityTreeRootNode::GetDisplayNameString() const
+{
+	return FString(TEXT("Root Node"));
+}
+
+FText FDMXEntityTreeRootNode::GetDisplayNameText() const
+{
+	return LOCTEXT("RootNodeName", "Root Node");
+}
+
+FDMXEntityTreeEntityNode::FDMXEntityTreeEntityNode(UDMXEntity* InEntity)
+	: FDMXEntityTreeNodeBase(FDMXEntityTreeNodeBase::ENodeType::EntityNode)
+{
+	DMXEntity = InEntity;
+}
+
+FString FDMXEntityTreeEntityNode::GetDisplayNameString() const
+{
+	if (UDMXEntity* Entity = GetEntity())
+	{
+		return Entity->GetDisplayName();
+	}
+	return TEXT("null entity");
+}
+
+FText FDMXEntityTreeEntityNode::GetDisplayNameText() const
+{
+	if (UDMXEntity* Entity = GetEntity())
+	{
+		return FText::FromString(Entity->GetDisplayName());
+	}
+
+	return LOCTEXT("NullEntityError", "Entity is null");
+}
+
+UDMXEntity* FDMXEntityTreeEntityNode::GetEntity() const
+{
+	return DMXEntity.Get();
+}
+
+
+FDMXEntityTreeCategoryNode::FDMXEntityTreeCategoryNode(ECategoryType InCategoryType, FText InCategoryName, const FText& InToolTip /*= FText::GetEmpty()*/)
+	: FDMXEntityTreeNodeBase(FDMXEntityTreeNodeBase::ENodeType::CategoryNode)
+	, CategoryType(InCategoryType)
+	, CategoryName(InCategoryName)
+	, ToolTip(InToolTip)
+	, bCanDropOntoCategory(false)
+	, IntValue(0)
+{}
+
+FDMXEntityTreeCategoryNode::FDMXEntityTreeCategoryNode(ECategoryType InCategoryType, FText InCategoryName, int32 Value, const FText& InToolTip /*= FText::GetEmpty()*/)
+	: FDMXEntityTreeNodeBase(FDMXEntityTreeNodeBase::ENodeType::CategoryNode)
+	, CategoryType(InCategoryType)
+	, CategoryName(InCategoryName)
+	, ToolTip(InToolTip)
+	, bCanDropOntoCategory(true)
+	, IntValue(Value)
+{}
+
+FDMXEntityTreeCategoryNode::FDMXEntityTreeCategoryNode(ECategoryType InCategoryType, FText InCategoryName, const FDMXFixtureCategory& Value, const FText& InToolTip /*= FText::GetEmpty()*/)
+	: FDMXEntityTreeNodeBase(FDMXEntityTreeNodeBase::ENodeType::CategoryNode)
+	, CategoryType(InCategoryType)
+	, CategoryName(InCategoryName)
+	, ToolTip(InToolTip)
+	, bCanDropOntoCategory(true)
+	, IntValue(0)
+	, CategoryValue(Value)
+{}
+
+FString FDMXEntityTreeCategoryNode::GetDisplayNameString() const
+{
+	return CategoryName.ToString();
+}
+
+FText FDMXEntityTreeCategoryNode::GetDisplayNameText() const
+{
+	return CategoryName;
+}
+
+FDMXEntityTreeCategoryNode::ECategoryType FDMXEntityTreeCategoryNode::GetCategoryType() const
+{
+	return CategoryType;
+}
+
+#undef LOCTEXT_NAMESPACE
diff --git a/Engine/Plugins/VirtualProduction/DMX/DMXEngine/Source/DMXEditor/Private/Widgets/DMXEntityTreeNode.h b/Engine/Plugins/VirtualProduction/DMX/DMXEngine/Source/DMXEditor/Private/Widgets/DMXEntityTreeNode.h
new file mode 100644
--- /dev/null
+++ b/Engine/Plugins/VirtualProduction/DMX/DMXEngine/Source/DMXEditor/Private/Widgets/DMXEntityTreeNode.h
@@ -0,0 +1,217 @@
+// Copyright Epic Games, Inc. All Rights Reserved.
+
+#pragma once
+
+#include "CoreMinimal.h"
+
+#include "DMXProtocolTypes.h"
+
+#include "UObject/WeakObjectPtr.h"
+#include "UObject/WeakObjectPtrTemplates.h"
+
+class FDMXEditor;
+class UDMXEntity;
+
+
+/** Abstract base class for nodes in a list of entities */
+class FDMXEntityTreeNodeBase
+	: public TSharedFromThis<FDMXEntityTreeNodeBase>
+{
+	/** Delegate for when the context menu requests a rename */
+	DECLARE_DELEGATE(FOnRenameRequested);
+
+public:
+	enum class ENodeType : uint8
+	{
+		CategoryNode,
+		EntityNode,
+		RootNode
+	};
+
+	/**  Constructor */
+	FDMXEntityTreeNodeBase(FDMXEntityTreeNodeBase::ENodeType InNodeType);
+	
+	/** Destructor */
+	virtual ~FDMXEntityTreeNodeBase() {}
+
+	/**  @return Whether or not this object represents a node that can be renamed from the entities tree. */
+	virtual bool CanRename() const = 0;
+
+	/**  @return The string to be used in the tree display. */
+	virtual FString GetDisplayNameString() const = 0;
+	
+	/**  @return The name of this node in text. */
+	virtual FText GetDisplayNameText() const = 0;
+	
+	/**  @return The type of this node. */
+	virtual ENodeType GetNodeType() const;
+
+	/** Add a child node to this node */
+	virtual void AddChild(TSharedPtr<FDMXEntityTreeNodeBase> InChildPtr);
+	
+	/** Remove a child node from this node */
+	virtual void RemoveChild(TSharedPtr<FDMXEntityTreeNodeBase> InChildPtr);
+	
+	/** Remove this node from its parent one */
+	virtual void RemoveFromParent();
+	
+	/**  @return Child nodes for this object. */
+	virtual const TArray<TSharedPtr<FDMXEntityTreeNodeBase>>& GetChildren() const;
+	
+	/**  Remove all child nodes from this node. */
+	virtual void ClearChildren();
+	
+	/**  Sort children by name */
+	virtual void SortChildren();
+	
+	/**  Sort children using custom predicate */
+	virtual void SortChildren(TFunction<bool (const TSharedPtr<FDMXEntityTreeNodeBase>&, const TSharedPtr<FDMXEntityTreeNodeBase>&)> Predicate);
+	
+	/**  @return This object's parent node (or an invalid reference if no parent is assigned). */
+	virtual TWeakPtr<FDMXEntityTreeNodeBase> GetParent() const { return ParentNodePtr; }
+
+	/** Query that determines if this item should be filtered out or not */
+	virtual bool IsFlaggedForFiltration() const;
+
+	/** Refreshes this item's filtration state. Use bUpdateParent to make sure the parent's EFilteredState::ChildMatches flag is properly updated based off the new state */
+	void UpdateCachedFilterState(bool bMatchesFilter, bool bUpdateParent);
+
+	/** Update this node's desired expansion state for when there're no filters */
+	void SetExpansionState(bool bNewExpansionState);
+
+	/** This node's desired expansion state for when there're no filters */
+	bool GetExpansionState() const { return bShouldBeExpanded; }
+
+	/** If the warning tool tip is not empty, the node will display a warning icon with said tool tip. */
+	void SetWarningStatus(const FText& InWarningToolTip);
+	
+	const FText& GetWarningStatus() const { return WarningToolTip; }
+	
+	/** If the error tool tip is not empty, the node will display an error icon with said tool tip. */
+	void SetErrorStatus(const FText& InErrorToolTip);
+	
+	const FText& GetErrorStatus() const { return ErrorToolTip; }
+
+	/** Operator used when sorting categories by name/number */
+	bool operator<(const FDMXEntityTreeNodeBase& Other) const;
+
+protected:
+	/** Updates the EFilteredState::ChildMatches flag, based off of children's current state */
+	void RefreshCachedChildFilterState(bool bUpdateParent);
+	
+	/** Used to update the EFilteredState::ChildMatches flag for parent nodes, when this item's filtration state has changed */
+	void ApplyFilteredStateToParent();
+
+	FText WarningToolTip;
+	
+	FText ErrorToolTip;
+
+private:
+	/** The type of the node */
+	ENodeType NodeType;
+
+	// Actual tree structure
+	TWeakPtr<FDMXEntityTreeNodeBase> ParentNodePtr;
+	
+	/** The children of the node */
+	TArray<TSharedPtr<FDMXEntityTreeNodeBase>> Children;
+
+	/** Register whether the node should be expanded when there's no search filter text */
+	bool bShouldBeExpanded = false;
+
+	/** Filtered states the node can take */
+	enum EFilteredState
+	{
+		FilteredOut = 0x00,
+		MatchesFilter = (1 << 0),
+		ChildMatches = (1 << 1),
+
+		FilteredInMask = (MatchesFilter | ChildMatches),
+		Unknown = 0xFC // ~FilteredInMask
+	};
+	
+	uint8 FilterFlags = static_cast<uint8>(EFilteredState::Unknown);
+};
+
+
+/** A root node in an an entity list */
+class FDMXEntityTreeRootNode
+	: public FDMXEntityTreeNodeBase
+{
+public:
+	FDMXEntityTreeRootNode();
+	 
+	// ~ start FDMXEntityTreeNodeBase interface
+	virtual bool CanRename() const override { return false; }
+	virtual FString GetDisplayNameString() const override;
+	virtual FText GetDisplayNameText() const override;
+	// ~ end FDMXEntityTreeNodeBase interface
+};
+
+
+/** An entity node in an an entity list */
+class FDMXEntityTreeEntityNode
+	: public FDMXEntityTreeNodeBase
+{
+public:
+	FDMXEntityTreeEntityNode(UDMXEntity* InEntity);
+
+	// ~ start FDMXEntityTreeNodeBase interface
+	virtual bool CanRename() const override { return true; }
+	virtual FString GetDisplayNameString() const override;
+	virtual FText GetDisplayNameText() const override;
+	// ~ end FDMXEntityTreeNodeBase interface
+
+	/** Returns the entity of this node */
+	virtual UDMXEntity* GetEntity() const;
+
+private:
+	/** The DMXEntity of the node */
+	TWeakObjectPtr<UDMXEntity> DMXEntity;
+};
+
+
+/** A category node in an an entity list */
+class FDMXEntityTreeCategoryNode
+	: public FDMXEntityTreeNodeBase
+{
+public:
+	enum class ECategoryType : uint8
+	{
+		DMXCategory,
+		FixtureAssignmentState,
+		UniverseID,
+		NONE
+	};
+
+	FDMXEntityTreeCategoryNode(ECategoryType InCategoryType, FText InCategoryName, const FText& ToolTip = FText::GetEmpty());
+	FDMXEntityTreeCategoryNode(ECategoryType InCategoryType, FText InCategoryName, int32 Value, const FText& ToolTip = FText::GetEmpty());
+	FDMXEntityTreeCategoryNode(ECategoryType InCategoryType, FText InCategoryName, const FDMXFixtureCategory& Value, const FText& ToolTip = FText::GetEmpty());
+
+	// ~ start FDMXEntityTreeNodeBase interface
+	virtual bool CanRename() const override { return false; }
+	virtual FString GetDisplayNameString() const override;
+	virtual FText GetDisplayNameText() const override;
+	ECategoryType GetCategoryType() const;
+	// ~ end FDMXEntityTreeNodeBase interface
+
+	const FText& GetToolTip() const { return ToolTip; }
+
+	bool CanDropOntoCategory() const { return bCanDropOntoCategory; }
+	int32 GetIntValue() const { return IntValue; }
+	const FDMXFixtureCategory& GetCategoryValue() { return CategoryValue; }
+
+private:
+	/** This node's category type */
+	ECategoryType CategoryType;
+	FText CategoryName;
+	FText ToolTip;
+
+	bool bCanDropOntoCategory;
+
+	/** The int value, if of UniverseID type */
+	int32 IntValue = 0;
+
+	/** The fixture category, if of FixtureCategory type */
+	FDMXFixtureCategory CategoryValue;
+};
diff --git a/Engine/Plugins/VirtualProduction/DMX/DMXEngine/Source/DMXEditor/Private/Widgets/FixturePatch/SDMXFixturePatchEditor.cpp b/Engine/Plugins/VirtualProduction/DMX/DMXEngine/Source/DMXEditor/Private/Widgets/FixturePatch/SDMXFixturePatchEditor.cpp
--- a/Engine/Plugins/VirtualProduction/DMX/DMXEngine/Source/DMXEditor/Private/Widgets/FixturePatch/SDMXFixturePatchEditor.cpp
+++ b/Engine/Plugins/VirtualProduction/DMX/DMXEngine/Source/DMXEditor/Private/Widgets/FixturePatch/SDMXFixturePatchEditor.cpp
@@ -6,12 +6,17 @@
 #include "DMXFixturePatchSharedData.h"
 #include "SDMXFixturePatcher.h"
 #include "SDMXFixturePatchInspector.h"
+#include "SDMXFixturePatchTree.h"
 #include "Library/DMXLibrary.h"
 #include "Library/DMXEntityFixturePatch.h"
-#include "Widgets/SDMXEntityList.h"
 
 #include "Widgets/Layout/SSplitter.h"
 
+#include "DMXSubsystem.h"
+#include "IO/DMXPortManager.h"
+#include "IO/DMXOutputPort.h"
+#include "IO/DMXOutputPortReference.h"
+
 
 #define LOCTEXT_NAMESPACE "SDMXFixturePatcher"
 
@@ -42,12 +47,12 @@ void SDMXFixturePatchEditor::Construct(const FArguments& InArgs)
 			+SSplitter::Slot()			
 			.Value(0.38f)
 			[
-				SAssignNew(EntityList, SDMXEntityList, UDMXEntityFixturePatch::StaticClass())
+				SAssignNew(FixturePatchTree, SDMXFixturePatchTree)
 				.DMXEditor(DMXEditorPtr)
-				.OnAutoAssignAddressChanged(this, &SDMXFixturePatchEditor::OnEntitiyListChangedAutoAssignAddress)
-				.OnEntitiesAdded(this, &SDMXFixturePatchEditor::OnEntityListAddedEntities)
-				.OnEntityOrderChanged(this, &SDMXFixturePatchEditor::OnEntityListChangedEntityOrder)
-				.OnEntitiesRemoved(this, &SDMXFixturePatchEditor::OnEntityListRemovedEntities)
+				.OnAutoAssignAddressChanged(this, &SDMXFixturePatchEditor::OnFixturePatchTreeChangedAutoAssignAddress)
+				.OnEntitiesAdded(this, &SDMXFixturePatchEditor::OnFixturePatchTreeAddedEntities)
+				.OnEntityOrderChanged(this, &SDMXFixturePatchEditor::OnFixturePatchTreeChangedEntityOrder)
+				.OnEntitiesRemoved(this, &SDMXFixturePatchEditor::OnFixturePatchTreeRemovedEntities)
 			]
 	
 			+SSplitter::Slot()
@@ -73,32 +78,31 @@ void SDMXFixturePatchEditor::Construct(const FArguments& InArgs)
 
 void SDMXFixturePatchEditor::RequestRenameOnNewEntity(const UDMXEntity* InEntity, ESelectInfo::Type SelectionType)
 {
-	check(EntityList.IsValid());
+	check(FixturePatchTree.IsValid());
 
-	EntityList->UpdateTree();
-	EntityList->SelectItemByEntity(InEntity, SelectionType);
-	EntityList->OnRenameNode();
+	FixturePatchTree->SelectItemByEntity(InEntity, SelectionType);
+	FixturePatchTree->UpdateTree();
 }
 
 void SDMXFixturePatchEditor::SelectEntity(UDMXEntity* InEntity, ESelectInfo::Type InSelectionType /*= ESelectInfo::Type::Direct*/)
 {
-	check(EntityList.IsValid());
+	check(FixturePatchTree.IsValid());
 
-	EntityList->SelectItemByEntity(InEntity, InSelectionType);
+	FixturePatchTree->SelectItemByEntity(InEntity, InSelectionType);
 }
 
 void SDMXFixturePatchEditor::SelectEntities(const TArray<UDMXEntity*>& InEntities, ESelectInfo::Type InSelectionType /*= ESelectInfo::Type::Direct*/)
 {
-	check(EntityList.IsValid());
+	check(FixturePatchTree.IsValid());
 
-	EntityList->SelectItemsByEntity(InEntities, InSelectionType);
+	FixturePatchTree->SelectItemsByEntities(InEntities, InSelectionType);
 }
 
 TArray<UDMXEntity*> SDMXFixturePatchEditor::GetSelectedEntities() const
 {
-	check(EntityList.IsValid());
+	check(FixturePatchTree.IsValid());
 
-	return EntityList->GetSelectedEntities();
+	return FixturePatchTree->GetSelectedEntities();
 }
 
 void SDMXFixturePatchEditor::SelectUniverse(int32 UniverseID)
@@ -114,7 +118,7 @@ void SDMXFixturePatchEditor::SelectUniverse(int32 UniverseID)
 	}
 }
 
-void SDMXFixturePatchEditor::OnEntitiyListChangedAutoAssignAddress(TArray<UDMXEntityFixturePatch*> ChangedPatches)
+void SDMXFixturePatchEditor::OnFixturePatchTreeChangedAutoAssignAddress(TArray<UDMXEntityFixturePatch*> ChangedPatches)
 {
 	check(FixturePatcher.IsValid());
 	check(ChangedPatches.Num() > 0);
@@ -140,20 +144,20 @@ void SDMXFixturePatchEditor::OnEntitiyListChangedAutoAssignAddress(TArray<UDMXEn
 	}
 }
 
-void SDMXFixturePatchEditor::OnEntityListAddedEntities()
+void SDMXFixturePatchEditor::OnFixturePatchTreeAddedEntities()
 {
 	check(FixturePatcher.IsValid());
 	FixturePatcher->RefreshFromLibrary();
 }
 
-void SDMXFixturePatchEditor::OnEntityListChangedEntityOrder()
+void SDMXFixturePatchEditor::OnFixturePatchTreeChangedEntityOrder()
 {
 	check(FixturePatcher.IsValid());
 	FixturePatcher->RefreshFromProperties();
 	FixturePatcher->SelectUniverseThatContainsSelectedPatches();
 }	
 
-void SDMXFixturePatchEditor::OnEntityListRemovedEntities()
+void SDMXFixturePatchEditor::OnFixturePatchTreeRemovedEntities()
 {
 	check(FixturePatcher.IsValid());
 	FixturePatcher->RefreshFromLibrary();
@@ -161,16 +165,16 @@ void SDMXFixturePatchEditor::OnEntityListRemovedEntities()
 
 void SDMXFixturePatchEditor::OnFixturePatcherPatchedFixturePatch()
 {
-	check(EntityList.IsValid());
-	EntityList->UpdateTree();
+	check(FixturePatchTree.IsValid());
+	FixturePatchTree->UpdateTree();
 }
 
 void SDMXFixturePatchEditor::OnFinishedChangingProperties(const FPropertyChangedEvent& PropertyChangedEvent)
 {
-	check(EntityList.IsValid());
+	check(FixturePatchTree.IsValid());
 	check(FixturePatcher.IsValid());
 
-	EntityList->UpdateTree();
+	FixturePatchTree->UpdateTree();
 	FixturePatcher->NotifyPropertyChanged(PropertyChangedEvent);
 }
 
diff --git a/Engine/Plugins/VirtualProduction/DMX/DMXEngine/Source/DMXEditor/Private/Widgets/FixturePatch/SDMXFixturePatchEditor.h b/Engine/Plugins/VirtualProduction/DMX/DMXEngine/Source/DMXEditor/Private/Widgets/FixturePatch/SDMXFixturePatchEditor.h
--- a/Engine/Plugins/VirtualProduction/DMX/DMXEngine/Source/DMXEditor/Private/Widgets/FixturePatch/SDMXFixturePatchEditor.h
+++ b/Engine/Plugins/VirtualProduction/DMX/DMXEngine/Source/DMXEditor/Private/Widgets/FixturePatch/SDMXFixturePatchEditor.h
@@ -5,11 +5,11 @@
 #include "Widgets/SDMXEntityEditor.h"
 
 #include "CoreMinimal.h"
-#include "Widgets/SDMXEntityList.h"
 
 class FDMXEditor;
 class SDMXEntityInspector;
 class SDMXFixturePatcher;
+class SDMXFixturePatchTree;
 class UDMXEntityFixturePatch;
 
 struct FPropertyChangedEvent;
@@ -28,7 +28,7 @@ public:
 
 	SLATE_END_ARGS()
 
-public:
+public:	
 	/** Constructs this widget */
 	void Construct(const FArguments& InArgs);
 
@@ -45,16 +45,16 @@ protected:
 	void SelectUniverse(int32 UniverseID);
 
 	/** Callback for when entities list changes the auto assign flag of a patch */
-	void OnEntitiyListChangedAutoAssignAddress(TArray<UDMXEntityFixturePatch*> ChangedPatches);
+	void OnFixturePatchTreeChangedAutoAssignAddress(TArray<UDMXEntityFixturePatch*> ChangedPatches);
 
 	/** Called when the entity list added entities to the library */
-	void OnEntityListAddedEntities();
+	void OnFixturePatchTreeAddedEntities();
 
 	/** Called when the entity list changed the entity order */
-	void OnEntityListChangedEntityOrder();
+	void OnFixturePatchTreeChangedEntityOrder();
 
 	/** Called when the entity list removed entities from the library */
-	void OnEntityListRemovedEntities();
+	void OnFixturePatchTreeRemovedEntities();
 
 	/** Called when the fixture patcher patched a fixture patch */
 	void OnFixturePatcherPatchedFixturePatch();
@@ -70,7 +70,7 @@ protected:
 	TWeakPtr<FDMXEditor> DMXEditorPtr;
 
 	/** Left child widget */
-	TSharedPtr<SDMXEntityList> EntityList;
+	TSharedPtr<SDMXFixturePatchTree> FixturePatchTree;
 
 	/** Right child widget */
 	TSharedPtr<SDMXEntityInspector> InspectorWidget;
diff --git a/Engine/Plugins/VirtualProduction/DMX/DMXEngine/Source/DMXEditor/Private/Widgets/FixturePatch/SDMXFixturePatchTree.cpp b/Engine/Plugins/VirtualProduction/DMX/DMXEngine/Source/DMXEditor/Private/Widgets/FixturePatch/SDMXFixturePatchTree.cpp
new file mode 100644
--- /dev/null
+++ b/Engine/Plugins/VirtualProduction/DMX/DMXEngine/Source/DMXEditor/Private/Widgets/FixturePatch/SDMXFixturePatchTree.cpp
@@ -0,0 +1,803 @@
+// Copyright Epic Games, Inc. All Rights Reserved.
+
+#include "SDMXFixturePatchTree.h"
+
+#include "DMXEditor.h"
+#include "DMXEditorLog.h"
+#include "DMXEditorUtils.h"
+#include "DMXFixturePatchSharedData.h"
+#include "SDMXFixturePatchTreeFixturePatchRow.h"
+#include "SDMXFixturePatchTreeUniverseRow.h"
+#include "Commands/DMXEditorCommands.h"
+#include "IO/DMXInputPort.h"
+#include "IO/DMXOutputPort.h"
+#include "Library/DMXEntityFixturePatch.h"
+#include "Library/DMXEntityFixtureType.h"
+#include "Library/DMXLibrary.h"
+#include "Widgets/SDMXEntityDropdownMenu.h"
+
+#include "EditorStyleSet.h"
+#include "ScopedTransaction.h"
+#include "Framework/Commands/GenericCommands.h"
+#include "Framework/MultiBox/MultiBoxBuilder.h"
+#include "Widgets/SBoxPanel.h"
+#include "Widgets/SToolTip.h"
+#include "Widgets/Input/SCheckBox.h"
+#include "Widgets/Input/SComboButton.h"
+#include "Widgets/Images/SImage.h"
+#include "Widgets/Views/STreeView.h"
+
+
+#define LOCTEXT_NAMESPACE "SDMXFixturePatchTree"
+
+void SDMXFixturePatchTree::Construct(const FArguments& InArgs)
+{
+	SDMXEntityTreeViewBase::FArguments BaseArguments =
+		SDMXEntityTreeViewBase::FArguments()
+		.DMXEditor(InArgs._DMXEditor)
+		.OnSelectionChanged(InArgs._OnSelectionChanged)
+		.OnEntitiesAdded(InArgs._OnEntitiesAdded)
+		.OnEntityOrderChanged(InArgs._OnEntityOrderChanged)
+		.OnEntitiesRemoved(InArgs._OnEntitiesRemoved);
+
+	SDMXEntityTreeViewBase::Construct(BaseArguments);
+
+	if (TSharedPtr<FDMXEditor> PinnedEditor = DMXEditor.Pin())
+	{
+		FixturePatchSharedData = PinnedEditor->GetFixturePatchSharedData();
+	}
+	check(FixturePatchSharedData.IsValid());
+
+	// Bind to fixture patch shared data selection changes
+	FixturePatchSharedData->OnFixturePatchSelectionChanged.AddSP(this, &SDMXFixturePatchTree::OnSharedDataSelectedFixturePatches);
+
+	OnAutoAssignAddressChanged = InArgs._OnAutoAssignAddressChanged;
+}
+
+TSharedPtr<FDMXEntityTreeEntityNode> SDMXFixturePatchTree::CreateEntityNode(UDMXEntity* Entity)
+{
+	check(Entity && Entity->GetClass() == UDMXEntityFixturePatch::StaticClass());
+
+	TSharedPtr<FDMXEntityTreeEntityNode> NewNode = MakeShared<FDMXEntityTreeEntityNode>(Entity);
+	RefreshFilteredState(NewNode, false);
+
+	// Error status
+	FText InvalidReason;
+	if (!Entity->IsValidEntity(InvalidReason))
+	{
+		NewNode->SetErrorStatus(InvalidReason);
+	}
+
+	return NewNode;
+}
+
+TSharedRef<SWidget> SDMXFixturePatchTree::GenerateAddNewEntityButton()
+{
+	// Top part, with the  [+ Add New] button and the filter box
+	FText AddButtonLabel = FDMXEditorCommands::Get().AddNewEntityFixturePatch->GetLabel();
+	FText AddButtonToolTip = FDMXEditorCommands::Get().AddNewEntityFixturePatch->GetDescription();
+
+	TSharedRef<SComboButton> AddComboButton = 
+		SNew(SComboButton)
+			.ButtonContent()
+			[
+				SNew(SHorizontalBox)
+				+ SHorizontalBox::Slot()
+				.AutoWidth()
+				.Padding(FMargin(0, 1))
+				[
+					SNew(SImage)
+					.Image(FEditorStyle::GetBrush("Plus"))
+				]
+				+ SHorizontalBox::Slot()
+				.VAlign(VAlign_Center)
+				.AutoWidth()
+				.Padding(FMargin(2, 0, 2, 0))
+				[
+					SNew(STextBlock)
+					.Text(AddButtonLabel)
+				]
+			]
+			.MenuContent()
+			[
+				SAssignNew(AddButtonDropdownList, SDMXEntityDropdownMenu<UDMXEntityFixtureType>)
+				.DMXEditor(DMXEditor)
+				.OnEntitySelected(this, &SDMXFixturePatchTree::OnFixtureTypeSelected)
+			]
+			.IsFocusable(true)
+			.ContentPadding(FMargin(5.0f, 1.0f))
+			.ComboButtonStyle(FEditorStyle::Get(), "ToolbarComboButton")
+			.ButtonStyle(FEditorStyle::Get(), "FlatButton.Success")
+			.ForegroundColor(FLinearColor::White)
+			.ToolTipText(AddButtonToolTip)
+			.OnComboBoxOpened(FOnComboBoxOpened::CreateLambda([this]() { AddButtonDropdownList->ClearSelection(); }));
+
+	AddButtonDropdownList->SetComboButton(AddComboButton);
+
+	return AddComboButton;
+}
+
+void SDMXFixturePatchTree::RebuildNodes(const TSharedPtr<FDMXEntityTreeRootNode>& InRootNode)
+{
+	struct Local
+	{
+		static void AssignPatchesToUniversesOrShowErrors(
+			SDMXFixturePatchTree* FixturePatchTree, 
+			const TSharedPtr<FDMXEntityTreeCategoryNode>& AssignedFixturesCategoryNode, 
+			const TSharedPtr<FDMXEntityTreeCategoryNode>& UnassignedFixturesCategoryNode)
+		{
+			UDMXLibrary* Library = FixturePatchTree->GetDMXLibrary();
+			check(Library);
+			
+			Library->ForEachEntityOfType<UDMXEntityFixturePatch>([&] (UDMXEntityFixturePatch* FixturePatch)
+				{
+					// Create this entity's node
+					TSharedPtr<FDMXEntityTreeEntityNode> FixturePatchNode = FixturePatchTree->CreateEntityNode(FixturePatch);
+
+					const FText ErrorMessage = FixturePatchTree->CheckForPatchError(FixturePatch);
+					const bool bHasPatchNoErrors = ErrorMessage.IsEmpty();
+				
+					if (bHasPatchNoErrors && 
+						FixturePatch->GetFixtureType() && 
+						FixturePatch->GetChannelSpan() > 0)
+					{
+						constexpr FDMXEntityTreeCategoryNode::ECategoryType CategoryType = FDMXEntityTreeCategoryNode::ECategoryType::UniverseID;
+
+						TSharedPtr<FDMXEntityTreeCategoryNode> UniverseCategoryNode = FixturePatchTree->GetOrCreateCategoryNode(
+							CategoryType,
+							FText::Format(LOCTEXT("UniverseSubcategoryLabel", "Universe {0}"),	FText::AsNumber(FixturePatch->GetUniverseID())),
+							FixturePatch->GetUniverseID(),
+							AssignedFixturesCategoryNode
+						);
+
+						UniverseCategoryNode->AddChild(FixturePatchNode);
+					}
+					else
+					{
+						FixturePatchNode->SetErrorStatus(ErrorMessage);
+						UnassignedFixturesCategoryNode->AddChild(FixturePatchNode);
+					}
+				}
+			);
+		}
+		static void SortPatchesByChannel(const TSharedPtr<FDMXEntityTreeCategoryNode>& AssignedFixturesCategoryNode)
+		{
+			for (TSharedPtr<FDMXEntityTreeNodeBase> UniverseIDCategory : AssignedFixturesCategoryNode->GetChildren())
+			{			
+				// Sort Patches by starting channel
+				UniverseIDCategory->SortChildren([](const TSharedPtr<FDMXEntityTreeNodeBase>& A, const TSharedPtr<FDMXEntityTreeNodeBase>& B)->bool
+					{
+						check(A->GetNodeType() == FDMXEntityTreeNodeBase::ENodeType::EntityNode && B->GetNodeType() == FDMXEntityTreeNodeBase::ENodeType::EntityNode);
+						const TSharedPtr<FDMXEntityTreeEntityNode> EntityNodeA = StaticCastSharedPtr<FDMXEntityTreeEntityNode>(A);
+						const TSharedPtr<FDMXEntityTreeEntityNode> EntityNodeB = StaticCastSharedPtr<FDMXEntityTreeEntityNode>(B);
+
+						const UDMXEntityFixturePatch* PatchA = Cast<UDMXEntityFixturePatch>(EntityNodeA->GetEntity());
+						const UDMXEntityFixturePatch* PatchB = Cast<UDMXEntityFixturePatch>(EntityNodeB->GetEntity());
+						if (PatchA != nullptr && PatchB != nullptr)
+						{
+							const int32& ChannelA = PatchA->GetStartingChannel();
+							const int32& ChannelB = PatchB->GetStartingChannel();
+
+							return PatchA->GetStartingChannel() <= PatchB->GetStartingChannel();
+						}
+						return false;
+					});
+			}
+		}
+		static void ShowErrorForOverlappingPatches(const TSharedPtr<FDMXEntityTreeCategoryNode>& AssignedFixturesCategoryNode)
+		{
+			for (const TSharedPtr<FDMXEntityTreeNodeBase>& ChildNode : AssignedFixturesCategoryNode->GetChildren())
+			{
+				check(ChildNode->GetNodeType() == FDMXEntityTreeNodeBase::ENodeType::CategoryNode);
+				const TSharedPtr<FDMXEntityTreeCategoryNode> UniverseIDNode = StaticCastSharedPtr<FDMXEntityTreeCategoryNode>(ChildNode);
+
+				// Store the latest occupied channel in this Universe
+				int32 AvailableChannel = 1;
+				const UDMXEntityFixturePatch* PreviousPatch = nullptr;
+
+				for (TSharedPtr<FDMXEntityTreeNodeBase> CategoryChildNode : UniverseIDNode->GetChildren())
+				{
+					check(CategoryChildNode->GetNodeType() == FDMXEntityTreeNodeBase::ENodeType::EntityNode);
+					const TSharedPtr<FDMXEntityTreeEntityNode> EntityNode = StaticCastSharedPtr<FDMXEntityTreeEntityNode>(CategoryChildNode);
+
+					if (UDMXEntityFixturePatch* Patch = Cast<UDMXEntityFixturePatch>(EntityNode->GetEntity()))
+					{
+						const int32 ChannelSpan = Patch->GetChannelSpan();
+
+						FText InvalidReason;
+						if (Patch->GetStartingChannel() < AvailableChannel && PreviousPatch)
+						{
+							// The Patch is overlapping occupied channels
+							EntityNode->SetErrorStatus(FText::Format(
+								LOCTEXT("PatchOverlapWarning", "Start channel overlaps channels from {0}"),
+								FText::FromString(PreviousPatch->GetDisplayName())
+							));
+						}
+						else if (!Patch->IsValidEntity(InvalidReason))
+						{
+							// Update error status because after auto-channel changes there could be validation errors
+							EntityNode->SetErrorStatus(InvalidReason);
+						}
+						else
+						{
+							EntityNode->SetErrorStatus(FText::GetEmpty());
+						}
+
+						// Update the next available channel from this Patch's functions
+						AvailableChannel = Patch->GetStartingChannel() + ChannelSpan;
+
+						PreviousPatch = Patch;
+					}
+				}
+			}
+		}
+	};
+
+	EntityNodeToEntityRowMap.Reset();
+	
+	static constexpr uint32 UnassignedUniverseValue = MAX_uint32;
+	
+	// These nodes' categories are either Assigned or Unassigned
+	TSharedPtr<FDMXEntityTreeCategoryNode> AssignedFixturesCategoryNode = MakeShared<FDMXEntityTreeCategoryNode>(
+		FDMXEntityTreeCategoryNode::ECategoryType::FixtureAssignmentState,
+		LOCTEXT("AssignedFixturesCategory", "Assigned Fixtures"),
+		LOCTEXT("AssignedFixturesToolTip", "Patches which Universe IDs match are in port range")
+	);
+	TSharedPtr<FDMXEntityTreeCategoryNode> UnassignedFixturesCategoryNode = MakeShared<FDMXEntityTreeCategoryNode>(
+		FDMXEntityTreeCategoryNode::ECategoryType::FixtureAssignmentState,
+		LOCTEXT("UnassignedFixturesCategory", "Unassigned Fixtures"),
+		-1,
+		LOCTEXT("UnassignedFixturesToolTip", "Patches which Universe IDs are out of port range")
+	);
+	
+	InRootNode->AddChild(AssignedFixturesCategoryNode);
+	InRootNode->AddChild(UnassignedFixturesCategoryNode);
+	EntitiesTreeWidget->SetItemExpansion(AssignedFixturesCategoryNode, true);
+	EntitiesTreeWidget->SetItemExpansion(UnassignedFixturesCategoryNode, true);
+	
+	RefreshFilteredState(AssignedFixturesCategoryNode, false);
+	RefreshFilteredState(UnassignedFixturesCategoryNode, false);
+
+	Local::AssignPatchesToUniversesOrShowErrors(this, AssignedFixturesCategoryNode, UnassignedFixturesCategoryNode);
+	
+	AssignedFixturesCategoryNode->SortChildren();
+	Local::SortPatchesByChannel(AssignedFixturesCategoryNode);
+	
+	Local::ShowErrorForOverlappingPatches(AssignedFixturesCategoryNode);
+}
+
+TSharedRef<ITableRow> SDMXFixturePatchTree::OnGenerateRow(TSharedPtr<FDMXEntityTreeNodeBase> InNodePtr, const TSharedRef<STableViewBase>& OwnerTable)
+{
+	// Create the node of the appropriate type
+	if (InNodePtr->GetNodeType() == FDMXEntityTreeNodeBase::ENodeType::CategoryNode)
+	{
+		TSharedPtr<FDMXEntityTreeCategoryNode> CategoryNode = StaticCastSharedPtr<FDMXEntityTreeCategoryNode>(InNodePtr);
+		const bool bIsRootCategory = CategoryNode->GetCategoryType() != FDMXEntityTreeCategoryNode::ECategoryType::UniverseID;
+
+		return 
+			SNew(SDMXFixturePatchTreeUniverseRow, OwnerTable, CategoryNode, bIsRootCategory, SharedThis(this))
+			.OnFixturePatchOrderChanged(OnEntityOrderChanged)
+			[
+				SNew(STextBlock)
+				.Text(InNodePtr->GetDisplayNameText())
+				.TextStyle(FEditorStyle::Get(), "DetailsView.CategoryTextStyle")
+			];
+	}
+	else
+	{
+		TSharedPtr<FDMXEntityTreeEntityNode> EntityNode = StaticCastSharedPtr<FDMXEntityTreeEntityNode>(InNodePtr);
+
+		TSharedRef<SDMXFixturePatchTreeFixturePatchRow> FixturePatchRow = SNew(SDMXFixturePatchTreeFixturePatchRow, EntityNode, OwnerTable, SharedThis(this))
+			.OnGetFilterText(this, &SDMXEntityTreeViewBase::GetFilterText)
+			.OnEntityDragged(this, &SDMXEntityTreeViewBase::OnEntitiesDragged)
+			.OnFixturePatchOrderChanged(OnEntityOrderChanged)
+			.OnAutoAssignChannelStateChanged(this, &SDMXFixturePatchTree::OnAutoAssignChannelStateChanged, EntityNode);
+
+		EntityNodeToEntityRowMap.Add(EntityNode.ToSharedRef(), FixturePatchRow);
+
+		return FixturePatchRow;
+	}
+}
+
+TSharedPtr<SWidget> SDMXFixturePatchTree::OnContextMenuOpen()
+{
+	const bool bCloseWindowAfterMenuSelection = true;
+	FMenuBuilder MenuBuilder(bCloseWindowAfterMenuSelection, CommandList);
+
+	if (GetSelectedEntities().Num() > 0)
+	{
+		MenuBuilder.BeginSection("BasicOperations");
+		{
+			MenuBuilder.AddMenuEntry(FGenericCommands::Get().Rename);
+			MenuBuilder.AddMenuEntry(FGenericCommands::Get().Cut);
+			MenuBuilder.AddMenuEntry(FGenericCommands::Get().Copy);
+			MenuBuilder.AddMenuEntry(FGenericCommands::Get().Paste);
+			MenuBuilder.AddMenuEntry(FGenericCommands::Get().Duplicate);
+			MenuBuilder.AddMenuEntry(FGenericCommands::Get().Delete);
+		}
+		MenuBuilder.EndSection();
+	}
+	else 
+	{
+		return SNew(SDMXEntityDropdownMenu<UDMXEntityFixtureType>)
+			.DMXEditor(DMXEditor)
+			.OnEntitySelected(this, &SDMXFixturePatchTree::OnFixtureTypeSelected);
+		// TODO add (somehow) Paste option to this menu
+	}
+
+	return MenuBuilder.MakeWidget();
+}
+
+void SDMXFixturePatchTree::OnCutSelectedNodes()
+{
+	TArray<UDMXEntity*> SelectedEntities = GetSelectedEntities();
+	const FScopedTransaction Transaction(SelectedEntities.Num() > 1 ? LOCTEXT("CutFixturePatches", "Cut Fixture Patches") : LOCTEXT("CutFixturePatch", "Cut Fixture Patch"));
+
+	OnCopySelectedNodes();
+	OnDeleteNodes();
+}
+
+bool SDMXFixturePatchTree::CanCutNodes() const
+{
+	return CanCopyNodes() && CanDeleteNodes();
+}
+
+void SDMXFixturePatchTree::OnCopySelectedNodes()
+{
+	TArray<UDMXEntity*> SelectedEntities = GetSelectedEntities();
+
+	// Copy the entities to the clipboard
+	FDMXEditorUtils::CopyEntities(MoveTemp(SelectedEntities));
+}
+
+bool SDMXFixturePatchTree::CanCopyNodes() const
+{
+	TArray<UDMXEntity*> SelectedEntities = GetSelectedEntities();
+	return SelectedEntities.Num() > 0;
+}
+
+void SDMXFixturePatchTree::OnPasteNodes()
+{
+	// Get the Entities to paste from the clipboard
+	TArray<UDMXEntity*> NewObjects;
+	FDMXEditorUtils::GetEntitiesFromClipboard(NewObjects);
+	if (NewObjects.Num() != 0)
+	{
+		// Get the library that's being edited
+		UDMXLibrary* Library = GetDMXLibrary();
+		check(Library);
+
+		// Start transaction for Undo and take a snapshot of the current Library state
+		const FScopedTransaction PasteEntities(NewObjects.Num() > 1 ? LOCTEXT("PasteFixturePatches", "Paste Fixture Patches") : LOCTEXT("PasteFixturePatch", "Paste Fixture Patch"));
+		Library->Modify();
+
+		// If pasting Patches we'll need to check against existing Fixture Types
+		TArray<UDMXEntityFixtureType*> ExistingFixtureTypes;
+		if (NewObjects[0]->GetClass()->IsChildOf<UDMXEntityFixturePatch>())
+		{
+			ExistingFixtureTypes = Library->GetEntitiesTypeCast<UDMXEntityFixtureType>();
+		}
+		// Caches suitable replacements for pasted FixtureTypes (Pasted -> Existing replacement)
+		TMap<UDMXEntityFixtureType*, UDMXEntityFixtureType*> PatchTemplateReplacements;
+
+		// Add each pasted Entity to the Library
+		for (UDMXEntity* NewEntity : NewObjects)
+		{
+			UDMXEntityFixturePatch* FixurePatch = CastChecked<UDMXEntityFixturePatch>(NewEntity);
+
+			// Check for existing similar Fixture Type templates in this editor's Library to replace the temp one from copy
+			// or add the temp one if there's no suitable replacement
+
+			// Do we need to replace the template?
+			if (UDMXEntityFixtureType* CopiedPatchTemplate = FixurePatch->GetFixtureType())
+			{
+				// Did it come from this editor's DMX Library and does the original still exists?
+				if (UDMXEntityFixtureType* OriginalTemplate = Cast<UDMXEntityFixtureType>(Library->FindEntity(CopiedPatchTemplate->GetID())))
+				{
+					FixurePatch->SetFixtureType(OriginalTemplate);
+				}
+				else
+				{
+					check(CopiedPatchTemplate != nullptr);
+
+					// Is there already a suitable replacement registered for this template?
+					if (PatchTemplateReplacements.Contains(CopiedPatchTemplate))
+					{
+						// Replace the Patch's template with the replacement
+						FixurePatch->SetFixtureType(*PatchTemplateReplacements.Find(CopiedPatchTemplate));
+					}
+					else
+					{
+						// Search for a suitable replacement for the pasted Fixture Type, with identical
+						// properties, except for the Name, ID and Parent Library
+						bool bFoundReplacement = false;
+						for (UDMXEntityFixtureType* ExistingFixtureType : ExistingFixtureTypes)
+						{
+							if (FDMXEditorUtils::AreFixtureTypesIdentical(CopiedPatchTemplate, ExistingFixtureType))
+							{
+								FixurePatch->SetFixtureType(ExistingFixtureType);
+								PatchTemplateReplacements.Add(CopiedPatchTemplate, ExistingFixtureType);
+								bFoundReplacement = true;
+								break;
+							}
+						}
+
+						if (!bFoundReplacement)
+						{
+							// Move the Fixture Type template from the transient package into the Library package
+							NewEntity->Rename(*MakeUniqueObjectName(Library, UDMXEntityFixtureType::StaticClass()).ToString(), Library, REN_DoNotDirty | REN_DontCreateRedirectors);
+							// Make sure the Template's name and ID won't collide with existing Fixture Types
+							CopiedPatchTemplate->SetName(FDMXEditorUtils::FindUniqueEntityName(Library, UDMXEntityFixtureType::StaticClass(), CopiedPatchTemplate->GetDisplayName()));
+							CopiedPatchTemplate->RefreshID();
+							// Add to the Library
+							Library->AddEntity(CopiedPatchTemplate);
+						}
+					}
+				}
+			}
+
+			AutoAssignCopiedPatch(FixurePatch);
+
+			// Move the Entity from the transient package into the Library package
+			NewEntity->Rename(*MakeUniqueObjectName(Library, NewEntity->GetClass()).ToString(), Library, REN_DoNotDirty | REN_DontCreateRedirectors);
+			
+			// Make sure the Entity's name won't collide with existing ones
+			NewEntity->SetName(FDMXEditorUtils::FindUniqueEntityName(Library, NewEntity->GetClass(), NewEntity->GetDisplayName()));
+
+			Library->AddEntity(NewEntity);
+
+			OnEntitiesAdded.ExecuteIfBound();
+		}
+
+		UpdateTree();
+		SelectItemsByEntities(NewObjects, ESelectInfo::OnMouseClick);
+	}
+}
+
+bool SDMXFixturePatchTree::CanPasteNodes() const
+{
+	return FDMXEditorUtils::CanPasteEntities();
+}
+
+void SDMXFixturePatchTree::OnDuplicateNodes()
+{
+	TArray<UDMXEntity*> SelectedEntities = GetSelectedEntities();
+	// Sort selected entities by universe and starting channel to get a meaningful order when auto assign addresses
+	SelectedEntities.Sort([&](UDMXEntity& FirstEntity, UDMXEntity& SecondEntity) {
+		UDMXEntityFixturePatch* FirstPatch = CastChecked<UDMXEntityFixturePatch>(&FirstEntity);
+		UDMXEntityFixturePatch* SecondPatch = CastChecked<UDMXEntityFixturePatch>(&SecondEntity);
+		return
+			FirstPatch->GetUniverseID() < SecondPatch->GetUniverseID() ||
+			(FirstPatch->GetUniverseID() == SecondPatch->GetUniverseID() &&
+				FirstPatch->GetStartingChannel() <= SecondPatch->GetStartingChannel());
+		});
+	
+	UDMXLibrary* Library = GetDMXLibrary();
+	if (SelectedEntities.Num() > 0 && Library)
+	{
+		// Force the text box being edited (if any) to commit its text. The duplicate operation may trigger a regeneration of the tree view,
+		// releasing all row widgets. If one row was in edit mode (rename/rename on create), it was released before losing the focus and
+		// this would prevent the completion of the 'rename' or 'create + give initial name' transaction (occurring on focus lost).
+		FSlateApplication::Get().ClearKeyboardFocus();
+
+		const FScopedTransaction Transaction(SelectedEntities.Num() > 1 ? LOCTEXT("DuplicateFixturePatches", "Duplicate Fixture Patch") : LOCTEXT("DuplicateFixturePatch", "Duplicate Fixture Patch"));
+		Library->Modify();
+
+		TArray<UDMXEntity*> NewEntities;
+		NewEntities.Reserve(SelectedEntities.Num());
+
+		// We'll have the duplicates be placed right after their original counterparts
+		int32 NewEntityIndex = Library->FindEntityIndex(SelectedEntities.Last(0));
+		for (UDMXEntity* Entity : SelectedEntities)
+		{
+			FObjectDuplicationParameters DuplicationParams(Entity, GetDMXLibrary());
+			
+			if (UDMXEntity* EntityCopy = CastChecked<UDMXEntity>(StaticDuplicateObjectEx(DuplicationParams)))
+			{
+				EntityCopy->SetName(FDMXEditorUtils::FindUniqueEntityName(Library, EntityCopy->GetClass(), EntityCopy->GetDisplayName()));
+				NewEntities.Add(EntityCopy);
+				
+				Library->AddEntity(EntityCopy);
+				Library->SetEntityIndex(EntityCopy, ++NewEntityIndex);
+
+				if (UDMXEntityFixturePatch* FixturePatch = Cast<UDMXEntityFixturePatch>(EntityCopy))
+				{
+					AutoAssignCopiedPatch(FixturePatch);
+				}
+			}
+		}
+		
+		OnEntitiesAdded.ExecuteIfBound();
+		
+		// Need to refresh tree so new entities have nodes created for them
+		UpdateTree(); 
+
+		SelectItemsByEntities(NewEntities, ESelectInfo::OnMouseClick); // OnMouseClick triggers selection updated event
+	}
+}
+
+bool SDMXFixturePatchTree::CanDuplicateNodes() const
+{
+	return CanCopyNodes();
+}
+
+void SDMXFixturePatchTree::OnDeleteNodes()
+{
+	const TArray<UDMXEntity*> EntitiesToDelete = GetSelectedEntities();
+
+	// Clears references to the Entities and delete them
+	FDMXEditorUtils::RemoveEntities(GetDMXLibrary(), EntitiesToDelete);
+
+	OnEntitiesRemoved.ExecuteIfBound();
+
+	UpdateTree();
+}
+
+bool SDMXFixturePatchTree::CanDeleteNodes() const
+{
+	return GetSelectedEntities().Num() > 0;
+}
+
+void SDMXFixturePatchTree::OnRenameNode()
+{
+	TArray<TSharedPtr<FDMXEntityTreeNodeBase>> SelectedItems = EntitiesTreeWidget->GetSelectedItems();
+
+	if (SelectedItems.Num() == 1 && SelectedItems[0]->GetNodeType() == FDMXEntityTreeNodeBase::ENodeType::EntityNode)
+	{
+		const TSharedPtr<FDMXEntityTreeEntityNode> EntityNode = StaticCastSharedPtr<FDMXEntityTreeEntityNode>(SelectedItems[0]);
+		const TSharedPtr<SDMXFixturePatchTreeFixturePatchRow> FixturePatchRow = FindEntityRowByNode(EntityNode.ToSharedRef());
+
+		if (FixturePatchRow.IsValid())
+		{
+			FixturePatchRow->EnterRenameMode();
+		}
+	}
+}
+
+bool SDMXFixturePatchTree::CanRenameNode() const
+{
+	return EntitiesTreeWidget->GetSelectedItems().Num() == 1 && EntitiesTreeWidget->GetSelectedItems()[0]->CanRename();
+}
+
+TSharedPtr<SDMXFixturePatchTreeFixturePatchRow> SDMXFixturePatchTree::FindEntityRowByNode(const TSharedRef<FDMXEntityTreeEntityNode>& EntityNode)
+{
+	if (const TSharedRef<SDMXFixturePatchTreeFixturePatchRow>* RowPtr = EntityNodeToEntityRowMap.Find(EntityNode))
+	{
+		return *RowPtr;
+	}
+
+	return nullptr;
+}
+
+void SDMXFixturePatchTree::AutoAssignCopiedPatch(UDMXEntityFixturePatch* Patch) const
+{
+	check(Patch);
+
+	const int32 OriginalStartingChannel = Patch->GetStartingChannel();
+
+	Patch->Modify();
+	Patch->SetAutoAssignAddressUnsafe(true);
+	FDMXEditorUtils::AutoAssignedAddresses(TArray<UDMXEntityFixturePatch*>{ Patch }, OriginalStartingChannel + 1);
+}
+
+FText SDMXFixturePatchTree::CheckForPatchError(UDMXEntityFixturePatch* FixturePatch) const
+{
+	if (FixturePatch)
+	{
+		UDMXEntityFixtureType* FixtureType = FixturePatch->GetFixtureType();
+		if (!FixtureType || FixtureType->Modes.Num() == 0)
+		{
+			return LOCTEXT("FixtureTypeHasNoModes", "This patch's fixture type has no modes.");
+		}
+
+		const bool bHasAnyFunctions = [FixturePatch]()
+		{
+			for (const FDMXFixtureMode& Mode : FixturePatch->GetFixtureType()->Modes)
+			{
+				if (Mode.Functions.Num() > 0 ||
+					Mode.FixtureMatrixConfig.CellAttributes.Num() > 0)
+				{
+					return true;
+				}
+			}
+			return false;
+		}();
+		if (!bHasAnyFunctions)
+		{
+			return LOCTEXT("FixtureTypeHasNoFunctions", "This patch's fixture type has no functions.");
+		}
+		else if (FixturePatch->GetChannelSpan() == 0)
+		{
+			return LOCTEXT("FixtureTypeHasNoChannelSpan", "This patch has a channel span of 0.");
+		}
+		else if (FixturePatch->GetChannelSpan() > DMX_UNIVERSE_SIZE)
+		{
+			return LOCTEXT("ChannelSpanExceedsUniverseSize", "Patch uses more than 512 channels.");
+		}
+		else
+		{
+			const bool bValidUniverse = [FixturePatch, this]()
+			{
+				UDMXLibrary* DMXLibrary = GetDMXLibrary();
+				if (DMXLibrary)
+				{
+					for (const FDMXInputPortSharedRef& InputPort : DMXLibrary->GetInputPorts())
+					{
+						if (InputPort->IsLocalUniverseInPortRange(FixturePatch->GetUniverseID()))
+						{
+							return true;
+						}
+					}
+
+					for (const FDMXOutputPortSharedRef& OutputPort : DMXLibrary->GetOutputPorts())
+					{
+						if (OutputPort->IsLocalUniverseInPortRange(FixturePatch->GetUniverseID()))
+						{
+							return true;
+						}
+					}
+				}
+
+				return false;
+			}();
+
+			if (!bValidUniverse)
+			{
+				return LOCTEXT("InvalidUniverse", "Universe is not supported by any port used in the library.");
+			}							
+		}
+
+		return FText::GetEmpty();
+	}
+
+	ensureMsgf(0, TEXT("Trying to validate null fixture patch in Fixture Patch editor."));
+	return FText::GetEmpty();
+}
+
+
+void SDMXFixturePatchTree::OnSharedDataSelectedFixturePatches()
+{
+	check(FixturePatchSharedData.IsValid());
+	const TArray<TWeakObjectPtr<UDMXEntityFixturePatch>>& SelectedFixturePatches = FixturePatchSharedData->GetSelectedFixturePatches();
+	const TArray<UDMXEntity*>& SelectedEntitiesInTree = GetSelectedEntities();
+
+	bool bSelectionChanged = false;
+	if (SelectedEntitiesInTree.Num() != SelectedFixturePatches.Num())
+	{
+		bSelectionChanged = true;
+	}
+
+	for (TWeakObjectPtr<UDMXEntityFixturePatch> Patch : SelectedFixturePatches)
+	{
+		if (!SelectedEntitiesInTree.Contains(Patch))
+		{
+			bSelectionChanged = true;
+			break;
+		}
+	}
+
+	if (bSelectionChanged)
+	{
+		TArray<UDMXEntity*> NewSelection;
+		for (TWeakObjectPtr<UDMXEntityFixturePatch> Patch : SelectedFixturePatches)
+		{
+			if (Patch.IsValid())
+			{
+				NewSelection.Add(Patch.Get());
+			}
+		}
+		SelectItemsByEntities(NewSelection);
+	}
+}
+
+void SDMXFixturePatchTree::OnFixtureTypeSelected(UDMXEntity* InSelectedFixtureType)
+{
+	if (TSharedPtr<FDMXEditor> PinnedEditor = DMXEditor.Pin())
+	{
+		// Editor will call these during the creation of the new Entity
+		UDMXEntityFixtureType* AsFixtureType = CastChecked<UDMXEntityFixtureType>(InSelectedFixtureType);
+		OnGetBaseNameForNewEntityHandle = PinnedEditor->GetOnGetBaseNameForNewEntity().AddSP(this, &SDMXFixturePatchTree::OnEditorGetBaseNameForNewFixturePatch, AsFixtureType);
+		OnSetupNewEntityHandle = PinnedEditor->GetOnSetupNewEntity().AddSP(this, &SDMXFixturePatchTree::OnEditorSetupNewFixturePatch, AsFixtureType);
+
+		PinnedEditor->GetToolkitCommands()->ExecuteAction(FDMXEditorCommands::Get().AddNewEntityFixturePatch.ToSharedRef());
+	}
+}
+
+void SDMXFixturePatchTree::OnEditorGetBaseNameForNewFixturePatch(TSubclassOf<UDMXEntity> InEntityClass, FString& OutBaseName, UDMXEntityFixtureType* InSelectedFixtureType)
+{
+	if (!InEntityClass->IsChildOf(UDMXEntityFixturePatch::StaticClass()))
+	{
+		return;
+	}
+
+	TSharedPtr<FDMXEditor> PinnedEditor = DMXEditor.Pin();
+	if (PinnedEditor.IsValid())
+	{
+		PinnedEditor->GetOnGetBaseNameForNewEntity().Remove(OnGetBaseNameForNewEntityHandle);
+
+		OutBaseName = InSelectedFixtureType->GetDisplayName() + TEXT("_Patch");
+	}
+}
+
+void SDMXFixturePatchTree::OnEditorSetupNewFixturePatch(UDMXEntity* InNewEntity, UDMXEntityFixtureType* InSelectedFixtureType)
+{
+	if (UDMXEntityFixturePatch* FixturePatch = Cast<UDMXEntityFixturePatch>(InNewEntity))
+	{
+		TSharedPtr<FDMXEditor> PinnedEditor = DMXEditor.Pin();
+		if (PinnedEditor.IsValid())
+		{
+			PinnedEditor->GetOnSetupNewEntity().Remove(OnSetupNewEntityHandle);
+			FixturePatch->SetFixtureType(InSelectedFixtureType);
+		
+			FDMXEditorUtils::UpdatePatchColors(FixturePatch->GetParentLibrary());
+			
+			const bool bExceedsUniverseSize = FixturePatch->GetChannelSpan() > DMX_UNIVERSE_SIZE;
+			if(bExceedsUniverseSize)
+			{
+				FixturePatch->SetUniverseID(INDEX_NONE);
+			}
+			else
+			{
+				TSet<int32> UniverseIDsReachableByPorts = [this]()
+				{
+					if (UDMXLibrary* Library = GetDMXLibrary())
+					{
+						return Library->GetAllLocalUniversesIDsInPorts();
+					}
+
+					return TSet<int32>();
+				}();
+				FDMXEditorUtils::TryAutoAssignToUniverses(FixturePatch, UniverseIDsReachableByPorts);
+			}
+
+			// Issue a selection to trigger a OnSelectionUpdate and make the inspector display the new values
+			SelectItemByEntity(FixturePatch);
+		}
+	}
+	else
+	{
+		UE_LOG_DMXEDITOR(Error, TEXT("%S: New Entity wasn't a FixturePatch!"), __FUNCTION__);
+	}
+}
+
+void SDMXFixturePatchTree::OnAutoAssignChannelStateChanged(bool NewState, TSharedPtr<FDMXEntityTreeEntityNode> InNodePtr)
+{
+	const FScopedTransaction Transaction(LOCTEXT("SetAutoAssignChannelTransaction", "Set Auto Assign Channel"));
+
+	TArray<UDMXEntityFixturePatch*> ChangedPatches;
+
+	// Was the changed entity one of the selected ones?
+	if (EntitiesTreeWidget->IsItemSelected(InNodePtr))
+	{
+		for (UDMXEntity* Entity : GetSelectedEntities())
+		{
+			if (UDMXEntityFixturePatch* FixturePatch = Cast<UDMXEntityFixturePatch>(Entity))
+			{
+				if (FixturePatch->IsAutoAssignAddress() != NewState)
+				{
+					FixturePatch->Modify();
+					FixturePatch->SetAutoAssignAddressUnsafe(NewState);
+
+					ChangedPatches.Add(FixturePatch);
+				}
+			}
+		}
+	}
+	else
+	{
+		if (UDMXEntityFixturePatch* Patch = Cast<UDMXEntityFixturePatch>(InNodePtr->GetEntity()))
+		{
+			Patch->Modify();
+			Patch->SetAutoAssignAddressUnsafe(NewState);
+
+			ChangedPatches.Add(Patch);
+		}
+	}
+
+	FDMXEditorUtils::AutoAssignedAddresses(ChangedPatches);
+
+	UpdateTree();
+	OnAutoAssignAddressChanged.ExecuteIfBound(ChangedPatches);
+}
+
+#undef LOCTEXT_NAMESPACE
diff --git a/Engine/Plugins/VirtualProduction/DMX/DMXEngine/Source/DMXEditor/Private/Widgets/FixturePatch/SDMXFixturePatchTree.h b/Engine/Plugins/VirtualProduction/DMX/DMXEngine/Source/DMXEditor/Private/Widgets/FixturePatch/SDMXFixturePatchTree.h
new file mode 100644
--- /dev/null
+++ b/Engine/Plugins/VirtualProduction/DMX/DMXEngine/Source/DMXEditor/Private/Widgets/FixturePatch/SDMXFixturePatchTree.h
@@ -0,0 +1,122 @@
+// Copyright Epic Games, Inc. All Rights Reserved.
+
+#pragma once
+
+#include "Widgets/SDMXEntityTreeViewBase.h"
+
+#include "CoreMinimal.h"
+
+class FDMXEntityTreeEntityNode;
+class FDMXFixturePatchSharedData;
+template<typename TEntityType> class SDMXEntityDropdownMenu;
+class SDMXFixturePatchTreeFixturePatchRow;
+class UDMXEntityFixturePatch;
+class UDMXEntityFixtureType;
+
+
+/** 
+ * A tree of Fixture Patches in a DMX Library. 
+ */
+class SDMXFixturePatchTree
+	: public SDMXEntityTreeViewBase
+{
+	DECLARE_DELEGATE_OneParam(FDMXOnAutoAssignAddressChanged, TArray<UDMXEntityFixturePatch*>);
+
+public:
+	SLATE_BEGIN_ARGS(SDMXFixturePatchTree)
+	{}
+		/** The DMX Editor that owns this widget */
+		SLATE_ARGUMENT(TWeakPtr<FDMXEditor>, DMXEditor)
+
+		/** Exectued when the list changed its selection */
+		SLATE_EVENT(FDMXOnSelectionChanged, OnSelectionChanged)
+
+		/** Exectued when entites were added to the DMXEditor's library */
+		SLATE_EVENT(FSimpleDelegate, OnEntitiesAdded)
+
+		/** Exectued when entites were reorderd in the list, and potentially in the library */
+		SLATE_EVENT(FSimpleDelegate, OnEntityOrderChanged)
+
+		/** Exectued when entites were removed from the DMXEditor's library */
+		SLATE_EVENT(FSimpleDelegate, OnEntitiesRemoved)
+
+		/** Exectued when the auto assign address of fixture patches changed */
+		SLATE_EVENT(FDMXOnAutoAssignAddressChanged, OnAutoAssignAddressChanged)
+
+	SLATE_END_ARGS()
+
+	/** Constructs the widget */
+	void Construct(const FArguments& InArgs);
+
+	/** Destructor */
+	virtual ~SDMXFixturePatchTree() {};
+
+	/** Creates a Node for the entity */
+	TSharedPtr<FDMXEntityTreeEntityNode> CreateEntityNode(UDMXEntity* Entity);
+
+protected:
+	//~ Begin SDMXEntityTreeViewBase interface
+	virtual TSharedRef<SWidget> GenerateAddNewEntityButton();
+	virtual void RebuildNodes(const TSharedPtr<FDMXEntityTreeRootNode>& InRootNode);
+	virtual TSharedRef<ITableRow> OnGenerateRow(TSharedPtr<FDMXEntityTreeNodeBase> Node, const TSharedRef<STableViewBase>& OwnerTable);
+	virtual TSharedPtr<SWidget> OnContextMenuOpen() override;
+	virtual void OnCutSelectedNodes();
+	virtual bool CanCutNodes() const;
+	virtual void OnCopySelectedNodes();
+	virtual bool CanCopyNodes() const;
+	virtual void OnPasteNodes();
+	virtual bool CanPasteNodes() const;
+	virtual bool CanDuplicateNodes() const;
+	virtual void OnDuplicateNodes();
+	virtual void OnDeleteNodes();
+	virtual bool CanDeleteNodes() const;
+	virtual void OnRenameNode();
+	virtual bool CanRenameNode() const;
+	//~ End SDMXEntityTreeViewBase interface
+
+private:
+	/** Returns the row that corresponds to the node */
+	TSharedPtr<SDMXFixturePatchTreeFixturePatchRow> FindEntityRowByNode(const TSharedRef<FDMXEntityTreeEntityNode>& EntityNode);
+
+	/** Auto assigns a patch */
+	void AutoAssignCopiedPatch(UDMXEntityFixturePatch* Patch) const;
+
+	/** Called when fixture patches were selected in shared data */
+	void OnSharedDataSelectedFixturePatches();
+
+	/** Returns error message or FText::GetEmpty() if no error found */
+	FText CheckForPatchError(UDMXEntityFixturePatch* FixturePatch) const;
+
+	/** Called when the user selects a Fixture Type to create a Fixture Patch from */
+	void OnFixtureTypeSelected(UDMXEntity* InSelectedFixtureType);
+
+	/** Called by the editor to set a base name for an Entity about to be created */
+	void OnEditorGetBaseNameForNewFixturePatch(TSubclassOf<UDMXEntity> InEntityClass, FString& OutBaseName, UDMXEntityFixtureType* InSelectedFixtureType);
+
+	/** Called by the editor to setup the properties of a new Entity before its selection and renaming in the editor */
+	void OnEditorSetupNewFixturePatch(UDMXEntity* InNewEntity, UDMXEntityFixtureType* InSelectedFixtureType);
+
+	/** Called when Auto Assign Channel is changed for a Patch */
+	void OnAutoAssignChannelStateChanged(bool NewState, TSharedPtr<FDMXEntityTreeEntityNode> InNodePtr);
+
+	/** Map of entity nodes and their row in the tree */
+	TMap<TSharedRef<FDMXEntityTreeEntityNode>, TSharedRef<SDMXFixturePatchTreeFixturePatchRow>> EntityNodeToEntityRowMap;
+
+	/** A dropdown menu that shows fixture types from which patches can be created */
+	TSharedPtr<SDMXEntityDropdownMenu<UDMXEntityFixtureType>> AddButtonDropdownList;
+
+	/** Delegate to invoke when a fixture patch changed the auto assign address property */
+	FDMXOnAutoAssignAddressChanged OnAutoAssignAddressChanged;
+
+	/** Delegate broadcast when a new entity is setup */
+	FDelegateHandle OnSetupNewEntityHandle;
+
+	/** Handle to the registered OnActiveTabChanged delegate */
+	FDelegateHandle OnActiveTabChangedDelegateHandle;
+
+	/** Handle to the base name for a new entity requested delegate */
+	FDelegateHandle OnGetBaseNameForNewEntityHandle;
+
+	/** Shared data for fixture patch editors */
+	TSharedPtr<FDMXFixturePatchSharedData> FixturePatchSharedData;
+};
diff --git a/Engine/Plugins/VirtualProduction/DMX/DMXEngine/Source/DMXEditor/Private/Widgets/FixturePatch/SDMXFixturePatchTreeFixturePatchRow.cpp b/Engine/Plugins/VirtualProduction/DMX/DMXEngine/Source/DMXEditor/Private/Widgets/FixturePatch/SDMXFixturePatchTreeFixturePatchRow.cpp
new file mode 100644
--- /dev/null
+++ b/Engine/Plugins/VirtualProduction/DMX/DMXEngine/Source/DMXEditor/Private/Widgets/FixturePatch/SDMXFixturePatchTreeFixturePatchRow.cpp
@@ -0,0 +1,521 @@
+// Copyright Epic Games, Inc. All Rights Reserved.
+
+#include "SDMXFixturePatchTreeFixturePatchRow.h"
+
+#include "DMXEditorUtils.h"
+#include "DMXEditorStyle.h"
+#include "SDMXFixturePatchTree.h"
+#include "DragDrop/DMXEntityDragDropOp.h"
+#include "Library/DMXEntityFixturePatch.h"
+#include "Library/DMXLibrary.h"
+#include "Widgets/DMXEntityTreeNode.h"
+
+#include "EditorStyleSet.h"
+#include "ScopedTransaction.h"
+#include "Widgets/SToolTip.h"
+#include "Widgets/Docking/SDockTab.h"
+#include "Widgets/Input/SCheckBox.h"
+#include "Widgets/Images/SImage.h"
+#include "Widgets/Text/SInlineEditableTextBlock.h"
+#include "Widgets/Text/STextBlock.h"
+
+
+#define LOCTEXT_NAMESPACE "SDMXFixturePatchTreeFixturePatchRow"
+
+void SDMXFixturePatchTreeFixturePatchRow::Construct(const FArguments& InArgs, TSharedPtr<FDMXEntityTreeEntityNode> InEntityNode, TSharedPtr<STableViewBase> InOwnerTableView, TWeakPtr<SDMXFixturePatchTree> InFixturePatchTree)
+{
+	// Without ETableRowSignalSelectionMode::Instantaneous, when the user is editing a property in
+	// the inspector panel and then clicks on a different row on the list panel, the selection event
+	// is deferred. But because we update the tree right after a property change and that triggers
+	// selection changes too, the selection change event is triggered only from UpdateTree, with
+	// Direct selection mode, which doesn't trigger the SDMXEntityTree::OnSelectionUpdated event.
+	// This setting forces the event with OnMouseClick selection type to be fired as soon as the
+	// row is clicked.
+	STableRow<TSharedPtr<FDMXEntityTreeEntityNode>>::FArguments BaseArguments =
+		STableRow<TSharedPtr<FDMXEntityTreeEntityNode>>::FArguments()
+		.SignalSelectionMode(ETableRowSignalSelectionMode::Instantaneous)
+		.OnDragDetected(this, &SDMXFixturePatchTreeFixturePatchRow::HandleOnDragDetected);
+
+	STableRow<TSharedPtr<FDMXEntityTreeEntityNode>>::Construct(BaseArguments, InOwnerTableView.ToSharedRef());
+
+	WeakEntityNode = StaticCastSharedPtr<FDMXEntityTreeEntityNode>(InEntityNode);
+	WeakFixturePatchTree = InFixturePatchTree;
+
+	OnEntityDragged = InArgs._OnEntityDragged;
+	OnGetFilterText = InArgs._OnGetFilterText;
+	OnFixturePatchOrderChanged = InArgs._OnFixturePatchOrderChanged;
+	OnAutoAssignChannelStateChanged = InArgs._OnAutoAssignChannelStateChanged;
+
+	const FSlateFontInfo NameFont = FCoreStyle::GetDefaultFontStyle("Regular", 10);
+
+	// Used address range labels
+	const FSlateFontInfo&& ChannelFont = FCoreStyle::GetDefaultFontStyle("Bold", 8);
+	const FLinearColor ChannelLabelColor(1.0f, 1.0f, 1.0f, 0.8f);
+	const float MinChannelTextWidth = 20.0f;
+
+	SetContent
+	(
+		SNew(SHorizontalBox)
+
+		// Status icon to show the user if there's an error with the Entity's usability
+		+ SHorizontalBox::Slot()
+		.AutoWidth()
+		.HAlign(HAlign_Left)
+		.VAlign(VAlign_Center)
+		[
+			SNew(SImage)
+			.Image(this, &SDMXFixturePatchTreeFixturePatchRow::GetStatusIcon)
+			.ToolTipText(this, &SDMXFixturePatchTreeFixturePatchRow::GetStatusToolTip)
+		]
+		
+		// Entity's name
+		+ SHorizontalBox::Slot()
+		.FillWidth(1.0f)
+		.VAlign(VAlign_Center)
+		.Padding(3.0f, 0.0f)
+		[
+			SAssignNew(InlineEditableFixturePatchNameWidget, SInlineEditableTextBlock)
+			.Text(this, &SDMXFixturePatchTreeFixturePatchRow::GetFixturePatchDisplayText)
+			.Font(NameFont)
+			.HighlightText(this, &SDMXFixturePatchTreeFixturePatchRow::GetFilterText)
+			.OnTextCommitted(this, &SDMXFixturePatchTreeFixturePatchRow::OnNameTextCommit)
+			.OnVerifyTextChanged(this, &SDMXFixturePatchTreeFixturePatchRow::OnNameTextVerifyChanged)
+			.IsSelected(this, &SDMXFixturePatchTreeFixturePatchRow::IsSelected)
+			.IsReadOnly(false)
+		]
+
+		// Auto channel assignment check box
+		+ SHorizontalBox::Slot()
+		.AutoWidth()
+		.VAlign(VAlign_Center)
+		[
+			SNew(SCheckBox)
+			.IsChecked(this, &SDMXFixturePatchTreeFixturePatchRow::IsAutoAssignChannelEnabled)
+			.OnCheckStateChanged(this, &SDMXFixturePatchTreeFixturePatchRow::OnAutoAssignChannelBoxStateChanged)
+			.ToolTipText(LOCTEXT("AutoAssignChannelToolTip", "Auto-assign address from drag/drop list order"))
+		]
+
+		// Starting channel number
+		+ SHorizontalBox::Slot()
+		.AutoWidth()
+		.VAlign(VAlign_Center)
+		[
+			SNew(SBorder)
+			.VAlign(VAlign_Center)
+			.Padding(FMargin(0.0f, 3.0f))
+			.BorderImage(FDMXEditorStyle::Get().GetBrush("DMXEditor.BlackBrush"))
+			.BorderBackgroundColor(FLinearColor(1.0f, 1.0f, 1.0f, 0.15f))
+			.ToolTipText(LOCTEXT("ChannelStartToolTip", "Channels range: start"))
+			[
+				SNew(STextBlock)
+				.Text(this, &SDMXFixturePatchTreeFixturePatchRow::GetStartingChannelText)
+				.Font(ChannelFont)
+				.ColorAndOpacity(ChannelLabelColor)
+				.MinDesiredWidth(MinChannelTextWidth)
+				.Justification(ETextJustify::Center)
+			]
+		]
+
+		// Ending channel number
+		+ SHorizontalBox::Slot()
+		.AutoWidth()
+		.VAlign(VAlign_Center)
+		[
+			SNew(SBorder)
+			.VAlign(VAlign_Center)
+			.Padding(FMargin(0.0f, 3.0f))
+			.BorderImage(FDMXEditorStyle::Get().GetBrush("DMXEditor.BlackBrush"))
+			.BorderBackgroundColor(FLinearColor(1.0f, 1.0f, 1.0f, 0.25f)) // darker background
+			.ColorAndOpacity(FLinearColor(1.0f, 1.0f, 1.0f, 0.5f)) // darker text
+			.ToolTipText(LOCTEXT("ChannelEndToolTip", "Channels range: end"))
+			[
+				SNew(STextBlock)
+				.Text(this, &SDMXFixturePatchTreeFixturePatchRow::GetEndingChannelText)
+				.Font(ChannelFont)
+				.ColorAndOpacity(ChannelLabelColor)
+				.MinDesiredWidth(MinChannelTextWidth)
+				.Justification(ETextJustify::Center)
+			]
+		]
+	);
+}
+
+void SDMXFixturePatchTreeFixturePatchRow::EnterRenameMode()
+{
+	InlineEditableFixturePatchNameWidget->EnterEditingMode();
+}
+
+FReply SDMXFixturePatchTreeFixturePatchRow::OnMouseButtonDoubleClick(const FGeometry& InMyGeometry, const FPointerEvent& InMouseEvent)
+{
+	const TSharedPtr<SDMXFixturePatchTree> FixturePatchTree = WeakFixturePatchTree.Pin();
+	const TSharedPtr<FDMXEntityTreeEntityNode> EntityNode = WeakEntityNode.Pin();
+
+	if (FixturePatchTree.IsValid() && EntityNode.IsValid())
+	{
+		FixturePatchTree->SelectItemByNode(EntityNode.ToSharedRef());
+
+		InlineEditableFixturePatchNameWidget->EnterEditingMode();
+
+		return FReply::Handled();
+	}
+
+	return FReply::Unhandled();
+}
+
+void SDMXFixturePatchTreeFixturePatchRow::OnDragEnter(const FGeometry& MyGeometry, const FDragDropEvent& DragDropEvent)
+{
+	if (const TSharedPtr<FDMXEntityDragDropOperation> EntityDragDropOp = DragDropEvent.GetOperationAs<FDMXEntityDragDropOperation>())
+	{
+		TestCanDropWithFeedback(EntityDragDropOp.ToSharedRef());
+	}
+}
+
+FReply SDMXFixturePatchTreeFixturePatchRow::OnDrop(const FGeometry& MyGeometry, const FDragDropEvent& DragDropEvent)
+{
+	const TSharedPtr<FDMXEntityDragDropOperation> EntityDragDropOp = DragDropEvent.GetOperationAs<FDMXEntityDragDropOperation>();
+	const TSharedPtr<SDMXFixturePatchTree> FixturePatchTree = WeakFixturePatchTree.Pin();
+	const TSharedPtr<FDMXEntityTreeEntityNode> EntityNode = WeakEntityNode.Pin();
+	UDMXEntityFixturePatch* FixturePatch = EntityNode.IsValid() ? Cast<UDMXEntityFixturePatch>(EntityNode->GetEntity()) : nullptr;
+
+	if (EntityDragDropOp.IsValid() && FixturePatchTree.IsValid() && EntityNode.IsValid() && FixturePatch)
+	{
+		if (TestCanDropWithFeedback(EntityDragDropOp.ToSharedRef()))
+		{
+			const TArray<TWeakObjectPtr<UDMXEntity>> DraggedEntities = EntityDragDropOp->GetDraggedEntities();
+
+			// Register transaction and current DMX library state for Undo
+			const FScopedTransaction ReorderTransaction = FScopedTransaction(
+				FText::Format(LOCTEXT("DropFixturePatchOnFixturePatch", "Assign Fixture {0}|plural(one=Patch, other=Patches)"), DraggedEntities.Num())
+			);
+
+			UDMXLibrary* DMXLibrary = GetDMXLibrary();
+			DMXLibrary->Modify();
+
+			// The index of the Entity, we're about to insert the dragged one before
+			const int32 InsertBeforeIndex = DMXLibrary->FindEntityIndex(FixturePatch);
+			check(InsertBeforeIndex != INDEX_NONE);
+
+			// Reverse for to keep dragged entities order
+			for (int32 EntityIndex = DraggedEntities.Num() - 1; EntityIndex > -1; --EntityIndex)
+			{
+				if (DraggedEntities[EntityIndex].IsValid())
+				{
+					UDMXEntity* DraggedEntity = DraggedEntities[EntityIndex].Get();
+					DMXLibrary->SetEntityIndex(DraggedEntity, InsertBeforeIndex);
+				}
+			}
+
+			TSharedPtr<FDMXEntityTreeCategoryNode> TargetCategory = StaticCastSharedPtr<FDMXEntityTreeCategoryNode>(EntityNode->GetParent().Pin());
+			const uint32 UniverseID = TargetCategory->GetIntValue();
+
+			for (TWeakObjectPtr<UDMXEntity> Entity : DraggedEntities)
+			{
+				if (UDMXEntityFixturePatch* DraggedFixturePatch = Cast<UDMXEntityFixturePatch>(Entity))
+				{
+					DraggedFixturePatch->Modify();
+					DraggedFixturePatch->PreEditChange(nullptr);
+
+					FDMXEditorUtils::AutoAssignedAddresses(TArray<UDMXEntityFixturePatch*>({ DraggedFixturePatch }));
+					DraggedFixturePatch->SetUniverseID(UniverseID);
+
+					DraggedFixturePatch->PostEditChange();
+				}
+			}
+
+			// Display the changes in the Entity Tree
+			FixturePatchTree->UpdateTree();
+
+			OnFixturePatchOrderChanged.ExecuteIfBound();
+
+			return FReply::Handled().EndDragDrop();
+		}
+	}
+
+	return FReply::Unhandled();
+}
+
+bool SDMXFixturePatchTreeFixturePatchRow::TestCanDropWithFeedback(const TSharedRef<FDMXEntityDragDropOperation>& EntityDragDropOp) const
+{
+	const TSharedPtr<FDMXEntityTreeEntityNode> EntityNode = WeakEntityNode.Pin();
+	const UDMXEntityFixturePatch* FixturePatch = EntityNode.IsValid() ? Cast<UDMXEntityFixturePatch>(EntityNode->GetEntity()) : nullptr;
+
+	if (EntityNode.IsValid() && FixturePatch)
+	{
+		const TArray<TWeakObjectPtr<UDMXEntity>> DraggedEntities = EntityDragDropOp->GetDraggedEntities();
+
+		const bool bOnlyFixturePatchesAreDragged = [DraggedEntities]()
+		{
+			for (TWeakObjectPtr<UDMXEntity> Entity : DraggedEntities)
+			{
+				if (Entity.IsValid() && Entity->GetClass() != UDMXEntityFixturePatch::StaticClass())
+				{
+					return false;
+				}
+			}
+			return true;
+		}();
+
+		const bool bRowAndEntitiesAreOfSameDMXLibrary = [DraggedEntities, this]()
+		{
+			UDMXLibrary* DMXLibrary = GetDMXLibrary();
+			for (TWeakObjectPtr<UDMXEntity> Entity : DraggedEntities)
+			{
+				if (Entity.IsValid() && Entity->GetParentLibrary() != DMXLibrary)
+				{
+					return false;
+				}
+			}
+			return true;
+		}();
+
+		if (bOnlyFixturePatchesAreDragged && bRowAndEntitiesAreOfSameDMXLibrary)
+		{
+			const bool bDragBetweenUniverses = [DraggedEntities, FixturePatch]()
+			{
+				for (TWeakObjectPtr<UDMXEntity> DraggedEntity : DraggedEntities)
+				{
+					if (UDMXEntityFixturePatch* DraggedFixturePatch = Cast<UDMXEntityFixturePatch>(DraggedEntity))
+					{
+						if (DraggedFixturePatch->GetUniverseID() != FixturePatch->GetUniverseID())
+						{
+							return true;
+						}
+					}
+				}
+				return false;
+			}();
+
+			if (DraggedEntities.Contains(FixturePatch))
+			{
+				// The fixture patch of this row is a dragged fixture patch
+				EntityDragDropOp->SetFeedbackMessageError(FText::Format(
+					LOCTEXT("ReorderBeforeItself", "Drop {0} on itself"),
+					EntityDragDropOp->GetDraggedEntitiesName()
+				));
+
+				return false;
+			}
+			else if (!bDragBetweenUniverses)
+			{
+				EntityDragDropOp->SetFeedbackMessageError(FText::Format(
+					LOCTEXT("DragFixturePatchToItsOwnUniverse", "{0} already resides in this universe."),
+					EntityDragDropOp->GetDraggedEntitiesName()
+				));
+
+				return false;
+			}
+			else
+			{
+				TSharedPtr<FDMXEntityTreeCategoryNode> TargetCategory = StaticCastSharedPtr<FDMXEntityTreeCategoryNode>(EntityNode->GetParent().Pin());
+
+				const FText	ListTypeText = LOCTEXT("Property_Universe", "Universe");
+				const int32 UniverseID = TargetCategory->GetIntValue();
+				const FText	CategoryText = UniverseID == -1
+					? LOCTEXT("UnassignedUniverseIDValue", "Unassigned")
+					: FText::AsNumber(UniverseID);
+
+				EntityDragDropOp->SetFeedbackMessageOK(FText::Format(
+					LOCTEXT("ReorderBeforeAndSetProperty", "Reorder {0} before '{1}'\nSet {2} = '{3}'"),
+					EntityDragDropOp->GetDraggedEntitiesName(),
+					FText::FromString(FixturePatch->GetDisplayName()),
+					ListTypeText,
+					CategoryText
+				));
+
+				return true;
+			}
+		}
+	}
+
+	return false;
+}
+
+FText SDMXFixturePatchTreeFixturePatchRow::GetFixturePatchDisplayText() const
+{
+	if (const TSharedPtr<FDMXEntityTreeEntityNode> EntityNode = WeakEntityNode.Pin())
+	{
+		return EntityNode->GetDisplayNameText();
+	}
+
+	return FText::GetEmpty();
+}
+
+FText SDMXFixturePatchTreeFixturePatchRow::GetStartingChannelText() const
+{
+	if (TSharedPtr<FDMXEntityTreeEntityNode> EntityNode = WeakEntityNode.Pin())
+	{
+		if (UDMXEntityFixturePatch* Patch = Cast<UDMXEntityFixturePatch>(EntityNode->GetEntity()))
+		{
+			return FText::AsNumber(Patch->GetStartingChannel());
+		}
+	}
+
+	return FText::GetEmpty();
+}
+
+FText SDMXFixturePatchTreeFixturePatchRow::GetEndingChannelText() const
+{
+	if (TSharedPtr<FDMXEntityTreeEntityNode> EntityNode = WeakEntityNode.Pin())
+	{
+		if (UDMXEntityFixturePatch* Patch = Cast<UDMXEntityFixturePatch>(EntityNode->GetEntity()))
+		{
+			return FText::AsNumber(Patch->GetStartingChannel() + Patch->GetChannelSpan() - 1);
+		}
+	}
+
+	return FText::GetEmpty();
+}
+
+void SDMXFixturePatchTreeFixturePatchRow::OnAutoAssignChannelBoxStateChanged(ECheckBoxState NewState)
+{
+	if (OnAutoAssignChannelStateChanged.IsBound())
+	{
+		switch (NewState)
+		{
+		case ECheckBoxState::Unchecked:
+			OnAutoAssignChannelStateChanged.Execute(false);
+			break;
+		case ECheckBoxState::Checked:
+			OnAutoAssignChannelStateChanged.Execute(true);
+			break;
+		case ECheckBoxState::Undetermined:
+		default:
+			break;
+		}
+	}
+}
+
+FText SDMXFixturePatchTreeFixturePatchRow::GetFilterText() const
+{
+	if (OnGetFilterText.IsBound())
+	{
+		return OnGetFilterText.Execute();
+	}
+	return FText();
+}
+
+void SDMXFixturePatchTreeFixturePatchRow::OnNameTextCommit(const FText& InNewName, ETextCommit::Type InTextCommit)
+{
+	const TSharedPtr<SDMXFixturePatchTree> FixturePatchTree = WeakFixturePatchTree.Pin();
+	const TSharedPtr<FDMXEntityTreeEntityNode> EntityNode = WeakEntityNode.Pin();
+
+	if (FixturePatchTree.IsValid() && EntityNode.IsValid())
+	{
+		const FString NewNameString = InNewName.ToString();
+
+		// Check if the name is unchanged
+		if (NewNameString.Equals(EntityNode->GetDisplayNameString()))
+		{
+			return;
+		}
+
+		const FScopedTransaction Transaction(LOCTEXT("RenameEntity", "Rename Entity"));
+		FixturePatchTree->GetDMXLibrary()->Modify();
+
+		FDMXEditorUtils::RenameEntity(FixturePatchTree->GetDMXLibrary(), EntityNode->GetEntity(), NewNameString);
+
+		FixturePatchTree->SelectItemByName(NewNameString, ESelectInfo::OnMouseClick);
+	}
+}
+
+bool SDMXFixturePatchTreeFixturePatchRow::OnNameTextVerifyChanged(const FText& InNewText, FText& OutErrorMessage)
+{
+	TSharedPtr<SDMXFixturePatchTree> FixturePatchTree = WeakFixturePatchTree.Pin();
+	const TSharedPtr<FDMXEntityTreeEntityNode> EntityNode = WeakEntityNode.Pin();
+
+	if (FixturePatchTree.IsValid() && EntityNode.IsValid())
+	{
+		FString TextAsString = InNewText.ToString();
+		if (TextAsString.Equals(WeakEntityNode.Pin()->GetDisplayNameString()))
+		{
+			return true;
+		}
+
+		return FDMXEditorUtils::ValidateEntityName(
+			TextAsString,
+			FixturePatchTree->GetDMXLibrary(),
+			WeakEntityNode.Pin()->GetEntity()->GetClass(),
+			OutErrorMessage
+		);
+	}
+
+	return false;
+}
+
+FReply SDMXFixturePatchTreeFixturePatchRow::HandleOnDragDetected(const FGeometry& MyGeometry, const FPointerEvent& MouseEvent)
+{
+	if (MouseEvent.IsMouseButtonDown(EKeys::LeftMouseButton) && OnEntityDragged.IsBound())
+	{
+		if (const TSharedPtr<FDMXEntityTreeEntityNode> EntityNode = WeakEntityNode.Pin())
+		{
+			return OnEntityDragged.Execute(EntityNode, MouseEvent);
+		}
+	}
+
+	return FReply::Unhandled();
+}
+
+const FSlateBrush* SDMXFixturePatchTreeFixturePatchRow::GetStatusIcon() const
+{
+	static FSlateNoResource EmptyBrush = FSlateNoResource();
+
+	if (const TSharedPtr<FDMXEntityTreeEntityNode> EntityNode = WeakEntityNode.Pin())
+	{
+		if (!EntityNode->GetErrorStatus().IsEmpty())
+		{
+			return FEditorStyle::GetBrush("Icons.Error");
+		}
+
+		if (!EntityNode->GetWarningStatus().IsEmpty())
+		{
+			return FEditorStyle::GetBrush("Icons.Warning");
+		}
+	}
+
+	return &EmptyBrush;
+}
+
+FText SDMXFixturePatchTreeFixturePatchRow::GetStatusToolTip() const
+{
+	if (const TSharedPtr<FDMXEntityTreeEntityNode> EntityNode = WeakEntityNode.Pin())
+	{
+		const FText& ErrorStatus = EntityNode->GetErrorStatus();
+		if (!ErrorStatus.IsEmpty())
+		{
+			return ErrorStatus;
+		}
+
+		const FText& WarningStatus = EntityNode->GetWarningStatus();
+		if (!WarningStatus.IsEmpty())
+		{
+			return WarningStatus;
+		}
+	}
+
+	return FText::GetEmpty();
+}
+
+ECheckBoxState SDMXFixturePatchTreeFixturePatchRow::IsAutoAssignChannelEnabled() const
+{
+	if (const TSharedPtr<FDMXEntityTreeEntityNode> EntityNode = WeakEntityNode.Pin())
+	{
+		if (UDMXEntityFixturePatch* FixturePatch = Cast<UDMXEntityFixturePatch>(EntityNode->GetEntity()))
+		{
+			return FixturePatch->IsAutoAssignAddress() ? ECheckBoxState::Checked : ECheckBoxState::Unchecked;
+		}
+	}
+
+	return ECheckBoxState::Undetermined;
+}
+
+UDMXLibrary* SDMXFixturePatchTreeFixturePatchRow::GetDMXLibrary() const
+{
+	if (const TSharedPtr<SDMXFixturePatchTree>& FixturePatchTree = WeakFixturePatchTree.Pin())
+	{
+		return FixturePatchTree->GetDMXLibrary();
+	}
+
+	return nullptr;
+}
+
+#undef LOCTEXT_NAMESPACE
diff --git a/Engine/Plugins/VirtualProduction/DMX/DMXEngine/Source/DMXEditor/Private/Widgets/FixturePatch/SDMXFixturePatchTreeFixturePatchRow.h b/Engine/Plugins/VirtualProduction/DMX/DMXEngine/Source/DMXEditor/Private/Widgets/FixturePatch/SDMXFixturePatchTreeFixturePatchRow.h
new file mode 100644
--- /dev/null
+++ b/Engine/Plugins/VirtualProduction/DMX/DMXEngine/Source/DMXEditor/Private/Widgets/FixturePatch/SDMXFixturePatchTreeFixturePatchRow.h
@@ -0,0 +1,119 @@
+// Copyright Epic Games, Inc. All Rights Reserved.
+
+#pragma once
+
+#include "CoreMinimal.h"
+#include "Widgets/Views/SListView.h"
+
+#include "Widgets/Views/STableRow.h"
+
+class FDMXEntityDragDropOperation;
+class FDMXEntityTreeEntityNode;
+class FDMXEntityTreeNodeBase;
+class SDMXFixturePatchTree;
+class UDMXLibrary;
+
+class FDragDropEvent;
+struct FGeometry;
+class SInlineEditableTextBlock;
+
+
+/** A fixture patch row in a fixture patch tree */
+class SDMXFixturePatchTreeFixturePatchRow 
+	: public STableRow<TSharedPtr<FDMXEntityTreeEntityNode>>
+{
+	DECLARE_DELEGATE_RetVal_TwoParams(FReply, FOnEntityDragged, TSharedPtr<FDMXEntityTreeNodeBase>, const FPointerEvent&);
+	DECLARE_DELEGATE_RetVal(FText, FOnGetFilterText);
+	DECLARE_DELEGATE_OneParam(FOnAutoAssignChannelStateChanged, bool);
+
+public:
+	SLATE_BEGIN_ARGS(SDMXFixturePatchTreeFixturePatchRow)
+	{}
+		SLATE_EVENT(FOnEntityDragged, OnEntityDragged)
+	
+		SLATE_EVENT(FOnGetFilterText, OnGetFilterText)
+		
+		SLATE_EVENT(FOnAutoAssignChannelStateChanged, OnAutoAssignChannelStateChanged)
+		
+		SLATE_EVENT(FSimpleDelegate, OnFixturePatchOrderChanged)
+
+	SLATE_END_ARGS()
+
+	/** Constructs this widget */
+	void Construct(const FArguments& InArgs, TSharedPtr<FDMXEntityTreeEntityNode> InEntityNode, TSharedPtr<STableViewBase> InOwnerTableView, TWeakPtr<SDMXFixturePatchTree> InFixturePatchTree);
+	
+	/** Starts renaming the widget (enters editing mode on its Inline Editable Text Block) */
+	void EnterRenameMode();
+
+protected:
+	//~ SWidget interface begin
+	virtual FReply OnMouseButtonDoubleClick(const FGeometry& InMyGeometry, const FPointerEvent& InMouseEvent) override;
+	virtual void OnDragEnter(const FGeometry& MyGeometry, const FDragDropEvent& DragDropEvent) override;
+	virtual FReply OnDrop(const FGeometry& MyGeometry, const FDragDropEvent& DragDropEvent) override;
+	//~ SWidget interface end
+
+private:
+	/** Returns true if the drag drop op can be dropped on this row */
+	bool TestCanDropWithFeedback(const TSharedRef<FDMXEntityDragDropOperation>& EntityDragDropOp) const;
+
+	/** Returns the name of the patch as text */
+	FText GetFixturePatchDisplayText() const;
+
+	/** Returns the text of the starting channel */
+	FText GetStartingChannelText() const;
+
+	/** Returns the text of the ending channel */
+	FText GetEndingChannelText() const;
+
+	/**  Called when the auto-assign address check-box state is changed */
+	void OnAutoAssignChannelBoxStateChanged(ECheckBoxState NewState);
+
+	/** Get the current filter text from the search box */
+	FText GetFilterText() const;
+
+	/** Verifies the name of the component when changing it */
+	bool OnNameTextVerifyChanged(const FText& InNewText, FText& OutErrorMessage);
+
+	/** Commits the new name of the component */
+	void OnNameTextCommit(const FText& InNewName, ETextCommit::Type InTextCommit);
+
+	/** Called when the row is being dragged */
+	FReply HandleOnDragDetected(const FGeometry& MyGeometry, const FPointerEvent& MouseEvent);
+
+	/** Get the icon for the Entity usability status. If it's all good, it's an emtpy image. */
+	const FSlateBrush* GetStatusIcon() const;
+
+	/** Get the tool tip text for the status icon */
+	FText GetStatusToolTip() const;
+
+	/** For Fixture Patches, returns wether Auto Assign Channel is enabled */
+	ECheckBoxState IsAutoAssignChannelEnabled() const;
+
+	/** Returns the DMXLibrary of this row */
+	UDMXLibrary* GetDMXLibrary() const;
+
+	/** The node this row displays */
+	TWeakPtr<FDMXEntityTreeEntityNode> WeakEntityNode;
+
+	/** The fixture patch tree widget that owns this row */
+	TWeakPtr<SDMXFixturePatchTree> WeakFixturePatchTree;
+
+	/** The tooltip of the status icon */
+	FText StatusIconToolTip;
+
+	/** Called when the entity list changed auto assign channel state */
+	FOnAutoAssignChannelStateChanged OnAutoAssignChannelStateChanged;
+
+	/** Delegate broadcast when an entity was dragged */
+	FOnEntityDragged OnEntityDragged;
+
+	/** Delegate broadcast when the filter text is requested */
+	FOnGetFilterText OnGetFilterText;
+
+	/** Called when the entity list changed order of the library's entity array */
+	FSimpleDelegate OnFixturePatchOrderChanged;
+
+	/** The editable text block to display and change the name of the patch */
+	TSharedPtr<SInlineEditableTextBlock> InlineEditableFixturePatchNameWidget;
+};
+
diff --git a/Engine/Plugins/VirtualProduction/DMX/DMXEngine/Source/DMXEditor/Private/Widgets/FixturePatch/SDMXFixturePatchTreeUniverseRow.cpp b/Engine/Plugins/VirtualProduction/DMX/DMXEngine/Source/DMXEditor/Private/Widgets/FixturePatch/SDMXFixturePatchTreeUniverseRow.cpp
new file mode 100644
--- /dev/null
+++ b/Engine/Plugins/VirtualProduction/DMX/DMXEngine/Source/DMXEditor/Private/Widgets/FixturePatch/SDMXFixturePatchTreeUniverseRow.cpp
@@ -0,0 +1,280 @@
+// Copyright Epic Games, Inc. All Rights Reserved.
+
+#include "SDMXFixturePatchTreeUniverseRow.h"
+
+#include "DMXEditorUtils.h"
+#include "SDMXFixturePatchTree.h"
+#include "DragDrop/DMXEntityDragDropOp.h"
+#include "Library/DMXEntityFixturePatch.h"
+#include "Library/DMXLibrary.h"
+
+#include "EditorStyleSet.h"
+#include "ScopedTransaction.h"
+#include "SlateOptMacros.h"
+
+
+#define LOCTEXT_NAMESPACE "SDMXFixturePatchTreeUniverseRow"
+
+BEGIN_SLATE_FUNCTION_BUILD_OPTIMIZATION
+void SDMXFixturePatchTreeUniverseRow::Construct(const FArguments& InArgs, const TSharedRef<STableViewBase>& InOwnerTableView, TSharedPtr<FDMXEntityTreeCategoryNode> InCategoryNode, bool bInIsRootCategory, TWeakPtr<SDMXFixturePatchTree> InFixturePatchTree)
+{
+	check(InCategoryNode.IsValid());
+
+	WeakFixturePatchTree = InFixturePatchTree;
+	WeakCategoryNode = InCategoryNode;
+
+	OnFixturePatchOrderChanged = InArgs._OnFixturePatchOrderChanged;
+
+	// background color tint
+	const FLinearColor BackgroundTint(0.6f, 0.6f, 0.6f, bInIsRootCategory ? 1.0f : 0.3f);
+
+	// rebuilds the whole table row from scratch
+	ChildSlot
+	.Padding(0.0f, 2.0f, 0.0f, 0.0f)
+	[
+		SAssignNew(ContentBorder, SBorder)
+		.BorderImage(this, &SDMXFixturePatchTreeUniverseRow::GetBackgroundImageBrush)
+		.Padding(FMargin(0.0f, 3.0f))
+		.BorderBackgroundColor(BackgroundTint)
+		.ToolTipText(WeakCategoryNode.Pin()->GetToolTip())
+		[
+			SNew(SHorizontalBox)
+
+			+ SHorizontalBox::Slot()
+			.VAlign(VAlign_Center)
+			.Padding(2.0f, 2.0f, 2.0f, 2.0f)
+			.AutoWidth()
+			[
+				SNew(SExpanderArrow, STableRow<TSharedPtr<FDMXEntityTreeCategoryNode>>::SharedThis(this))
+			]
+
+			+ SHorizontalBox::Slot()
+			.VAlign(VAlign_Center)
+			[
+				InArgs._Content.Widget
+			]
+		]
+	];
+
+	STableRow<TSharedPtr<FDMXEntityTreeCategoryNode>>::ConstructInternal(
+		STableRow<TSharedPtr<FDMXEntityTreeCategoryNode>>::FArguments()
+			.Style(FEditorStyle::Get(), "DetailsView.TreeView.TableRow")
+			.ShowSelection(false),
+		InOwnerTableView
+	);
+}
+END_SLATE_FUNCTION_BUILD_OPTIMIZATION
+
+void SDMXFixturePatchTreeUniverseRow::SetContent(TSharedRef< SWidget > InContent)
+{
+	ContentBorder->SetContent(InContent);
+}
+
+void SDMXFixturePatchTreeUniverseRow::OnDragEnter(const FGeometry& MyGeometry, const FDragDropEvent& DragDropEvent)
+{
+	if (TSharedPtr<FDMXEntityDragDropOperation> EntityDragDropOp = DragDropEvent.GetOperationAs<FDMXEntityDragDropOperation>())
+	{
+		TestCanDropWithFeedback(EntityDragDropOp.ToSharedRef());
+	}
+}
+
+FReply SDMXFixturePatchTreeUniverseRow::OnDrop(const FGeometry& MyGeometry, const FDragDropEvent& DragDropEvent)
+{
+	TSharedPtr<FDMXEntityDragDropOperation> EntityDragDropOp = DragDropEvent.GetOperationAs<FDMXEntityDragDropOperation>();
+	TSharedPtr<SDMXFixturePatchTree> FixturePatchTree = WeakFixturePatchTree.Pin();
+	const TSharedPtr<FDMXEntityTreeCategoryNode>& MyPinnedCategoryNode = WeakCategoryNode.Pin();
+	UDMXLibrary* DMXLibrary = GetDMXLibrary();
+
+	if (EntityDragDropOp.IsValid() && FixturePatchTree.IsValid() && MyPinnedCategoryNode.IsValid() && DMXLibrary)
+	{
+		if (TestCanDropWithFeedback(EntityDragDropOp.ToSharedRef()))
+		{
+			const TArray<TWeakObjectPtr<UDMXEntity>> DraggedEntities = EntityDragDropOp->GetDraggedEntities();
+
+			DMXLibrary->Modify();
+
+			// Register transaction and current DMX library state for Undo
+			const FScopedTransaction ChangeCategoryTransaction = FScopedTransaction(
+				FText::Format(LOCTEXT("DropFixturePatchOnUniverse", "Set Fixture {0}|plural(one=Patch, other=Patches) Universe"), DraggedEntities.Num())
+			);
+
+			if (MyPinnedCategoryNode->GetChildren().Num() > 0)
+			{
+				if (TSharedPtr<FDMXEntityTreeEntityNode> EntityNode = StaticCastSharedPtr<FDMXEntityTreeEntityNode>(MyPinnedCategoryNode->GetChildren().Last(0)))
+				{
+					// Index after last entity in hovered category
+					UDMXEntity* LastEntityInCategory = EntityNode->GetEntity();
+
+					const int32 LastEntityIndex = DMXLibrary->FindEntityIndex(LastEntityInCategory);
+					check(LastEntityIndex != INDEX_NONE);
+
+					// Move dragged entities after the last ones in the category
+					// Reverse for to keep dragged entities order.
+					int32 DesiredIndex = LastEntityIndex + 1;
+					for (int32 EntityIndex = DraggedEntities.Num() - 1; EntityIndex > -1; --EntityIndex)
+					{
+						if (UDMXEntity* Entity = DraggedEntities[EntityIndex].Get())
+						{
+							DMXLibrary->SetEntityIndex(Entity, DesiredIndex);
+						}
+					}
+				}
+			}
+
+			const int32 UniverseID = MyPinnedCategoryNode->GetIntValue();
+			for (TWeakObjectPtr<UDMXEntity> Entity : DraggedEntities)
+			{
+				if (UDMXEntityFixturePatch* FixturePatch = Cast<UDMXEntityFixturePatch>(Entity.Get()))
+				{
+					FixturePatch->Modify();
+					FixturePatch->PreEditChange(nullptr);
+
+					if (FixturePatch->IsAutoAssignAddress())
+					{
+						FDMXEditorUtils::TryAutoAssignToUniverses(FixturePatch, TSet<int32>({ UniverseID }));
+					}
+					FixturePatch->SetUniverseID(UniverseID);
+
+					FixturePatch->PostEditChange();
+				}
+			}
+
+			FixturePatchTree->UpdateTree();
+
+			OnFixturePatchOrderChanged.ExecuteIfBound();
+
+			return FReply::Handled().EndDragDrop();
+		}
+	}
+
+	return FReply::Unhandled();
+}
+
+bool SDMXFixturePatchTreeUniverseRow::TestCanDropWithFeedback(const TSharedRef<FDMXEntityDragDropOperation>& EntityDragDropOp) const
+{
+	if (const TSharedPtr<FDMXEntityTreeCategoryNode>& MyPinnedCategoryNode = WeakCategoryNode.Pin())
+	{
+		// Only handle entity drag drop ops that contain fixture patches only 
+		const bool bOnlyFixturePatchesAreDragged = [EntityDragDropOp]()
+		{
+			for (UClass* Class : EntityDragDropOp->GetDraggedEntityTypes())
+			{
+				if (Class != UDMXEntityFixturePatch::StaticClass())
+				{
+					return false;
+				}
+			}
+			return true;
+		}();
+
+		const bool bRowAndEntitiesAreOfSameLibrary = [EntityDragDropOp, this]()
+		{
+			UDMXLibrary* DMXLibrary = GetDMXLibrary();
+			for (TWeakObjectPtr<UDMXEntity> Entity : EntityDragDropOp->GetDraggedEntities())
+			{
+				if (Entity.IsValid() && Entity->GetParentLibrary() != DMXLibrary)
+				{
+					return false;
+				}
+			}
+			return true;
+		}();
+
+		if (bOnlyFixturePatchesAreDragged && bRowAndEntitiesAreOfSameLibrary)
+		{
+			if (MyPinnedCategoryNode->CanDropOntoCategory())
+			{
+				const int32 NumDraggedEntities = EntityDragDropOp->GetDraggedEntities().Num();
+				const TArray<TSharedRef<FDMXEntityTreeCategoryNode>> DraggedFromCategories = GetCategoryNodesFromDragDropOp(EntityDragDropOp);
+
+				if (DraggedFromCategories.Num() == 1 && DraggedFromCategories[0] == MyPinnedCategoryNode)
+				{
+					// There wouldn't be any change by dragging the items into their own category.
+					EntityDragDropOp->SetFeedbackMessageError(FText::Format(
+						LOCTEXT("DragIntoSelfCategory", "The selected {0} {1}|plural(one=is, other=are) already in this category"),
+						FDMXEditorUtils::GetEntityTypeNameText(UDMXEntityFixturePatch::StaticClass(), NumDraggedEntities > 1),
+						NumDraggedEntities
+					));
+
+					return false;
+				}
+				else
+				{
+					const uint32 MyUniverseID = MyPinnedCategoryNode->GetIntValue();
+
+					EntityDragDropOp->SetFeedbackMessageOK(FText::Format(
+						LOCTEXT("ReorderBeforeAndSetProperty", "Reorder {0} before '{1}'\nSet {2} = '{3}'"),
+						EntityDragDropOp->GetDraggedEntitiesName(),
+						MyPinnedCategoryNode->GetDisplayNameText(),
+						LOCTEXT("UniverseCategoryName", "Universe"),
+						FText::AsNumber(MyUniverseID)
+					));
+
+					return true;
+				}
+			}
+			else
+			{
+				EntityDragDropOp->SetFeedbackMessageError(FText::Format(
+					LOCTEXT("DragIntoNoValueCategory", "Cannot drop {0} onto this category"),
+					EntityDragDropOp->GetDraggedEntitiesName()
+				));
+
+				return false;
+			}
+		}
+	}
+
+	return false;
+}
+
+TArray<TSharedRef<FDMXEntityTreeCategoryNode>> SDMXFixturePatchTreeUniverseRow::GetCategoryNodesFromDragDropOp(const TSharedRef<FDMXEntityDragDropOperation>& DragDropOp) const
+{
+	TArray<TSharedRef<FDMXEntityTreeCategoryNode>> CategoryNodes;
+	if (const TSharedPtr<SDMXFixturePatchTree>& FixturePatchTree = WeakFixturePatchTree.Pin())
+	{
+		TArray<TWeakObjectPtr<UDMXEntity>> DraggedEntities = DragDropOp->GetDraggedEntities();
+
+		for (TWeakObjectPtr<UDMXEntity> Entity : DraggedEntities)
+		{
+			if (UDMXEntityFixturePatch* FixturePatch = Cast<UDMXEntityFixturePatch>(Entity.Get()))
+			{
+				TSharedPtr<FDMXEntityTreeCategoryNode> CategoryNode = FixturePatchTree->FindCategoryNodeOfEntity(FixturePatch);
+				if (CategoryNode.IsValid())
+				{
+					CategoryNodes.Add(CategoryNode.ToSharedRef());
+				}
+			}
+		}
+	}
+
+	return CategoryNodes;
+}
+
+const FSlateBrush* SDMXFixturePatchTreeUniverseRow::GetBackgroundImageBrush() const
+{
+	if (IsHovered())
+	{
+		return IsItemExpanded()
+			? FEditorStyle::GetBrush("DetailsView.CategoryTop_Hovered")
+			: FEditorStyle::GetBrush("DetailsView.CollapsedCategory_Hovered");
+	}
+	else
+	{
+		return IsItemExpanded()
+			? FEditorStyle::GetBrush("DetailsView.CategoryTop")
+			: FEditorStyle::GetBrush("DetailsView.CollapsedCategory");
+	}
+}
+
+UDMXLibrary* SDMXFixturePatchTreeUniverseRow::GetDMXLibrary() const
+{
+	if (const TSharedPtr<SDMXFixturePatchTree>& FixturePatchTree = WeakFixturePatchTree.Pin())
+	{
+		return FixturePatchTree->GetDMXLibrary();
+	}
+
+	return nullptr;
+}
+
+#undef LOCTEXT_NAMESPACE
diff --git a/Engine/Plugins/VirtualProduction/DMX/DMXEngine/Source/DMXEditor/Private/Widgets/FixturePatch/SDMXFixturePatchTreeUniverseRow.h b/Engine/Plugins/VirtualProduction/DMX/DMXEngine/Source/DMXEditor/Private/Widgets/FixturePatch/SDMXFixturePatchTreeUniverseRow.h
new file mode 100644
--- /dev/null
+++ b/Engine/Plugins/VirtualProduction/DMX/DMXEngine/Source/DMXEditor/Private/Widgets/FixturePatch/SDMXFixturePatchTreeUniverseRow.h
@@ -0,0 +1,67 @@
+// Copyright Epic Games, Inc. All Rights Reserved.
+
+#pragma once
+
+#include "CoreMinimal.h"
+#include "Widgets/Views/SListView.h"
+
+#include "Widgets/Views/STableRow.h"
+
+class FDMXEntityDragDropOperation;
+class FDMXEntityTreeCategoryNode;
+class SDMXFixturePatchTree;
+class UDMXLibrary;
+
+class STableViewBase;
+
+
+/** A row displaying a universe in a Fixture Patch Tree */
+class SDMXFixturePatchTreeUniverseRow
+	: public STableRow<TSharedPtr<FDMXEntityTreeCategoryNode>>
+{
+public:
+	SLATE_BEGIN_ARGS(SDMXFixturePatchTreeUniverseRow)
+	{}
+		SLATE_DEFAULT_SLOT(typename SDMXFixturePatchTreeUniverseRow::FArguments, Content)
+
+		SLATE_EVENT(FSimpleDelegate, OnFixturePatchOrderChanged)
+
+	SLATE_END_ARGS()
+
+	/** Constructs this widget */
+	void Construct(const FArguments& InArgs, const TSharedRef<STableViewBase>& InOwnerTableView, TSharedPtr<FDMXEntityTreeCategoryNode> InCategoryNode, bool bInIsRootCategory, TWeakPtr<SDMXFixturePatchTree> InEditorList);
+
+	/** Sets an arbitrary widget as the content of the row */
+	virtual void SetContent(TSharedRef<SWidget> InContent) override;
+
+protected:
+	//~ SWidget interface begin
+	virtual void OnDragEnter(const FGeometry& MyGeometry, const FDragDropEvent& DragDropEvent) override;
+	virtual FReply OnDrop(const FGeometry& MyGeometry, const FDragDropEvent& DragDropEvent) override;
+	//~ SWidget interface end
+
+private:
+	/** Returns true if the drag drop op can be dropped on this row */
+	bool TestCanDropWithFeedback(const TSharedRef<FDMXEntityDragDropOperation>& EntityDragDropOp) const;
+
+	/** Returns the categories nodes the entities were dragged out from */
+	TArray<TSharedRef<FDMXEntityTreeCategoryNode>> GetCategoryNodesFromDragDropOp(const TSharedRef<FDMXEntityDragDropOperation>& DragDropOp) const;
+
+	/** Returns the background image brush */
+	const FSlateBrush* GetBackgroundImageBrush() const;
+
+	/** Returns the DMXLibrary of this row */
+	UDMXLibrary* GetDMXLibrary() const;
+
+	/** Pointer to the node this row displays */
+	TWeakPtr<FDMXEntityTreeCategoryNode> WeakCategoryNode;
+
+	/** The owning Fixture Patch List */
+	TWeakPtr<SDMXFixturePatchTree> WeakFixturePatchTree;
+
+	/** Border that holds the content */
+	TSharedPtr<SBorder> ContentBorder;
+
+	/** Called when the entity list changed order of the library's entity array */
+	FSimpleDelegate OnFixturePatchOrderChanged;
+};
diff --git a/Engine/Plugins/VirtualProduction/DMX/DMXEngine/Source/DMXEditor/Private/Widgets/FixtureType/SDMXFixtureTypeEditor.cpp b/Engine/Plugins/VirtualProduction/DMX/DMXEngine/Source/DMXEditor/Private/Widgets/FixtureType/SDMXFixtureTypeEditor.cpp
--- a/Engine/Plugins/VirtualProduction/DMX/DMXEngine/Source/DMXEditor/Private/Widgets/FixtureType/SDMXFixtureTypeEditor.cpp
+++ b/Engine/Plugins/VirtualProduction/DMX/DMXEngine/Source/DMXEditor/Private/Widgets/FixtureType/SDMXFixtureTypeEditor.cpp
@@ -3,13 +3,14 @@
 #include "Widgets/FixtureType/SDMXFixtureTypeEditor.h"
 
 #include "DMXEditor.h"
+#include "DMXEditorUtils.h"
 #include "DMXFixtureTypeSharedData.h"
+#include "SDMXFixtureTypeTree.h"
 #include "Library/DMXEntityFixtureType.h"
 #include "Library/DMXImport.h"
 #include "Library/DMXLibrary.h"
 #include "Library/DMXEntityFixtureType.h"
 #include "Widgets/SDMXEntityInspector.h"
-#include "Widgets/SDMXEntityList.h"
 
 #include "Widgets/Layout/SSplitter.h"
 
@@ -66,9 +67,9 @@ void SDMXFixtureTypeEditor::Construct(const FArguments& InArgs)
 			// Top Row
 			+ SSplitter::Slot()
 			[
-				SAssignNew(EntityList, SDMXEntityList, UDMXEntityFixtureType::StaticClass())
+				SAssignNew(FixtureTypeTree, SDMXFixtureTypeTree, UDMXEntityFixtureType::StaticClass())
 				.DMXEditor(InArgs._DMXEditor)
-				.OnSelectionUpdated(this, &SDMXFixtureTypeEditor::OnSelectionUpdated)
+				.OnSelectionChanged(this, &SDMXFixtureTypeEditor::OnSelectionChanged)
 			]
 
 			// Bottom Row
@@ -122,35 +123,34 @@ void SDMXFixtureTypeEditor::Construct(const FArguments& InArgs)
 
 void SDMXFixtureTypeEditor::RequestRenameOnNewEntity(const UDMXEntity* InEntity, ESelectInfo::Type SelectionType)
 {
-	check(EntityList.IsValid());
+	check(FixtureTypeTree.IsValid());
 
-	EntityList->UpdateTree();
-	EntityList->SelectItemByEntity(InEntity, SelectionType);
-	EntityList->OnRenameNode();
+	FixtureTypeTree->UpdateTree();
+	FixtureTypeTree->SelectItemByEntity(InEntity, SelectionType);
 }
 
 void SDMXFixtureTypeEditor::SelectEntity(UDMXEntity* InEntity, ESelectInfo::Type InSelectionType /*= ESelectInfo::Type::Direct*/)
 {
-	check(EntityList.IsValid());
+	check(FixtureTypeTree.IsValid());
 
-	EntityList->SelectItemByEntity(InEntity, InSelectionType);
+	FixtureTypeTree->SelectItemByEntity(InEntity, InSelectionType);
 }
 
 void SDMXFixtureTypeEditor::SelectEntities(const TArray<UDMXEntity*>& InEntities, ESelectInfo::Type InSelectionType /*= ESelectInfo::Type::Direct*/)
 {
-	check(EntityList.IsValid());
+	check(FixtureTypeTree.IsValid());
 
-	EntityList->SelectItemsByEntity(InEntities, InSelectionType);
+	FixtureTypeTree->SelectItemsByEntities(InEntities, InSelectionType);
 }
 
 TArray<UDMXEntity*> SDMXFixtureTypeEditor::GetSelectedEntities() const
 {
-	check(EntityList.IsValid());
+	check(FixtureTypeTree.IsValid());
 
-	return EntityList->GetSelectedEntities();
+	return FixtureTypeTree->GetSelectedEntities();
 }
 
-void SDMXFixtureTypeEditor::OnSelectionUpdated(TArray<UDMXEntity*> InSelectedEntities)
+void SDMXFixtureTypeEditor::OnSelectionChanged(const TArray<UDMXEntity*>& InSelectedEntities)
 {
 	if (TSharedPtr<FDMXEditor> DMXEditorPtr = DMXEditor.Pin())
 	{
@@ -191,7 +191,7 @@ void SDMXFixtureTypeEditor::OnFinishedChangingProperties(const FPropertyChangedE
 		{
 			// When the user adds a Mode or Function, their names can't be empty
 
-			const TArray<UDMXEntity*>&& SelectedEntities = EntityList->GetSelectedEntities();
+			const TArray<UDMXEntity*>&& SelectedEntities = FixtureTypeTree->GetSelectedEntities();
 
 			for (UDMXEntity* Entity : SelectedEntities)
 			{
@@ -206,7 +206,7 @@ void SDMXFixtureTypeEditor::OnFinishedChangingProperties(const FPropertyChangedE
 	{
 		if (PropertyChangedEvent.GetPropertyName() == GET_MEMBER_NAME_CHECKED(UDMXEntityFixtureType, DMXImport))
 		{
-			const TArray<UDMXEntity*>&& SelectedEntities = EntityList->GetSelectedEntities();
+			const TArray<UDMXEntity*>&& SelectedEntities = FixtureTypeTree->GetSelectedEntities();
 
 			for (UDMXEntity* Entity : SelectedEntities)
 			{
@@ -222,8 +222,8 @@ void SDMXFixtureTypeEditor::OnFinishedChangingProperties(const FPropertyChangedE
 		}
 	}
 
-	if (EntityList.IsValid())
+	if (FixtureTypeTree.IsValid())
 	{
-		EntityList->UpdateTree();
+		FixtureTypeTree->UpdateTree();
 	}
 }
diff --git a/Engine/Plugins/VirtualProduction/DMX/DMXEngine/Source/DMXEditor/Private/Widgets/FixtureType/SDMXFixtureTypeEditor.h b/Engine/Plugins/VirtualProduction/DMX/DMXEngine/Source/DMXEditor/Private/Widgets/FixtureType/SDMXFixtureTypeEditor.h
--- a/Engine/Plugins/VirtualProduction/DMX/DMXEngine/Source/DMXEditor/Private/Widgets/FixtureType/SDMXFixtureTypeEditor.h
+++ b/Engine/Plugins/VirtualProduction/DMX/DMXEngine/Source/DMXEditor/Private/Widgets/FixtureType/SDMXFixtureTypeEditor.h
@@ -5,10 +5,11 @@
 #include "Widgets/SDMXEntityEditor.h"
 
 #include "CoreMinimal.h"
-#include "Widgets/SDMXEntityList.h"
 
 class FDMXEditor;
 class SDMXEntityInspector;
+class SDMXFixtureTypeTree;
+
 struct FPropertyChangedEvent;
 
 
@@ -38,7 +39,7 @@ public:
 
 protected:
 	/** Callback for when entities list selection changes */
-	virtual void OnSelectionUpdated(TArray<UDMXEntity*> InSelectedEntities);
+	virtual void OnSelectionChanged(const TArray<UDMXEntity*>& InSelectedEntities);
 
 	/** Callback for when some property has changed in the inspector */
 	virtual void OnFinishedChangingProperties(const FPropertyChangedEvent& PropertyChangedEvent);
@@ -47,8 +48,8 @@ protected:
 	/** Pointer back to the DMXEditor tool that owns us */
 	TWeakPtr<FDMXEditor> DMXEditor;
 
-	/** List of Fixture Types */
-	TSharedPtr<SDMXEntityList> EntityList;
+	/** Tree of Fixture Types */
+	TSharedPtr<SDMXFixtureTypeTree> FixtureTypeTree;
 
 	/** Inspector for the selected Fixture(s) Settings */
 	TSharedPtr<SDMXEntityInspector> FixtureSettingsInspector;
diff --git a/Engine/Plugins/VirtualProduction/DMX/DMXEngine/Source/DMXEditor/Private/Widgets/FixtureType/SDMXFixtureTypeTree.cpp b/Engine/Plugins/VirtualProduction/DMX/DMXEngine/Source/DMXEditor/Private/Widgets/FixtureType/SDMXFixtureTypeTree.cpp
new file mode 100644
--- /dev/null
+++ b/Engine/Plugins/VirtualProduction/DMX/DMXEngine/Source/DMXEditor/Private/Widgets/FixtureType/SDMXFixtureTypeTree.cpp
@@ -0,0 +1,371 @@
+// Copyright Epic Games, Inc. All Rights Reserved.
+
+#include "SDMXFixtureTypeTree.h"
+
+#include "DMXEditor.h"
+#include "DMXEditorUtils.h"
+#include "SDMXFixtureTypeTreeCategoryRow.h"
+#include "SDMXFixtureTypeTreeFixtureTypeRow.h"
+#include "Commands/DMXEditorCommands.h"
+#include "Library/DMXEntityFixtureType.h"
+#include "Library/DMXLibrary.h"
+#include "Widgets/DMXEntityTreeNode.h"
+
+#include "ScopedTransaction.h"
+#include "Framework/Application/SlateApplication.h"
+#include "Framework/Commands/GenericCommands.h"
+#include "Framework/MultiBox/MultiBoxBuilder.h"
+#include "Widgets/SBoxPanel.h"
+#include "Widgets/Input/SButton.h"
+#include "Widgets/Images/SImage.h"
+#include "Widgets/Text/STextBlock.h"
+#include "Widgets/Views/STreeView.h"
+
+
+#define LOCTEXT_NAMESPACE "SDMXFixtureTypeTree"
+
+void SDMXFixtureTypeTree::Construct(const FArguments& InArgs, const TSubclassOf<UDMXEntity> InListType)
+{
+	SDMXEntityTreeViewBase::FArguments BaseArguments =
+		SDMXEntityTreeViewBase::FArguments()
+		.DMXEditor(InArgs._DMXEditor)
+		.OnSelectionChanged(InArgs._OnSelectionChanged)
+		.OnEntitiesAdded(InArgs._OnEntitiesAdded)
+		.OnEntityOrderChanged(InArgs._OnEntityOrderChanged)
+		.OnEntitiesRemoved(InArgs._OnEntitiesRemoved);
+
+	SDMXEntityTreeViewBase::Construct(BaseArguments);
+}
+
+TSharedPtr<FDMXEntityTreeEntityNode> SDMXFixtureTypeTree::CreateEntityNode(UDMXEntity* Entity)
+{
+	check(Entity);
+	TSharedPtr<FDMXEntityTreeEntityNode> NewNode = MakeShared<FDMXEntityTreeEntityNode>(Entity);
+	RefreshFilteredState(NewNode, false);
+
+	// Error status
+	FText InvalidReason;
+	if (!Entity->IsValidEntity(InvalidReason))
+	{
+		NewNode->SetErrorStatus(InvalidReason);
+	}
+
+	return NewNode;
+}
+
+TSharedRef<SWidget> SDMXFixtureTypeTree::GenerateAddNewEntityButton()
+{
+	FText AddButtonLabel = FDMXEditorCommands::Get().AddNewEntityFixtureType->GetLabel();
+	FText AddButtonToolTip = FDMXEditorCommands::Get().AddNewEntityFixtureType->GetDescription();
+
+	return
+		SNew(SButton)
+			.ButtonStyle(FEditorStyle::Get(), "FlatButton.Success")
+			.ForegroundColor(FLinearColor::White)
+			.ToolTipText(AddButtonToolTip)
+			.ContentPadding(FMargin(5.0f, 1.0f))
+			.OnClicked(this, &SDMXFixtureTypeTree::OnAddNewFixtureTypeClicked)
+			.Content()
+			[
+				SNew(SHorizontalBox)
+
+				+ SHorizontalBox::Slot()
+				.AutoWidth()
+				.Padding(FMargin(0, 1))
+				[
+					SNew(SImage)
+					.Image(FEditorStyle::GetBrush("Plus"))
+				]
+
+				+ SHorizontalBox::Slot()
+				.VAlign(VAlign_Center)
+				.AutoWidth()
+				.Padding(FMargin(2, 0, 2, 0))
+				[
+					SNew(STextBlock)
+					.Text(AddButtonLabel)
+				]
+			];
+}
+
+void SDMXFixtureTypeTree::RebuildNodes(const TSharedPtr<FDMXEntityTreeRootNode>& InRootNode)
+{
+	EntityNodeToEntityRowMap.Reset();
+
+	UDMXLibrary* Library = GetDMXLibrary();
+	check(Library != nullptr);
+
+	Library->ForEachEntityOfType<UDMXEntityFixtureType>([this](UDMXEntityFixtureType* FixtureType)
+		{
+			// Create this entity's node
+			TSharedPtr<FDMXEntityTreeEntityNode> FixtureTypeNode = CreateEntityNode(FixtureType);
+
+			// For each Entity, we find or create a category node then add the entity as its child
+			const FDMXFixtureCategory DMXCategory = FixtureType->DMXCategory;
+			const FText DMXCategoryName = FText::FromName(DMXCategory);
+
+			// Get the category if already existent or create it
+			constexpr FDMXEntityTreeCategoryNode::ECategoryType CategoryType = FDMXEntityTreeCategoryNode::ECategoryType::DMXCategory;
+			TSharedPtr<FDMXEntityTreeNodeBase> CategoryNode = GetOrCreateCategoryNode(CategoryType, DMXCategoryName, FixtureType->DMXCategory);
+
+			CategoryNode->AddChild(FixtureTypeNode);
+		}
+	);
+
+	InRootNode->SortChildren();
+}
+
+TSharedRef<ITableRow> SDMXFixtureTypeTree::OnGenerateRow(TSharedPtr<FDMXEntityTreeNodeBase> Node, const TSharedRef<STableViewBase>& OwnerTable)
+{
+		// Create the node of the appropriate type
+	if (Node->GetNodeType() == FDMXEntityTreeNodeBase::ENodeType::CategoryNode)
+	{
+		const TSharedPtr<FDMXEntityTreeCategoryNode> CategoryNode = StaticCastSharedPtr<FDMXEntityTreeCategoryNode>(Node);
+		constexpr bool bIsRootCategory = true;
+
+		return 
+			SNew(SDMXFixtureTypeTreeCategoryRow, OwnerTable, CategoryNode, bIsRootCategory, SharedThis(this))
+			.OnFixtureTypeOrderChanged(OnEntityOrderChanged)
+			[
+				SNew(STextBlock)
+				.Text(Node->GetDisplayNameText())
+				.TextStyle(FEditorStyle::Get(), "DetailsView.CategoryTextStyle")
+			];
+	}
+	else
+	{
+		const TSharedPtr<FDMXEntityTreeEntityNode> EntityNode = StaticCastSharedPtr<FDMXEntityTreeEntityNode>(Node);
+
+		TSharedRef<SDMXFixtureTypeTreeFixtureTypeRow> FixtureTypeRow = SNew(SDMXFixtureTypeTreeFixtureTypeRow, EntityNode, OwnerTable, SharedThis(this))
+			.OnGetFilterText(this, &SDMXFixtureTypeTree::GetFilterText)
+			.OnEntityDragged(this, &SDMXFixtureTypeTree::OnEntitiesDragged)
+			.OnFixtureTypeOrderChanged(OnEntityOrderChanged);
+
+		EntityNodeToEntityRowMap.Add(EntityNode.ToSharedRef(), FixtureTypeRow);
+
+		return FixtureTypeRow;
+	}
+}
+
+TSharedPtr<SWidget> SDMXFixtureTypeTree::OnContextMenuOpen()
+{
+	const bool bCloseWindowAfterMenuSelection = true;
+	FMenuBuilder MenuBuilder(bCloseWindowAfterMenuSelection, CommandList);
+
+	if (GetSelectedEntities().Num() > 0)
+	{
+		MenuBuilder.BeginSection("BasicOperations");
+		{
+			MenuBuilder.AddMenuEntry(FGenericCommands::Get().Rename);
+			MenuBuilder.AddMenuEntry(FGenericCommands::Get().Cut);
+			MenuBuilder.AddMenuEntry(FGenericCommands::Get().Copy);
+			MenuBuilder.AddMenuEntry(FGenericCommands::Get().Paste);
+			MenuBuilder.AddMenuEntry(FGenericCommands::Get().Duplicate);
+			MenuBuilder.AddMenuEntry(FGenericCommands::Get().Delete);
+		}
+		MenuBuilder.EndSection();
+	}
+	else
+	{
+		MenuBuilder.BeginSection("BasicOperations");
+		{
+			MenuBuilder.AddMenuEntry(FDMXEditorCommands::Get().AddNewEntityFixtureType);
+			MenuBuilder.AddMenuEntry(FGenericCommands::Get().Paste);
+		}
+		MenuBuilder.EndSection();
+	}
+
+	return MenuBuilder.MakeWidget();
+}
+
+void SDMXFixtureTypeTree::OnCutSelectedNodes()
+{
+	TArray<UDMXEntity*>&& SelectedNodes = GetSelectedEntities();
+	const FScopedTransaction Transaction(SelectedNodes.Num() > 1 ? LOCTEXT("CutFixtureTypes", "Cut Fixture Types") : LOCTEXT("CutFixtureType", "Cut Fixture Type"));
+
+	OnCopySelectedNodes();
+	OnDeleteNodes();
+}
+
+bool SDMXFixtureTypeTree::CanCutNodes() const
+{
+	return CanCopyNodes() && CanDeleteNodes();
+}
+
+void SDMXFixtureTypeTree::OnCopySelectedNodes()
+{
+	TArray<UDMXEntity*>&& EntitiesToCopy = GetSelectedEntities();
+
+	// Copy the entities to the clipboard
+	FDMXEditorUtils::CopyEntities(MoveTemp(EntitiesToCopy));
+}
+
+bool SDMXFixtureTypeTree::CanCopyNodes() const
+{
+	TArray<UDMXEntity*>&& EntitiesToCopy = GetSelectedEntities();
+	return EntitiesToCopy.Num() > 0;
+}
+
+void SDMXFixtureTypeTree::OnPasteNodes()
+{
+	// Get the Entities to paste from the clipboard
+	TArray<UDMXEntity*> NewObjects;
+	FDMXEditorUtils::GetEntitiesFromClipboard(NewObjects);
+
+	if (NewObjects.Num() != 0)
+	{
+		// Get the library that's being edited
+		UDMXLibrary* Library = GetDMXLibrary();
+		check(Library);
+
+		// Start transaction for Undo and take a snapshot of the current Library state
+		const FScopedTransaction PasteEntities(NewObjects.Num() > 1 ? LOCTEXT("PasteFixtureTypes", "Paste Fixture Types") : LOCTEXT("PasteFixtureType", "Paste Fixture Type"));
+		Library->Modify();
+
+		// Add each pasted Entity to the Library
+		for (UDMXEntity* NewEntity : NewObjects)
+		{
+			check(NewEntity);
+
+			// Move the Entity from the transient package into the Library package
+			NewEntity->Rename(*MakeUniqueObjectName(Library, NewEntity->GetClass()).ToString(), Library, REN_DoNotDirty | REN_DontCreateRedirectors);
+
+			// Make sure the Entity's name won't collide with existing ones
+			NewEntity->SetName(FDMXEditorUtils::FindUniqueEntityName(Library, NewEntity->GetClass(), NewEntity->GetDisplayName()));
+
+			Library->AddEntity(NewEntity);
+
+			OnEntitiesAdded.ExecuteIfBound();
+		}
+
+		UpdateTree();
+		SelectItemsByEntities(NewObjects, ESelectInfo::OnMouseClick);
+	}
+}
+
+bool SDMXFixtureTypeTree::CanPasteNodes() const
+{
+	return FDMXEditorUtils::CanPasteEntities();
+}
+
+void SDMXFixtureTypeTree::OnDuplicateNodes()
+{
+	UDMXLibrary* Library = GetDMXLibrary();
+	TArray<UDMXEntity*> SelectedEntities = GetSelectedEntities();
+
+	if (Library && SelectedEntities.Num() > 0)
+	{
+		// Force the text box being edited (if any) to commit its text. The duplicate operation may trigger a regeneration of the tree view,
+		// releasing all row widgets. If one row was in edit mode (rename/rename on create), it was released before losing the focus and
+		// this would prevent the completion of the 'rename' or 'create + give initial name' transaction (occurring on focus lost).
+		FSlateApplication::Get().ClearKeyboardFocus();
+
+		const FScopedTransaction Transaction(SelectedEntities.Num() > 1 ? LOCTEXT("DuplicateFixtureTypes", "Duplicate Fixture Types") : LOCTEXT("DuplicateFixtureType", "Duplicate Fixture Type"));
+		Library->Modify();
+
+		TArray<UDMXEntity*> NewEntities;
+		NewEntities.Reserve(SelectedEntities.Num());
+
+		// We'll have the duplicates be placed right after their original counterparts
+		int32 NewEntityIndex = Library->FindEntityIndex(SelectedEntities.Last(0));
+		for (UDMXEntity* Entity : SelectedEntities)
+		{
+			FObjectDuplicationParameters DuplicationParams(Entity, GetDMXLibrary());
+			
+			if (UDMXEntity* EntityCopy = CastChecked<UDMXEntity>(StaticDuplicateObjectEx(DuplicationParams)))
+			{
+				EntityCopy->SetName(FDMXEditorUtils::FindUniqueEntityName(Library, EntityCopy->GetClass(), EntityCopy->GetDisplayName()));
+				NewEntities.Add(EntityCopy);
+				
+				Library->AddEntity(EntityCopy);
+				Library->SetEntityIndex(EntityCopy, ++NewEntityIndex);
+			}
+		}
+		
+		OnEntitiesAdded.ExecuteIfBound();
+		
+		UpdateTree(); // Need to refresh tree so new entities have nodes created for them
+		SelectItemsByEntities(NewEntities, ESelectInfo::OnMouseClick); // OnMouseClick triggers selection updated event
+	}
+}
+
+bool SDMXFixtureTypeTree::CanDuplicateNodes() const
+{
+	return CanCopyNodes();
+}
+
+void SDMXFixtureTypeTree::OnDeleteNodes()
+{
+	TArray<UDMXEntity*> EntitiesToDelete = GetSelectedEntities();
+	
+	// Check for entities being used by other objects
+	TArray<UDMXEntity*> EntitiesInUse;
+	for (UDMXEntity* Entity : EntitiesToDelete)
+	{
+		if (FDMXEditorUtils::IsEntityUsed(GetDMXLibrary(), Entity))
+		{
+			EntitiesInUse.Add(Entity);
+		}
+	}
+
+	// Clears references to the Entities and delete them
+	FDMXEditorUtils::RemoveEntities(GetDMXLibrary(), MoveTemp(EntitiesToDelete));
+
+	OnEntitiesRemoved.ExecuteIfBound();
+
+	UpdateTree();
+}
+
+bool SDMXFixtureTypeTree::CanDeleteNodes() const
+{
+	return GetSelectedEntities().Num() > 0;
+}
+
+void SDMXFixtureTypeTree::OnRenameNode()
+{
+	TArray<TSharedPtr<FDMXEntityTreeNodeBase>> SelectedItems = EntitiesTreeWidget->GetSelectedItems();
+
+	if (SelectedItems.Num() == 1 && SelectedItems[0].IsValid() && SelectedItems[0]->GetNodeType() == FDMXEntityTreeNodeBase::ENodeType::EntityNode)
+	{
+		EntitiesTreeWidget->RequestScrollIntoView(SelectedItems[0].ToSharedRef());
+
+		const TSharedPtr<FDMXEntityTreeEntityNode> EntityNode = StaticCastSharedPtr<FDMXEntityTreeEntityNode>(SelectedItems[0]);
+		const TSharedPtr<SDMXFixtureTypeTreeFixtureTypeRow> FixtureTypeRow = FindEntityRowByNode(EntityNode.ToSharedRef());
+
+		if (FixtureTypeRow.IsValid())
+		{
+			FixtureTypeRow->EnterRenameMode();
+		}
+	}
+}
+
+bool SDMXFixtureTypeTree::CanRenameNode() const
+{
+	TArray<TSharedPtr<FDMXEntityTreeNodeBase>> SelectedItems = EntitiesTreeWidget->GetSelectedItems();
+
+	return 
+		SelectedItems.Num() == 1 && 
+		SelectedItems[0].IsValid() &&
+		SelectedItems[0]->GetNodeType() == FDMXEntityTreeNodeBase::ENodeType::EntityNode;
+}
+
+TSharedPtr<SDMXFixtureTypeTreeFixtureTypeRow> SDMXFixtureTypeTree::FindEntityRowByNode(const TSharedRef<FDMXEntityTreeEntityNode>& EntityNode)
+{
+	if (const TSharedRef<SDMXFixtureTypeTreeFixtureTypeRow>* RowPtr = EntityNodeToEntityRowMap.Find(EntityNode))
+	{
+		return *RowPtr;
+	}
+
+	return nullptr;
+}
+
+FReply SDMXFixtureTypeTree::OnAddNewFixtureTypeClicked()
+{
+	check(DMXEditor.IsValid());
+	TSharedPtr<FDMXEditor> PinnedEditor = DMXEditor.Pin();
+
+	PinnedEditor->GetToolkitCommands()->ExecuteAction(FDMXEditorCommands::Get().AddNewEntityFixtureType.ToSharedRef());
+	return FReply::Handled();
+}
+
+#undef LOCTEXT_NAMESPACE
diff --git a/Engine/Plugins/VirtualProduction/DMX/DMXEngine/Source/DMXEditor/Private/Widgets/FixtureType/SDMXFixtureTypeTree.h b/Engine/Plugins/VirtualProduction/DMX/DMXEngine/Source/DMXEditor/Private/Widgets/FixtureType/SDMXFixtureTypeTree.h
new file mode 100644
--- /dev/null
+++ b/Engine/Plugins/VirtualProduction/DMX/DMXEngine/Source/DMXEditor/Private/Widgets/FixtureType/SDMXFixtureTypeTree.h
@@ -0,0 +1,75 @@
+// Copyright Epic Games, Inc. All Rights Reserved.
+
+#pragma once
+
+#include "Widgets/SDMXEntityTreeViewBase.h"
+
+class FDMXEntityTreeEntityNode;
+class SDMXFixtureTypeTreeFixtureTypeRow;
+
+
+/**
+ * A tree of Fixture Types in a DMX Library.
+ */
+class SDMXFixtureTypeTree
+	: public SDMXEntityTreeViewBase
+{
+public:
+	SLATE_BEGIN_ARGS(SDMXFixtureTypeTree)
+	{}
+		/** The DMX Editor that owns this widget */
+		SLATE_ARGUMENT(TWeakPtr<FDMXEditor>, DMXEditor)
+
+		/** Exectued when the list changed its selection */
+		SLATE_EVENT(FDMXOnSelectionChanged, OnSelectionChanged)
+
+		/** Exectued when entites were added to the DMXEditor's library */
+		SLATE_EVENT(FSimpleDelegate, OnEntitiesAdded)
+
+		/** Exectued when entites were reorderd in the list, and potentially in the library */
+		SLATE_EVENT(FSimpleDelegate, OnEntityOrderChanged)
+
+		/** Exectued when entites were removed from the DMXEditor's library */
+		SLATE_EVENT(FSimpleDelegate, OnEntitiesRemoved)
+
+	SLATE_END_ARGS()
+
+	/** Constructs the widget */
+	void Construct(const FArguments& InArgs, const TSubclassOf<UDMXEntity> InListType);
+
+	/** Destructor */
+	virtual ~SDMXFixtureTypeTree() {}
+
+	/** Creates a Node for the entity */
+	TSharedPtr<FDMXEntityTreeEntityNode> CreateEntityNode(UDMXEntity* Entity);
+
+protected:
+	//~ Begin SDMXEntityTreeViewBase interface
+	virtual TSharedRef<SWidget> GenerateAddNewEntityButton();
+	virtual void RebuildNodes(const TSharedPtr<FDMXEntityTreeRootNode>& InRootNode);
+	virtual TSharedRef<ITableRow> OnGenerateRow(TSharedPtr<FDMXEntityTreeNodeBase> Node, const TSharedRef<STableViewBase>& OwnerTable);
+	virtual TSharedPtr<SWidget> OnContextMenuOpen() override;
+	virtual void OnCutSelectedNodes();
+	virtual bool CanCutNodes() const;
+	virtual void OnCopySelectedNodes();
+	virtual bool CanCopyNodes() const;
+	virtual void OnPasteNodes();
+	virtual bool CanPasteNodes() const;
+	virtual bool CanDuplicateNodes() const;
+	virtual void OnDuplicateNodes();
+	virtual void OnDeleteNodes();
+	virtual bool CanDeleteNodes() const;
+	virtual void OnRenameNode();
+	virtual bool CanRenameNode() const;
+	//~ End SDMXEntityTreeViewBase interface
+
+private:
+	/** Called when the Add button was clicked */
+	FReply OnAddNewFixtureTypeClicked();
+
+	/** Returns the row that corresponds to the node */
+	TSharedPtr<SDMXFixtureTypeTreeFixtureTypeRow> FindEntityRowByNode(const TSharedRef<FDMXEntityTreeEntityNode>& EntityNode);
+
+	/** Map of entity nodes and their row in the tree */
+	TMap<TSharedRef<FDMXEntityTreeEntityNode>, TSharedRef<SDMXFixtureTypeTreeFixtureTypeRow>> EntityNodeToEntityRowMap;
+};
diff --git a/Engine/Plugins/VirtualProduction/DMX/DMXEngine/Source/DMXEditor/Private/Widgets/FixtureType/SDMXFixtureTypeTreeCategoryRow.cpp b/Engine/Plugins/VirtualProduction/DMX/DMXEngine/Source/DMXEditor/Private/Widgets/FixtureType/SDMXFixtureTypeTreeCategoryRow.cpp
new file mode 100644
--- /dev/null
+++ b/Engine/Plugins/VirtualProduction/DMX/DMXEngine/Source/DMXEditor/Private/Widgets/FixtureType/SDMXFixtureTypeTreeCategoryRow.cpp
@@ -0,0 +1,264 @@
+// Copyright Epic Games, Inc. All Rights Reserved.
+
+#include "SDMXFixtureTypeTreeCategoryRow.h"
+
+#include "DMXEditorUtils.h"
+#include "SDMXFixtureTypeTree.h"
+#include "DragDrop/DMXEntityDragDropOp.h"
+#include "Library/DMXEntityFixtureType.h"
+#include "Library/DMXLibrary.h"
+#include "Widgets/DMXEntityTreeNode.h"
+
+#include "EditorStyleSet.h"
+#include "ScopedTransaction.h"
+#include "Widgets/SBoxPanel.h"
+#include "Widgets/Layout/SBorder.h"
+#include "Widgets/Views/SExpanderArrow.h"
+
+
+#define LOCTEXT_NAMESPACE "SDMXFixtureTypeTreeCategoryRow"
+
+void SDMXFixtureTypeTreeCategoryRow::Construct(const FArguments& InArgs, const TSharedRef<STableViewBase>& InOwnerTableView, TSharedPtr<FDMXEntityTreeCategoryNode> InCategoryNode, bool bInIsRootCategory, TWeakPtr<SDMXFixtureTypeTree> InFixtureTypeTree)
+{
+	check(InCategoryNode.IsValid());
+
+	WeakFixtureTypeTree = InFixtureTypeTree;
+	WeakCategoryNode = InCategoryNode;
+	
+	OnFixtureTypeOrderChanged = InArgs._OnFixtureTypeOrderChanged;
+
+	// background color tint
+	const FLinearColor BackgroundTint(0.6f, 0.6f, 0.6f, bInIsRootCategory ? 1.0f : 0.3f);
+
+	// rebuilds the whole table row from scratch
+	ChildSlot
+	.Padding(0.0f, 2.0f, 0.0f, 0.0f)
+	[
+		SAssignNew(ContentBorder, SBorder)
+		.BorderImage(this, &SDMXFixtureTypeTreeCategoryRow::GetBackgroundImageBrush)
+		.Padding(FMargin(0.0f, 3.0f))
+		.BorderBackgroundColor(BackgroundTint)
+		.ToolTipText(InCategoryNode->GetToolTip())
+		[
+			SNew(SHorizontalBox)
+
+			+ SHorizontalBox::Slot()
+			.VAlign(VAlign_Center)
+			.Padding(2.0f, 2.0f, 2.0f, 2.0f)
+			.AutoWidth()
+			[
+				SNew(SExpanderArrow, STableRow<TSharedPtr<FDMXEntityTreeCategoryNode>>::SharedThis(this))
+			]
+
+			+ SHorizontalBox::Slot()
+			.VAlign(VAlign_Center)
+			[
+				InArgs._Content.Widget
+			]
+		]
+	];
+
+	STableRow<TSharedPtr<FDMXEntityTreeCategoryNode>>::ConstructInternal(
+		STableRow<TSharedPtr<FDMXEntityTreeCategoryNode>>::FArguments()
+			.Style(FEditorStyle::Get(), "DetailsView.TreeView.TableRow")
+			.ShowSelection(false),
+		InOwnerTableView
+	);
+}
+
+void SDMXFixtureTypeTreeCategoryRow::OnDragEnter(const FGeometry& MyGeometry, const FDragDropEvent& DragDropEvent)
+{
+	if (TSharedPtr<FDMXEntityDragDropOperation> EntityDragDropOp = DragDropEvent.GetOperationAs<FDMXEntityDragDropOperation>())
+	{
+		TestCanDropWithFeedback(EntityDragDropOp.ToSharedRef());
+	}
+}
+
+FReply SDMXFixtureTypeTreeCategoryRow::OnDrop(const FGeometry& MyGeometry, const FDragDropEvent& DragDropEvent)
+{
+	const TSharedPtr<FDMXEntityDragDropOperation> EntityDragDropOp = DragDropEvent.GetOperationAs<FDMXEntityDragDropOperation>();
+	const TSharedPtr<SDMXFixtureTypeTree> FixtureTypeTree = WeakFixtureTypeTree.Pin();
+	const TSharedPtr<FDMXEntityTreeCategoryNode> CategoryNode = WeakCategoryNode.Pin();
+	UDMXLibrary* DMXLibrary = GetDMXLibrary();
+
+	if (EntityDragDropOp.IsValid() && CategoryNode.IsValid() && DMXLibrary)
+	{
+		if (TestCanDropWithFeedback(EntityDragDropOp.ToSharedRef()))
+		{
+			const TArray<TWeakObjectPtr<UDMXEntity>> DraggedEntities = EntityDragDropOp->GetDraggedEntities();
+
+			// Register transaction and current DMX library state for Undo
+			const FScopedTransaction ChangeCategoryTransaction = FScopedTransaction(
+				FText::Format(LOCTEXT("ChangeEntitiesCategory", "Change {0}|plural(one=Entity, other=Entities) category"), DraggedEntities.Num())
+			);
+
+			DMXLibrary->Modify();
+
+			const FDMXFixtureCategory& FixtureCategory = CategoryNode->GetCategoryValue();
+			for (TWeakObjectPtr<UDMXEntity> Entity : DraggedEntities)
+			{
+				if (UDMXEntityFixtureType* FixtureType = Cast<UDMXEntityFixtureType>(Entity))
+				{
+					FixtureType->Modify();
+
+					FixtureType->PreEditChange(nullptr);
+					FixtureType->DMXCategory = FixtureCategory;
+					FixtureType->PostEditChange();
+				}
+			}
+
+			if (CategoryNode->GetChildren().Num() > 0)
+			{
+				if (TSharedPtr<FDMXEntityTreeEntityNode> EntityNode = StaticCastSharedPtr<FDMXEntityTreeEntityNode>(CategoryNode->GetChildren().Last(0)))
+				{
+					// Index after last entity in hovered category
+					UDMXEntity* LastEntityInCategory = EntityNode->GetEntity();
+					const int32 LastEntityIndex = DMXLibrary->FindEntityIndex(LastEntityInCategory);
+					check(LastEntityIndex != INDEX_NONE);
+
+					// Move dragged entities after the last ones in the category
+					// Reverse for to keep dragged entities order.
+					int32 DesiredIndex = LastEntityIndex + 1;
+					for (int32 EntityIndex = DraggedEntities.Num() - 1; EntityIndex > -1; --EntityIndex)
+					{
+						if (UDMXEntity* Entity = DraggedEntities[EntityIndex].Get())
+						{
+							DMXLibrary->SetEntityIndex(Entity, DesiredIndex);
+						}
+					}
+				}
+			}
+
+			FixtureTypeTree->UpdateTree();
+
+			OnFixtureTypeOrderChanged.ExecuteIfBound();
+
+			return FReply::Handled().EndDragDrop();
+		}
+	}
+
+	return FReply::Unhandled();
+}
+
+bool SDMXFixtureTypeTreeCategoryRow::TestCanDropWithFeedback(const TSharedRef<FDMXEntityDragDropOperation>& EntityDragDropOp) const
+{
+	if (const TSharedPtr<FDMXEntityTreeCategoryNode> CategoryNode = WeakCategoryNode.Pin())
+	{
+		// Only handle entity drag drop ops that contain fixture types only 
+		const bool bOnlyFixturePatchesAreDragged = [EntityDragDropOp]()
+		{
+			for (UClass* Class : EntityDragDropOp->GetDraggedEntityTypes())
+			{
+				if (Class != UDMXEntityFixtureType::StaticClass())
+				{
+					return false;
+				}
+			}
+			return true;
+		}();
+
+		// Only handle entity drag drop ops that were dragged from and to the same dmx library
+		const bool bRowAndEntitiesAreOfSameLibrary = [EntityDragDropOp, this]()
+		{
+			UDMXLibrary* DMXLibrary = GetDMXLibrary();
+			for (TWeakObjectPtr<UDMXEntity> Entity : EntityDragDropOp->GetDraggedEntities())
+			{
+				if (Entity.IsValid() && Entity->GetParentLibrary() != DMXLibrary)
+				{
+					return false;
+				}
+			}
+			return true;
+		}();
+
+		if (bOnlyFixturePatchesAreDragged && bRowAndEntitiesAreOfSameLibrary)
+		{
+			if (CategoryNode->CanDropOntoCategory())
+			{
+				const TArray<TWeakObjectPtr<UDMXEntity>> DraggedEntities = EntityDragDropOp->GetDraggedEntities();
+				const TArray<TSharedRef<FDMXEntityTreeCategoryNode>> DraggedFromCategories = GetCategoryNodesFromDragDropOp(EntityDragDropOp);
+
+				if (DraggedFromCategories.Num() == 1 && DraggedFromCategories[0] == CategoryNode)
+				{
+					// There wouldn't be any change by dragging the items into their own category.
+					EntityDragDropOp->SetFeedbackMessageError(FText::Format(
+						LOCTEXT("DragIntoSelfCategory", "The selected {0} {1}|plural(one=is, other=are) already in this category"),
+						FDMXEditorUtils::GetEntityTypeNameText(UDMXEntityFixtureType::StaticClass(), DraggedEntities.Num() > 1),
+						DraggedEntities.Num()
+					));
+
+					return false;
+				}
+				else
+				{
+					// Generate 'OK' feedback per category
+					const FText ListTypeText = LOCTEXT("Property_DMXCategory", "DMX Category");
+					const FText CategoryText = CategoryNode->GetDisplayNameText();
+
+					EntityDragDropOp->SetFeedbackMessageOK(FText::Format(
+						LOCTEXT("ReorderBeforeAndSetProperty", "Reorder {0} before '{1}'\nSet {2} = '{3}'"),
+						EntityDragDropOp->GetDraggedEntitiesName(),
+						CategoryNode->GetDisplayNameText(),
+						ListTypeText,
+						CategoryText
+					));
+
+					return true;
+				}
+			}
+		}
+	}
+
+	return false;
+}
+
+TArray<TSharedRef<FDMXEntityTreeCategoryNode>> SDMXFixtureTypeTreeCategoryRow::GetCategoryNodesFromDragDropOp(const TSharedRef<FDMXEntityDragDropOperation>& DragDropOp) const
+{
+	TArray<TSharedRef<FDMXEntityTreeCategoryNode>> CategoryNodes;
+	if (const TSharedPtr<SDMXFixtureTypeTree>& FixtureTypeTree = WeakFixtureTypeTree.Pin())
+	{
+		TArray<TWeakObjectPtr<UDMXEntity>> DraggedEntities = DragDropOp->GetDraggedEntities();
+
+		for (TWeakObjectPtr<UDMXEntity> Entity : DraggedEntities)
+		{
+			if (UDMXEntityFixtureType* FixtureType = Cast<UDMXEntityFixtureType>(Entity.Get()))
+			{
+				TSharedPtr<FDMXEntityTreeCategoryNode> CategoryNode = FixtureTypeTree->FindCategoryNodeOfEntity(FixtureType);
+				if (CategoryNode.IsValid())
+				{
+					CategoryNodes.Add(CategoryNode.ToSharedRef());
+				}
+			}
+		}
+	}
+
+	return CategoryNodes;
+}
+
+const FSlateBrush* SDMXFixtureTypeTreeCategoryRow::GetBackgroundImageBrush() const
+{
+	if (IsHovered())
+	{
+		return IsItemExpanded()
+			? FEditorStyle::GetBrush("DetailsView.CategoryTop_Hovered")
+			: FEditorStyle::GetBrush("DetailsView.CollapsedCategory_Hovered");
+	}
+	else
+	{
+		return IsItemExpanded()
+			? FEditorStyle::GetBrush("DetailsView.CategoryTop")
+			: FEditorStyle::GetBrush("DetailsView.CollapsedCategory");
+	}
+}
+
+UDMXLibrary* SDMXFixtureTypeTreeCategoryRow::GetDMXLibrary() const
+{
+	if (const TSharedPtr<SDMXFixtureTypeTree>& FixtureTypeTree = WeakFixtureTypeTree.Pin())
+	{
+		return FixtureTypeTree->GetDMXLibrary();
+	}
+
+	return nullptr;
+}
+
+#undef LOCTEXT_NAMESPACE
diff --git a/Engine/Plugins/VirtualProduction/DMX/DMXEngine/Source/DMXEditor/Private/Widgets/FixtureType/SDMXFixtureTypeTreeCategoryRow.h b/Engine/Plugins/VirtualProduction/DMX/DMXEngine/Source/DMXEditor/Private/Widgets/FixtureType/SDMXFixtureTypeTreeCategoryRow.h
new file mode 100644
--- /dev/null
+++ b/Engine/Plugins/VirtualProduction/DMX/DMXEngine/Source/DMXEditor/Private/Widgets/FixtureType/SDMXFixtureTypeTreeCategoryRow.h
@@ -0,0 +1,64 @@
+// Copyright Epic Games, Inc. All Rights Reserved.
+
+#pragma once
+
+#include "CoreMinimal.h"
+#include "Widgets/Views/SListView.h"
+
+#include "Widgets/Views/STableRow.h"
+
+class FDMXEntityDragDropOperation;
+class FDMXEntityTreeCategoryNode;
+class SDMXFixtureTypeTree;
+class UDMXLibrary;
+
+class STableViewBase;
+
+
+/** A row displaying a Category in a Fixture Type Tree */
+class SDMXFixtureTypeTreeCategoryRow
+	: public STableRow<TSharedPtr<FDMXEntityTreeCategoryNode>>
+{
+public:
+	SLATE_BEGIN_ARGS(SDMXFixtureTypeTreeCategoryRow)
+	{}
+		SLATE_DEFAULT_SLOT(typename SDMXFixtureTypeTreeCategoryRow::FArguments, Content)
+
+		SLATE_EVENT(FSimpleDelegate, OnFixtureTypeOrderChanged)
+
+	SLATE_END_ARGS()
+
+	/** Constructs this widget */
+	void Construct(const FArguments& InArgs, const TSharedRef<STableViewBase>& InOwnerTableView, TSharedPtr<FDMXEntityTreeCategoryNode> InCategoryNode, bool bInIsRootCategory, TWeakPtr<SDMXFixtureTypeTree> InFixtureTypeTree);
+
+protected:
+	//~ SWidget interface begin
+	virtual void OnDragEnter(const FGeometry& MyGeometry, const FDragDropEvent& DragDropEvent) override;
+	virtual FReply OnDrop(const FGeometry& MyGeometry, const FDragDropEvent& DragDropEvent) override;
+	//~ SWidget interface end
+
+private:
+	/** Returns true if the drag drop op can be dropped on this row */
+	bool TestCanDropWithFeedback(const TSharedRef<FDMXEntityDragDropOperation>& EntityDragDropOp) const;
+
+	/** Returns the categories nodes the entities were dragged out from */
+	TArray<TSharedRef<FDMXEntityTreeCategoryNode>> GetCategoryNodesFromDragDropOp(const TSharedRef<FDMXEntityDragDropOperation>& DragDropOp) const;
+
+	/** Returns the background image brush */
+	const FSlateBrush* GetBackgroundImageBrush() const;
+
+	/** Returns the DMXLibrary of this row */
+	UDMXLibrary* GetDMXLibrary() const;
+
+	/** Pointer to the node this row displays */
+	TWeakPtr<FDMXEntityTreeCategoryNode> WeakCategoryNode;
+
+	/** The owning Fixture Patch List */
+	TWeakPtr<SDMXFixtureTypeTree> WeakFixtureTypeTree;
+
+	/** Border that holds the content */
+	TSharedPtr<SBorder> ContentBorder;
+
+	/** Broadcast when the order of fixture types in the library changed */
+	FSimpleDelegate OnFixtureTypeOrderChanged;
+};
diff --git a/Engine/Plugins/VirtualProduction/DMX/DMXEngine/Source/DMXEditor/Private/Widgets/FixtureType/SDMXFixtureTypeTreeFixtureTypeRow.cpp b/Engine/Plugins/VirtualProduction/DMX/DMXEngine/Source/DMXEditor/Private/Widgets/FixtureType/SDMXFixtureTypeTreeFixtureTypeRow.cpp
new file mode 100644
--- /dev/null
+++ b/Engine/Plugins/VirtualProduction/DMX/DMXEngine/Source/DMXEditor/Private/Widgets/FixtureType/SDMXFixtureTypeTreeFixtureTypeRow.cpp
@@ -0,0 +1,371 @@
+// Copyright Epic Games, Inc. All Rights Reserved.
+
+#include "SDMXFixtureTypeTreeFixtureTypeRow.h"
+
+#include "DMXEditorUtils.h"
+#include "DMXEditorStyle.h"
+#include "SDMXFixtureTypeTree.h"
+#include "DragDrop/DMXEntityDragDropOp.h"
+#include "Library/DMXEntityFixtureType.h"
+#include "Library/DMXLibrary.h"
+#include "Widgets/DMXEntityTreeNode.h"
+
+#include "EditorStyleSet.h"
+#include "ScopedTransaction.h"
+#include "Widgets/Images/SImage.h"
+#include "Widgets/Text/SInlineEditableTextBlock.h"
+#include "Widgets/Text/STextBlock.h"
+
+
+#define LOCTEXT_NAMESPACE "SDMXFixtureTypeTreeFixtureTypeRow"
+
+void SDMXFixtureTypeTreeFixtureTypeRow::Construct(const FArguments& InArgs, TSharedPtr<FDMXEntityTreeEntityNode> InEntityNode, TSharedPtr<STableViewBase> InOwnerTableView, TWeakPtr<SDMXFixtureTypeTree> InFixtureTypeTree)
+{
+	// Without ETableRowSignalSelectionMode::Instantaneous, when the user is editing a property in
+	// the inspector panel and then clicks on a different row on the list panel, the selection event
+	// is deferred. But because we update the tree right after a property change and that triggers
+	// selection changes too, the selection change event is triggered only from UpdateTree, with
+	// Direct selection mode, which doesn't trigger the SDMXEntityList::OnSelectionUpdated event.
+	// This setting forces the event with OnMouseClick selection type to be fired as soon as the
+	// row is clicked.
+	STableRow<TSharedPtr<FDMXEntityTreeEntityNode>>::FArguments BaseArguments = STableRow<TSharedPtr<FDMXEntityTreeEntityNode>>::FArguments()
+		.SignalSelectionMode(ETableRowSignalSelectionMode::Instantaneous)
+		.OnDragDetected(this, &SDMXFixtureTypeTreeFixtureTypeRow::HandleOnDragDetected);
+
+	STableRow<TSharedPtr<FDMXEntityTreeEntityNode>>::Construct(BaseArguments, InOwnerTableView.ToSharedRef());
+
+	WeakEntityNode = InEntityNode;
+	WeakFixtureTypeTree = InFixtureTypeTree;
+
+	OnEntityDragged = InArgs._OnEntityDragged;
+	OnGetFilterText = InArgs._OnGetFilterText;
+	OnFixtureTypeOrderChanged = InArgs._OnFixtureTypeOrderChanged;
+
+	const FSlateFontInfo NameFont = FCoreStyle::GetDefaultFontStyle("Regular", 10);
+
+	ChildSlot
+	[
+		SNew(SHorizontalBox)
+		// Status icon to show the user if there's an error with the Entity's usability
+		+ SHorizontalBox::Slot()
+		.AutoWidth()
+		.HAlign(HAlign_Left)
+		.VAlign(VAlign_Center)
+		[
+			SNew(SImage)
+			.Image(this, &SDMXFixtureTypeTreeFixtureTypeRow::GetStatusIcon)
+			.ToolTipText(this, &SDMXFixtureTypeTreeFixtureTypeRow::GetStatusToolTip)
+		]
+		
+		// Entity's name
+		+ SHorizontalBox::Slot()
+		.FillWidth(1.0f)
+		.VAlign(VAlign_Center)
+		.Padding(3.0f, 0.0f)
+		[
+			SAssignNew(InlineEditableFixtureTypeNameWidget, SInlineEditableTextBlock)
+			.Text(this, &SDMXFixtureTypeTreeFixtureTypeRow::GetFixtureTypeName)
+			.Font(NameFont)
+			.HighlightText(this, &SDMXFixtureTypeTreeFixtureTypeRow::GetFilterText)
+			.OnTextCommitted(this, &SDMXFixtureTypeTreeFixtureTypeRow::OnNameTextCommit)
+			.OnVerifyTextChanged(this, &SDMXFixtureTypeTreeFixtureTypeRow::OnNameTextVerifyChanged)
+			.IsSelected(this, &SDMXFixtureTypeTreeFixtureTypeRow::IsSelected)
+		]
+	];
+}
+
+void SDMXFixtureTypeTreeFixtureTypeRow::EnterRenameMode()
+{
+	InlineEditableFixtureTypeNameWidget->EnterEditingMode();
+}
+
+FReply SDMXFixtureTypeTreeFixtureTypeRow::OnMouseButtonDoubleClick(const FGeometry& InMyGeometry, const FPointerEvent& InMouseEvent)
+{
+	const TSharedPtr<SDMXFixtureTypeTree> FixtureTypeTree = WeakFixtureTypeTree.Pin();
+	const TSharedPtr<FDMXEntityTreeEntityNode> EntityNode = WeakEntityNode.Pin();
+
+	if (FixtureTypeTree.IsValid() && EntityNode.IsValid())
+	{
+		FixtureTypeTree->SelectItemByNode(EntityNode.ToSharedRef());
+
+		InlineEditableFixtureTypeNameWidget->EnterEditingMode();
+
+		return FReply::Handled();
+	}
+
+	return FReply::Unhandled();
+}
+
+void SDMXFixtureTypeTreeFixtureTypeRow::OnDragEnter(const FGeometry& MyGeometry, const FDragDropEvent& DragDropEvent)
+{	
+	if (TSharedPtr<FDMXEntityDragDropOperation> EntityDragDropOp = DragDropEvent.GetOperationAs<FDMXEntityDragDropOperation>())
+	{
+		TestCanDropWithFeedback(EntityDragDropOp.ToSharedRef());
+	}
+}
+
+FReply SDMXFixtureTypeTreeFixtureTypeRow::OnDrop(const FGeometry& MyGeometry, const FDragDropEvent& DragDropEvent)
+{
+	const TSharedPtr<FDMXEntityDragDropOperation> EntityDragDropOp = DragDropEvent.GetOperationAs<FDMXEntityDragDropOperation>();
+	const TSharedPtr<SDMXFixtureTypeTree> FixtureTypeTree = WeakFixtureTypeTree.Pin();
+	const TSharedPtr<FDMXEntityTreeEntityNode> EntityNode = WeakEntityNode.Pin();
+	UDMXLibrary* DMXLibrary = FixtureTypeTree.IsValid() ? FixtureTypeTree->GetDMXLibrary() : nullptr;
+	UDMXEntityFixtureType* HoveredFixtureType = EntityNode.IsValid() ? Cast<UDMXEntityFixtureType>(EntityNode->GetEntity()) : nullptr;
+ 
+	if (EntityDragDropOp.IsValid() && FixtureTypeTree.IsValid() && EntityNode.IsValid() && DMXLibrary && HoveredFixtureType)
+	{
+		if(TestCanDropWithFeedback(EntityDragDropOp.ToSharedRef()))
+		{
+			const TArray<TWeakObjectPtr<UDMXEntity>> DraggedEntities = EntityDragDropOp->GetDraggedEntities();
+
+			// Register transaction and current DMX library state for Undo
+			const FScopedTransaction ReorderTransaction = FScopedTransaction(
+				FText::Format(LOCTEXT("ReorderEntities", "Reorder {0}|plural(one=Entity, other=Entities)"), DraggedEntities.Num())
+			);
+
+			DMXLibrary->Modify();
+
+			// The index of the Entity, we're about to insert the dragged one before
+			const int32 InsertBeforeIndex = DMXLibrary->FindEntityIndex(HoveredFixtureType);
+			check(InsertBeforeIndex != INDEX_NONE);
+
+			// Reverse for to keep dragged entities order
+			for (int32 EntityIndex = DraggedEntities.Num() - 1; EntityIndex > -1; --EntityIndex)
+			{
+				if (DraggedEntities[EntityIndex].IsValid())
+				{
+					UDMXEntity* Entity = DraggedEntities[EntityIndex].Get();
+					DMXLibrary->SetEntityIndex(Entity, InsertBeforeIndex);
+				}
+			}
+
+			// Handle cases where it was was dropped on another category
+			TSharedPtr<FDMXEntityTreeCategoryNode> TargetCategoryNode = StaticCastSharedPtr<FDMXEntityTreeCategoryNode>(EntityNode->GetParent().Pin());
+			check(TargetCategoryNode.IsValid());
+
+			const FDMXFixtureCategory& FixtureCategory = TargetCategoryNode->GetCategoryValue();
+			for (TWeakObjectPtr<UDMXEntity> Entity : DraggedEntities)
+			{
+				if (UDMXEntityFixtureType* FixtureType = Cast<UDMXEntityFixtureType>(Entity))
+				{
+					FixtureType->Modify();
+					FixtureType->DMXCategory = FixtureCategory;
+				}
+			}
+
+			// Display the changes in the Entities list
+			FixtureTypeTree->UpdateTree();
+
+			OnFixtureTypeOrderChanged.ExecuteIfBound();
+
+			return FReply::Handled().EndDragDrop();
+		}
+	}
+
+	return FReply::Unhandled();
+}
+
+bool SDMXFixtureTypeTreeFixtureTypeRow::TestCanDropWithFeedback(const TSharedRef<FDMXEntityDragDropOperation>& EntityDragDropOp) const
+{
+	const TSharedPtr<SDMXFixtureTypeTree> FixtureTypeTree = WeakFixtureTypeTree.Pin();
+	const TSharedPtr<FDMXEntityTreeEntityNode> EntityNode = WeakEntityNode.Pin();
+	UDMXEntityFixtureType* HoveredFixtureType = EntityNode.IsValid() ? Cast<UDMXEntityFixtureType>(EntityNode->GetEntity()) : nullptr;
+
+	if (FixtureTypeTree.IsValid() && EntityNode.IsValid() && HoveredFixtureType)
+	{
+		const TArray<TWeakObjectPtr<UDMXEntity>> DraggedEntities = EntityDragDropOp->GetDraggedEntities();
+
+		const bool bOnlyFixtureTypesAreDragged = [DraggedEntities]()
+		{
+			for (TWeakObjectPtr<UDMXEntity> Entity : DraggedEntities)
+			{
+				if (Entity.IsValid() && Entity->GetClass() != UDMXEntityFixtureType::StaticClass())
+				{
+					return false;
+				}
+			}
+			return true;
+		}();
+
+		const bool bRowAndEntitiesAreOfSameDMXLibrary = [DraggedEntities, this]()
+		{
+			UDMXLibrary* DMXLibrary = GetDMXLibrary();
+			for (TWeakObjectPtr<UDMXEntity> Entity : DraggedEntities)
+			{
+				if (Entity.IsValid() && Entity->GetParentLibrary() != DMXLibrary)
+				{
+					return false;
+				}
+			}
+			return true;
+		}();
+
+		if (bOnlyFixtureTypesAreDragged && bRowAndEntitiesAreOfSameDMXLibrary)
+		{
+			// Test the hovered entity is not the dragged entitiy
+			if (DraggedEntities.Contains(HoveredFixtureType))
+			{
+				EntityDragDropOp->SetFeedbackMessageError(FText::Format(
+					LOCTEXT("ReorderBeforeItself", "Drop {0} on itself"),
+					EntityDragDropOp->GetDraggedEntitiesName()
+				));
+
+				return false;
+			}
+
+			check(EntityNode->GetParent().IsValid());
+			TSharedPtr<FDMXEntityTreeCategoryNode> TargetCategoryNode = StaticCastSharedPtr<FDMXEntityTreeCategoryNode>(EntityNode->GetParent().Pin());
+			check(TargetCategoryNode.IsValid());
+
+			// Generate 'OK' feedback per category
+			FText ListTypeText = LOCTEXT("Property_DMXCategory", "DMX Category");
+			FText CategoryText = TargetCategoryNode->GetDisplayNameText();
+
+			EntityDragDropOp->SetFeedbackMessageOK(FText::Format(
+				LOCTEXT("ReorderBeforeAndSetProperty", "Reorder {0} before '{1}'\nSet {2} = '{3}'"),
+				EntityDragDropOp->GetDraggedEntitiesName(),
+				FText::FromString(HoveredFixtureType->GetDisplayName()),
+				ListTypeText,
+				CategoryText
+			));
+
+			return true;
+		}
+	}
+
+	return false;
+}
+
+FText SDMXFixtureTypeTreeFixtureTypeRow::GetFilterText() const
+{
+	if (OnGetFilterText.IsBound())
+	{
+		return OnGetFilterText.Execute();
+	}
+	return FText();
+}
+
+void SDMXFixtureTypeTreeFixtureTypeRow::OnNameTextCommit(const FText& InNewName, ETextCommit::Type InTextCommit)
+{
+	TSharedPtr<SDMXFixtureTypeTree> FixtureTypeTree = WeakFixtureTypeTree.Pin();
+	TSharedPtr<FDMXEntityTreeEntityNode> EntityNode = WeakEntityNode.Pin();
+	UDMXEntity* Entity = EntityNode.IsValid() ? EntityNode->GetEntity() : nullptr;
+	UDMXLibrary* DMXLibrary = GetDMXLibrary();
+
+	if (FixtureTypeTree.IsValid() && EntityNode.IsValid() && Entity && DMXLibrary)
+	{
+		const FString NewNameString = InNewName.ToString();
+
+		// Check if the name is unchanged
+		if (NewNameString.Equals(EntityNode->GetDisplayNameString()))
+		{
+			return;
+		}
+
+		const FScopedTransaction Transaction(LOCTEXT("RenameEntity", "Rename Entity"));
+		DMXLibrary->Modify();
+
+		FDMXEditorUtils::RenameEntity(DMXLibrary, Entity, NewNameString);
+
+		FixtureTypeTree->SelectItemByName(NewNameString, ESelectInfo::OnMouseClick);
+	}
+}
+
+bool SDMXFixtureTypeTreeFixtureTypeRow::OnNameTextVerifyChanged(const FText& InNewText, FText& OutErrorMessage)
+{
+	TSharedPtr<FDMXEntityTreeEntityNode> EntityNode = WeakEntityNode.Pin();
+	UDMXEntity* Entity = EntityNode.IsValid() ? EntityNode->GetEntity() : nullptr;
+	UDMXLibrary* DMXLibrary = GetDMXLibrary();
+
+	if (EntityNode.IsValid() && Entity && DMXLibrary)
+	{
+		FString TextAsString = InNewText.ToString();
+		if (TextAsString.Equals(EntityNode->GetDisplayNameString()))
+		{
+			return true;
+		}
+
+		return FDMXEditorUtils::ValidateEntityName(
+			TextAsString,
+			DMXLibrary,
+			Entity->GetClass(),
+			OutErrorMessage
+		);
+	}
+
+	return false;
+}
+
+FText SDMXFixtureTypeTreeFixtureTypeRow::GetFixtureTypeName() const
+{
+	if (TSharedPtr<FDMXEntityTreeEntityNode> EntityNode = WeakEntityNode.Pin())
+	{
+		return EntityNode->GetDisplayNameText();
+	}
+	return LOCTEXT("InvalidNodeLabel", "Invalid Node");
+}
+
+FReply SDMXFixtureTypeTreeFixtureTypeRow::HandleOnDragDetected(const FGeometry& MyGeometry, const FPointerEvent& MouseEvent)
+{
+	if (MouseEvent.IsMouseButtonDown(EKeys::LeftMouseButton) && OnEntityDragged.IsBound())
+	{
+		if (TSharedPtr<FDMXEntityTreeEntityNode> EntityNode = WeakEntityNode.Pin())
+		{
+			return OnEntityDragged.Execute(EntityNode, MouseEvent);
+		}
+	}
+
+	return FReply::Unhandled();
+}
+
+const FSlateBrush* SDMXFixtureTypeTreeFixtureTypeRow::GetStatusIcon() const
+{
+	static FSlateNoResource EmptyBrush = FSlateNoResource();
+
+	if (TSharedPtr<FDMXEntityTreeEntityNode> EntityNode = WeakEntityNode.Pin())
+	{
+		if (!EntityNode->GetErrorStatus().IsEmpty())
+		{
+			return FEditorStyle::GetBrush("Icons.Error");
+		}
+
+		if (!EntityNode->GetWarningStatus().IsEmpty())
+		{
+			return FEditorStyle::GetBrush("Icons.Warning");
+		}
+	}
+
+	return &EmptyBrush;
+}
+
+FText SDMXFixtureTypeTreeFixtureTypeRow::GetStatusToolTip() const
+{
+	if (TSharedPtr<FDMXEntityTreeEntityNode> EntityNode = WeakEntityNode.Pin())
+	{
+		const FText& ErrorStatus = EntityNode->GetErrorStatus();
+		if (!ErrorStatus.IsEmpty())
+		{
+			return ErrorStatus;
+		}
+
+		const FText& WarningStatus = EntityNode->GetWarningStatus();
+		if (!WarningStatus.IsEmpty())
+		{
+			return WarningStatus;
+		}
+	}
+
+	return FText::GetEmpty();
+}
+
+UDMXLibrary* SDMXFixtureTypeTreeFixtureTypeRow::GetDMXLibrary() const
+{
+	if (const TSharedPtr<SDMXFixtureTypeTree>& FixtureTypeTree = WeakFixtureTypeTree.Pin())
+	{
+		return FixtureTypeTree->GetDMXLibrary();
+	}
+
+	return nullptr;
+}
+
+#undef LOCTEXT_NAMESPACE
diff --git a/Engine/Plugins/VirtualProduction/DMX/DMXEngine/Source/DMXEditor/Private/Widgets/FixtureType/SDMXFixtureTypeTreeFixtureTypeRow.h b/Engine/Plugins/VirtualProduction/DMX/DMXEngine/Source/DMXEditor/Private/Widgets/FixtureType/SDMXFixtureTypeTreeFixtureTypeRow.h
new file mode 100644
--- /dev/null
+++ b/Engine/Plugins/VirtualProduction/DMX/DMXEngine/Source/DMXEditor/Private/Widgets/FixtureType/SDMXFixtureTypeTreeFixtureTypeRow.h
@@ -0,0 +1,103 @@
+// Copyright Epic Games, Inc. All Rights Reserved.
+
+#pragma once
+
+#include "CoreMinimal.h"
+#include "Widgets/Views/SListView.h"
+
+#include "Widgets/Views/STableRow.h"
+
+class FDMXEntityDragDropOperation;
+class FDMXEntityTreeEntityNode;
+class FDMXEntityTreeNodeBase;
+class SDMXFixtureTypeTree;
+class UDMXLibrary;
+
+class SInlineEditableTextBlock;
+
+
+/** A fixture type row in a fixture patch tree */
+class SDMXFixtureTypeTreeFixtureTypeRow
+	: public STableRow<TSharedPtr<FDMXEntityTreeEntityNode>>
+{
+	DECLARE_DELEGATE_RetVal_TwoParams(FReply, FOnEntityDragged, TSharedPtr<FDMXEntityTreeNodeBase>, const FPointerEvent&);
+	DECLARE_DELEGATE_RetVal(FText, FOnGetFilterText);
+	DECLARE_DELEGATE_OneParam(FOnAutoAssignChannelStateChanged, bool);
+
+public:
+	SLATE_BEGIN_ARGS(SDMXFixtureTypeTreeFixtureTypeRow)
+	{}
+		SLATE_DEFAULT_SLOT(typename SDMXFixtureTypeTreeFixtureTypeRow::FArguments, Content)
+
+		SLATE_EVENT(FOnEntityDragged, OnEntityDragged)
+
+		SLATE_EVENT(FOnGetFilterText, OnGetFilterText)
+
+		SLATE_EVENT(FOnAutoAssignChannelStateChanged, OnAutoAssignChannelStateChanged)
+
+		SLATE_EVENT(FSimpleDelegate, OnFixtureTypeOrderChanged)
+
+	SLATE_END_ARGS()
+
+	/** Constructs this widget */
+	void Construct(const FArguments& InArgs, TSharedPtr<FDMXEntityTreeEntityNode> InEntityNode, TSharedPtr<STableViewBase> InOwnerTableView, TWeakPtr<SDMXFixtureTypeTree> InFixtureTypeTree);
+
+	/** Starts renaming the widget (enters editing mode on its Inline Editable Text Block) */
+	void EnterRenameMode();
+
+protected:
+	//~ SWidget interface begin
+	virtual FReply OnMouseButtonDoubleClick(const FGeometry& InMyGeometry, const FPointerEvent& InMouseEvent) override;
+	virtual void OnDragEnter(const FGeometry& MyGeometry, const FDragDropEvent& DragDropEvent) override;
+	virtual FReply OnDrop(const FGeometry& MyGeometry, const FDragDropEvent& DragDropEvent) override;
+	//~ SWidget interface end
+
+private:
+	/** Returns true if the drag drop op can be dropped on this row */
+	bool TestCanDropWithFeedback(const TSharedRef<FDMXEntityDragDropOperation>& EntityDragDropOp) const;
+
+	/** Called when the row is being dragged */
+	FReply HandleOnDragDetected(const FGeometry& MyGeometry, const FPointerEvent& MouseEvent);
+
+	/** Get the current filter text from the search box */
+	FText GetFilterText() const;
+
+	/** Verifies the name of the component when changing it */
+	bool OnNameTextVerifyChanged(const FText& InNewText, FText& OutErrorMessage);
+
+	/** Commits the new name of the component */
+	void OnNameTextCommit(const FText& InNewName, ETextCommit::Type InTextCommit);
+
+	/** Returns the name of the patch as text */
+	FText GetFixtureTypeName() const;
+
+	/** Get the icon for the Entity usability status. If it's all good, it's an emtpy image. */
+	const FSlateBrush* GetStatusIcon() const;
+
+	/** Get the tool tip text for the status icon */
+	FText GetStatusToolTip() const;
+
+	/** Returns the DMXLibrary of this row */
+	UDMXLibrary* GetDMXLibrary() const;
+
+	/** The node this row displays */
+	TWeakPtr<FDMXEntityTreeEntityNode> WeakEntityNode;
+
+	/** The fixture type tree widget that owns this row */
+	TWeakPtr<SDMXFixtureTypeTree> WeakFixtureTypeTree;
+
+	/** Called when the entity list changed auto assign channel state */
+	FOnAutoAssignChannelStateChanged OnAutoAssignChannelStateChanged;
+
+	/** Delegate broadcast when an entity was dragged */
+	FOnEntityDragged OnEntityDragged;
+
+	/** Delegate broadcast when the filter text is requested */
+	FOnGetFilterText OnGetFilterText;
+
+	/** Called when the entity list changed order of the library's entity array */
+	FSimpleDelegate OnFixtureTypeOrderChanged;
+
+	/** The editable text block to display and change the name of the fixture type */
+	TSharedPtr<SInlineEditableTextBlock> InlineEditableFixtureTypeNameWidget;
+};
diff --git a/Engine/Plugins/VirtualProduction/DMX/DMXEngine/Source/DMXEditor/Private/Widgets/SDMXEntityList.cpp b/Engine/Plugins/VirtualProduction/DMX/DMXEngine/Source/DMXEditor/Private/Widgets/SDMXEntityList.cpp
deleted file mode 100644
diff --git a/Engine/Plugins/VirtualProduction/DMX/DMXEngine/Source/DMXEditor/Private/Widgets/SDMXEntityList.h b/Engine/Plugins/VirtualProduction/DMX/DMXEngine/Source/DMXEditor/Private/Widgets/SDMXEntityList.h
deleted file mode 100644
diff --git a/Engine/Plugins/VirtualProduction/DMX/DMXEngine/Source/DMXEditor/Private/Widgets/SDMXEntityTreeViewBase.cpp b/Engine/Plugins/VirtualProduction/DMX/DMXEngine/Source/DMXEditor/Private/Widgets/SDMXEntityTreeViewBase.cpp
new file mode 100644
--- /dev/null
+++ b/Engine/Plugins/VirtualProduction/DMX/DMXEngine/Source/DMXEditor/Private/Widgets/SDMXEntityTreeViewBase.cpp
@@ -0,0 +1,746 @@
+// Copyright Epic Games, Inc. All Rights Reserved.
+
+#include "Widgets/SDMXEntityTreeViewBase.h"
+
+#include "DMXEditor.h"
+#include "DragDrop/DMXEntityDragDropOp.h"
+#include "Library/DMXLibrary.h"
+
+#include "EditorStyleSet.h"
+#include "Framework/Commands/GenericCommands.h"
+#include "Framework/MultiBox/MultiBoxBuilder.h"
+#include "Framework/Commands/UICommandList.h"
+#include "Widgets/Docking/SDockTab.h"
+#include "Widgets/Input/SSearchBox.h"
+#include "Widgets/Images/SImage.h"
+#include "Widgets/Views/STreeView.h"
+
+
+#define LOCTEXT_NAMESPACE "SDMXEntityTreeViewBase"
+
+void SDMXEntityTreeViewBase::Construct(const FArguments& InArgs)
+{
+	// Initialize Widget input variables
+	DMXEditor = InArgs._DMXEditor;
+	OnSelectionChangedDelegate = InArgs._OnSelectionChanged;
+	OnEntitiesAdded = InArgs._OnEntitiesAdded;
+	OnEntityOrderChanged = InArgs._OnEntityOrderChanged;
+	OnEntitiesRemoved = InArgs._OnEntitiesRemoved;
+
+	// listen to common editor shortcuts for copy/paste etc
+	CommandList = MakeShared<FUICommandList>();
+	CommandList->MapAction(FGenericCommands::Get().Cut,
+		FUIAction(FExecuteAction::CreateSP(this, &SDMXEntityTreeViewBase::OnCutSelectedNodes),
+			FCanExecuteAction::CreateSP(this, &SDMXEntityTreeViewBase::CanCutNodes))
+	);
+	CommandList->MapAction(FGenericCommands::Get().Copy,
+		FUIAction(FExecuteAction::CreateSP(this, &SDMXEntityTreeViewBase::OnCopySelectedNodes),
+			FCanExecuteAction::CreateSP(this, &SDMXEntityTreeViewBase::CanCopyNodes))
+	);
+	CommandList->MapAction(FGenericCommands::Get().Paste,
+		FUIAction(FExecuteAction::CreateSP(this, &SDMXEntityTreeViewBase::OnPasteNodes),
+			FCanExecuteAction::CreateSP(this, &SDMXEntityTreeViewBase::CanPasteNodes))
+	);
+	CommandList->MapAction(FGenericCommands::Get().Duplicate,
+		FUIAction(FExecuteAction::CreateSP(this, &SDMXEntityTreeViewBase::OnDuplicateNodes),
+			FCanExecuteAction::CreateSP(this, &SDMXEntityTreeViewBase::CanDuplicateNodes))
+	);
+	CommandList->MapAction(FGenericCommands::Get().Delete,
+		FUIAction(FExecuteAction::CreateSP(this, &SDMXEntityTreeViewBase::OnDeleteNodes),
+			FCanExecuteAction::CreateSP(this, &SDMXEntityTreeViewBase::CanDeleteNodes))
+	);
+	CommandList->MapAction(FGenericCommands::Get().Rename,
+		FUIAction(FExecuteAction::CreateSP(this, &SDMXEntityTreeViewBase::OnRenameNode),
+			FCanExecuteAction::CreateSP(this, &SDMXEntityTreeViewBase::CanRenameNode))
+	);
+
+	GEditor->RegisterForUndo(this);
+
+	TSharedPtr<SBorder> HeaderBox = SNew(SBorder)
+		.Padding(0)
+		.BorderImage(FEditorStyle::GetBrush("DetailsView.CategoryTop"))
+		.BorderBackgroundColor(FLinearColor(0.6f, 0.6f, 0.6f, 1.0f))
+		[
+			SNew(SHorizontalBox)
+
+			// [+ Add New] button
+			+ SHorizontalBox::Slot()
+			.Padding(3.0f)
+			.AutoWidth()
+			.HAlign(HAlign_Left)
+			[
+				GenerateAddNewEntityButton()
+			]
+
+			// Filter box
+			+ SHorizontalBox::Slot()
+			.FillWidth(1.0f)
+			.VAlign(VAlign_Center)
+			.Padding(3.0f)
+			[
+				SAssignNew(FilterBox, SSearchBox)
+				.HintText(LOCTEXT("SearchEntitiesHint", "Search entities"))
+				.OnTextChanged(this, &SDMXEntityTreeViewBase::OnFilterTextChanged)
+			]
+
+		];
+
+	// Tree widget which displays the entities in their categories (e.g. protocol),
+	// and also controls selection and drag/drop
+	RootNode = MakeShared<FDMXEntityTreeRootNode>();
+	EntitiesTreeWidget = SNew(STreeView<TSharedPtr<FDMXEntityTreeNodeBase>>)
+		.ItemHeight(24)
+		.TreeItemsSource(&RootNode->GetChildren())
+		.SelectionMode(ESelectionMode::Multi)
+		.OnGenerateRow(this, &SDMXEntityTreeViewBase::OnGenerateRow)
+		.OnGetChildren(this, &SDMXEntityTreeViewBase::OnGetChildren)
+		.OnExpansionChanged(this, &SDMXEntityTreeViewBase::OnExpansionChanged)
+		.OnSelectionChanged(this, &SDMXEntityTreeViewBase::OnSelectionChanged)
+		.OnContextMenuOpening(this, &SDMXEntityTreeViewBase::OnContextMenuOpen)
+		// .OnItemScrolledIntoView(this, &SDMXEntityTreeViewBase::OnItemScrolledIntoView)
+		.HighlightParentNodesForSelection(false);
+
+	ChildSlot
+	[
+		SNew(SVerticalBox)
+		
+		+ SVerticalBox::Slot()
+		.Padding(0)
+		.AutoHeight()
+		.HAlign(HAlign_Fill)
+		[
+			HeaderBox.ToSharedRef()
+		]
+
+		+ SVerticalBox::Slot()
+		.Padding(0)
+		[
+			SNew(SBorder)
+			.Padding(2.0f)
+			.BorderImage(FEditorStyle::GetBrush("SCSEditor.TreePanel"))
+			[
+				EntitiesTreeWidget.ToSharedRef()
+			]
+		]
+	];
+
+	UpdateTree();
+
+	// Make sure we know when tabs become active to update details tab
+	OnActiveTabChangedDelegateHandle = FGlobalTabmanager::Get()->OnActiveTabChanged_Subscribe(FOnActiveTabChanged::FDelegate::CreateSP(this, &SDMXEntityTreeViewBase::OnActiveTabChanged));
+}
+
+SDMXEntityTreeViewBase::~SDMXEntityTreeViewBase()
+{
+	FGlobalTabmanager::Get()->OnActiveTabChanged_Unsubscribe(OnActiveTabChangedDelegateHandle);
+	GEditor->UnregisterForUndo(this);
+}
+
+void SDMXEntityTreeViewBase::UpdateTree(bool bRegenerateTreeNodes /*= true*/)
+{
+	check(EntitiesTreeWidget.IsValid());
+
+	// the DMXEditor may have been closed, no need to update the tree
+	if (DMXEditor.IsValid())
+	{
+		if (bRegenerateTreeNodes)
+		{
+			// Obtain the set of expandable tree nodes that are currently collapsed
+			TSet<TSharedPtr<FDMXEntityTreeNodeBase>> CollapsedTreeNodes;
+			GetCollapsedNodes(CollapsedTreeNodes);
+
+			// Obtain the list of selected items
+			TArray<TSharedPtr<FDMXEntityTreeNodeBase>> SelectedTreeNodes = EntitiesTreeWidget->GetSelectedItems();
+
+			// Clear the current tree
+			if (SelectedTreeNodes.Num() != 0)
+			{
+				EntitiesTreeWidget->ClearSelection();
+			}
+
+			RootNode->ClearChildren();
+			RebuildNodes(RootNode);
+
+			// Restore the previous expansion state on the new tree nodes
+			TArray<TSharedPtr<FDMXEntityTreeNodeBase>> CollapsedTreeNodeArray = CollapsedTreeNodes.Array();
+			for (int32 i = 0; i < CollapsedTreeNodeArray.Num(); ++i)
+			{
+				// Look for a category match in the new hierarchy; if found, mark it as collapsed to match the previous setting
+				TSharedPtr<FDMXEntityTreeNodeBase> NodeToExpandPtr = FindNodeByName(CollapsedTreeNodeArray[i]->GetDisplayNameText());
+				if (NodeToExpandPtr.IsValid())
+				{
+					EntitiesTreeWidget->SetItemExpansion(NodeToExpandPtr, false);
+				}
+				else
+				{
+					EntitiesTreeWidget->SetItemExpansion(NodeToExpandPtr, true);
+				}
+			}
+
+			if (SelectedTreeNodes.Num() > 0)
+			{
+				// Restore the previous selection state on the new tree nodes
+				for (int32 i = 0; i < SelectedTreeNodes.Num(); ++i)
+				{
+					if (SelectedTreeNodes[i]->GetNodeType() == FDMXEntityTreeNodeBase::ENodeType::EntityNode)
+					{
+						TSharedPtr<FDMXEntityTreeEntityNode> EntityNode = StaticCastSharedPtr<FDMXEntityTreeEntityNode>(SelectedTreeNodes[i]);
+						TSharedPtr<FDMXEntityTreeEntityNode> NodeToSelectPtr = FindNodeByEntity(EntityNode->GetEntity());
+						if (NodeToSelectPtr.IsValid())
+						{
+							EntitiesTreeWidget->SetItemSelection(NodeToSelectPtr, true, ESelectInfo::Direct);
+						}
+					}
+				}
+			}
+		}
+
+		// Refresh the Tree Widget
+		EntitiesTreeWidget->RequestTreeRefresh();
+
+		// If no entity is selected, select first available one, if any
+		if (EntitiesTreeWidget->GetNumItemsSelected() == 0)
+		{
+			UDMXLibrary* Library = GetDMXLibrary();
+			check(Library != nullptr);
+
+			// Find the first non filtered out entity
+			for (UDMXEntity* Entity : Library->GetEntities())
+			{
+				if (TSharedPtr<FDMXEntityTreeNodeBase> EntityNode = FindNodeByEntity(Entity))
+				{
+					if (!EntityNode->IsFlaggedForFiltration())
+					{
+						EntitiesTreeWidget->SetSelection(EntityNode, ESelectInfo::OnMouseClick);
+						break;
+					}
+				}
+			}
+		}
+	}
+}
+
+TSharedPtr<FDMXEntityTreeCategoryNode> SDMXEntityTreeViewBase::FindCategoryNodeOfEntity(UDMXEntity* Entity) const
+{
+	TSharedPtr<FDMXEntityTreeEntityNode> EntityNode = FindNodeByEntity(Entity);
+	if (EntityNode.IsValid())
+	{
+		TSharedPtr<FDMXEntityTreeNodeBase> CategoryNode = EntityNode->GetParent().Pin();
+		if (CategoryNode.IsValid())
+		{
+			check(CategoryNode->GetNodeType() == FDMXEntityTreeNodeBase::ENodeType::CategoryNode);
+			return StaticCastSharedPtr<FDMXEntityTreeCategoryNode>(CategoryNode);
+		}
+	}
+
+	return nullptr;
+}
+
+TSharedPtr<FDMXEntityTreeEntityNode> SDMXEntityTreeViewBase::FindNodeByEntity(const UDMXEntity* Entity, TSharedPtr<FDMXEntityTreeNodeBase> StartNode /*= nullptr*/) const
+{
+	if (Entity)
+	{
+		// Start at root node if none was provided
+		if (!StartNode.IsValid())
+		{
+			StartNode = RootNode;
+		}
+
+		// Test the StartNode
+		if (StartNode.IsValid() && StartNode->GetNodeType() == FDMXEntityTreeNodeBase::ENodeType::EntityNode)
+		{
+			TSharedPtr<FDMXEntityTreeEntityNode> EntityNode = StaticCastSharedPtr<FDMXEntityTreeEntityNode>(StartNode);
+			if (EntityNode->GetEntity() == Entity)
+			{
+				return EntityNode;
+			}
+		}
+
+		// Test children recursive
+		for (const TSharedPtr<FDMXEntityTreeNodeBase>& ChildNode : StartNode->GetChildren())
+		{
+			TSharedPtr<FDMXEntityTreeEntityNode> EntityNode = FindNodeByEntity(Entity, ChildNode);
+			if (EntityNode.IsValid())
+			{
+				return EntityNode;
+			}
+		}
+	}
+
+	return nullptr;
+}
+
+TSharedPtr<FDMXEntityTreeNodeBase> SDMXEntityTreeViewBase::FindNodeByName(const FText& InName, TSharedPtr<FDMXEntityTreeNodeBase> StartNode /*= nullptr*/) const
+{
+	if (!InName.IsEmpty())
+	{
+		// Start at root node if none was provided
+		if (!StartNode.IsValid())
+		{
+			StartNode = RootNode;
+		}
+
+		// Test the StartNode
+		if (StartNode.IsValid())
+		{
+			// Check to see if the given entity matches the given tree node
+			if (StartNode->GetDisplayNameText().CompareTo(InName) == 0)
+			{
+				return StartNode;
+			}
+		}
+
+		// Test children recursive
+		for (const TSharedPtr<FDMXEntityTreeNodeBase>& ChildNode : StartNode->GetChildren())
+		{
+			TSharedPtr<FDMXEntityTreeNodeBase> Node = FindNodeByName(InName, ChildNode);
+			if (Node.IsValid())
+			{
+				return Node;
+			}
+		}
+	}
+
+	return nullptr;
+}
+
+void SDMXEntityTreeViewBase::SelectItemByNode(const TSharedRef<FDMXEntityTreeNodeBase>& Node, ESelectInfo::Type SelectInfo)
+{
+	// If Node is filtered out, we won't be able to select it
+	if (Node->IsFlaggedForFiltration())
+	{
+		FilterBox->SetText(FText::GetEmpty());
+	}
+
+	// Expand the parent nodes
+	for (TSharedPtr<FDMXEntityTreeNodeBase> ParentNode = Node->GetParent().Pin(); ParentNode.IsValid(); ParentNode = ParentNode->GetParent().Pin())
+	{
+		EntitiesTreeWidget->SetItemExpansion(ParentNode, true);
+	}
+
+	EntitiesTreeWidget->SetSelection(Node, SelectInfo);
+	EntitiesTreeWidget->RequestScrollIntoView(Node);
+	FSlateApplication::Get().SetKeyboardFocus(EntitiesTreeWidget, EFocusCause::SetDirectly);
+}
+
+void SDMXEntityTreeViewBase::SelectItemByEntity(const UDMXEntity* Entity, ESelectInfo::Type SelectInfo /*= ESelectInfo::Direct*/)
+{
+	// Check if the tree is being told to clear
+	if (!Entity)
+	{
+		EntitiesTreeWidget->ClearSelection();
+	}
+	else
+	{
+		const TSharedPtr<FDMXEntityTreeNodeBase>& ItemNode = FindNodeByEntity(Entity);
+		if (ItemNode.IsValid())
+		{
+			// If ItemNode is filtered out, we won't be able to select it
+			if (ItemNode->IsFlaggedForFiltration())
+			{
+				FilterBox->SetText(FText::GetEmpty());
+			}
+
+			// Expand the parent nodes
+			for (TSharedPtr<FDMXEntityTreeNodeBase> ParentNode = ItemNode->GetParent().Pin(); ParentNode.IsValid(); ParentNode = ParentNode->GetParent().Pin())
+			{
+				EntitiesTreeWidget->SetItemExpansion(ParentNode, true);
+			}
+
+			EntitiesTreeWidget->SetSelection(ItemNode, SelectInfo);
+			EntitiesTreeWidget->RequestScrollIntoView(ItemNode);
+			FSlateApplication::Get().SetKeyboardFocus(EntitiesTreeWidget, EFocusCause::SetDirectly);
+		}
+	}
+}
+
+void SDMXEntityTreeViewBase::SelectItemsByEntities(const TArray<UDMXEntity*>& InEntities, ESelectInfo::Type SelectInfo /*= ESelectInfo::Direct*/)
+{
+	EntitiesTreeWidget->ClearSelection();
+
+	if (InEntities.Num() > 0)
+	{
+		TSharedPtr<FDMXEntityTreeNodeBase> FirstNode;
+		TArray<TSharedPtr<FDMXEntityTreeNodeBase>> SelectedNodes;
+		for (UDMXEntity* Entity : InEntities)
+		{
+			if (Entity == nullptr)
+			{
+				continue;
+			}
+
+			// Find the Entity node for this Entity
+
+			if (TSharedPtr<FDMXEntityTreeNodeBase> EntityNode = FindNodeByEntity(Entity))
+			{
+				SelectedNodes.Add(EntityNode);
+
+				if (!FirstNode.IsValid())
+				{
+					FirstNode = EntityNode;
+				}
+			}
+		}
+		EntitiesTreeWidget->SetItemSelection(SelectedNodes, true, ESelectInfo::OnMouseClick);
+
+		// Scroll the first selected node into view
+		if (FirstNode.IsValid())
+		{
+			EntitiesTreeWidget->RequestScrollIntoView(FirstNode);
+		}
+
+		// Notify about the new selection
+		if (SelectInfo != ESelectInfo::Direct)
+		{
+			UpdateSelectionFromNodes(EntitiesTreeWidget->GetSelectedItems());
+		}
+
+		FSlateApplication::Get().SetKeyboardFocus(EntitiesTreeWidget, EFocusCause::SetDirectly);
+	}
+}
+
+void SDMXEntityTreeViewBase::SelectItemByName(const FString& ItemName, ESelectInfo::Type SelectInfo /*= ESelectInfo::Direct*/)
+{
+	// Check if the tree is being told to clear
+	if (ItemName.IsEmpty())
+	{
+		EntitiesTreeWidget->ClearSelection();
+	}
+	else
+	{
+		const TSharedPtr<FDMXEntityTreeNodeBase>& ItemNode = FindNodeByName(FText::FromString(ItemName));
+		if (ItemNode.IsValid())
+		{
+			// If ItemNode is filtered out, we won't be able to select it
+			if (ItemNode->IsFlaggedForFiltration())
+			{
+				FilterBox->SetText(FText::GetEmpty());
+			}
+
+			// Expand the parent nodes
+			for (TSharedPtr<FDMXEntityTreeNodeBase> ParentNode = ItemNode->GetParent().Pin(); ParentNode.IsValid(); ParentNode = ParentNode->GetParent().Pin())
+			{
+				EntitiesTreeWidget->SetItemExpansion(ParentNode, true);
+			}
+
+			EntitiesTreeWidget->SetSelection(ItemNode, SelectInfo);
+			EntitiesTreeWidget->RequestScrollIntoView(ItemNode);
+			FSlateApplication::Get().SetKeyboardFocus(EntitiesTreeWidget, EFocusCause::SetDirectly);
+		}
+	}
+}
+
+TArray<TSharedPtr<FDMXEntityTreeEntityNode>> SDMXEntityTreeViewBase::GetSelectedNodes() const
+{
+	TArray<TSharedPtr<FDMXEntityTreeEntityNode>> Result;
+	TArray<TSharedPtr<FDMXEntityTreeNodeBase>> SelectedItems = EntitiesTreeWidget->GetSelectedItems();
+	for (TSharedPtr<FDMXEntityTreeNodeBase> SelectedItem : SelectedItems)
+	{
+		if (SelectedItem->GetNodeType() == FDMXEntityTreeNodeBase::ENodeType::EntityNode)
+		{
+			Result.Add(StaticCastSharedPtr<FDMXEntityTreeEntityNode>(SelectedItem));
+		}
+	}
+
+	return Result;
+}
+
+TArray<UDMXEntity*> SDMXEntityTreeViewBase::GetSelectedEntities() const
+{
+	TArray<UDMXEntity*> SelectedEntities;
+
+	if (EntitiesTreeWidget.IsValid())
+	{
+		const TArray<TSharedPtr<FDMXEntityTreeNodeBase>>& SelectedItems = EntitiesTreeWidget->GetSelectedItems();
+		for (const TSharedPtr<FDMXEntityTreeNodeBase>& Node : SelectedItems)
+		{
+			if (Node.IsValid() && Node->GetNodeType() == FDMXEntityTreeNodeBase::ENodeType::EntityNode)
+			{
+				TSharedPtr<FDMXEntityTreeEntityNode> EntityNode = StaticCastSharedPtr<FDMXEntityTreeEntityNode>(Node);
+				SelectedEntities.Add(EntityNode->GetEntity());
+			}
+		}
+	}
+
+	return SelectedEntities;
+}
+
+FText SDMXEntityTreeViewBase::GetFilterText() const
+{
+	return FilterBox->GetText();
+}
+
+UDMXLibrary* SDMXEntityTreeViewBase::GetDMXLibrary() const
+{
+	if (DMXEditor.IsValid())
+	{
+		TSharedPtr<FDMXEditor> PinnedEditor(DMXEditor.Pin());
+		if (PinnedEditor.IsValid())
+		{
+			return PinnedEditor->GetDMXLibrary();
+		}
+	}
+	return nullptr;
+}
+
+FReply SDMXEntityTreeViewBase::OnEntitiesDragged(TSharedPtr<FDMXEntityTreeNodeBase> Node, const FPointerEvent& MouseEvent)
+{
+	if (Node.IsValid() && Node->GetNodeType() == FDMXEntityTreeNodeBase::ENodeType::EntityNode)
+	{
+		TArray<TSharedPtr<FDMXEntityTreeNodeBase>> SelectedItems = EntitiesTreeWidget->GetSelectedItems();
+		TArray<TWeakObjectPtr<UDMXEntity>> DraggedEntities;
+
+		for (TSharedPtr<FDMXEntityTreeNodeBase> SelectedItem : SelectedItems)
+		{
+			if (ensureMsgf(SelectedItem->GetNodeType() == FDMXEntityTreeNodeBase::ENodeType::EntityNode, TEXT("Unexpected non-entity node is selected")))
+			{
+				TSharedPtr<FDMXEntityTreeEntityNode> EntityNode = StaticCastSharedPtr<FDMXEntityTreeEntityNode>(SelectedItem);
+				DraggedEntities.Add(EntityNode->GetEntity());
+			}
+		}
+
+		// If no entities are selected, use the dragged entity instead
+		if (DraggedEntities.Num() == 0)
+		{
+			TSharedPtr<FDMXEntityTreeEntityNode> EntityNode = StaticCastSharedPtr<FDMXEntityTreeEntityNode>(Node);
+			DraggedEntities.Add(EntityNode->GetEntity());
+		}
+
+		TSharedRef<FDMXEntityDragDropOperation> DragOperation = MakeShared<FDMXEntityDragDropOperation>(GetDMXLibrary(), DraggedEntities);
+		return FReply::Handled().BeginDragDrop(DragOperation);
+	}
+
+	return FReply::Unhandled();
+}
+
+FReply SDMXEntityTreeViewBase::OnKeyDown(const FGeometry& MyGeometry, const FKeyEvent& InKeyEvent)
+{
+	if (CommandList->ProcessCommandBindings(InKeyEvent))
+	{
+		return FReply::Handled();
+	}
+
+	return FReply::Unhandled();
+}
+
+void SDMXEntityTreeViewBase::PostUndo(bool bSuccess)
+{
+	if (UDMXLibrary* Library = GetDMXLibrary())
+	{
+		Library->Modify();
+	}
+
+	UpdateTree();
+}
+
+void SDMXEntityTreeViewBase::PostRedo(bool bSuccess)
+{
+	if (UDMXLibrary* Library = GetDMXLibrary())
+	{
+		Library->Modify();
+	}
+
+	UpdateTree();
+}
+
+void SDMXEntityTreeViewBase::OnGetChildren(TSharedPtr<FDMXEntityTreeNodeBase> InNode, TArray<TSharedPtr<FDMXEntityTreeNodeBase>>& OutChildren)
+{
+	if (InNode.IsValid())
+	{
+		const TArray<TSharedPtr<FDMXEntityTreeNodeBase>>& Children = InNode->GetChildren();
+
+		if (!GetFilterText().IsEmpty())
+		{
+			OutChildren.Reserve(Children.Num());
+
+			for (const TSharedPtr<FDMXEntityTreeNodeBase>& Child : Children)
+			{
+				if (!Child->IsFlaggedForFiltration())
+				{
+					OutChildren.Add(Child);
+				}
+			}
+		}
+		else
+		{
+			OutChildren = Children;
+		}
+	}
+	else
+	{
+		OutChildren.Empty();
+	}
+}
+
+void SDMXEntityTreeViewBase::OnExpansionChanged(TSharedPtr<FDMXEntityTreeNodeBase> Node, bool bInExpansionState)
+{
+	// Only applies when there's no filtering
+	if (Node.IsValid() && GetFilterText().IsEmpty())
+	{
+		Node->SetExpansionState(bInExpansionState);
+	}
+}
+
+void SDMXEntityTreeViewBase::OnSelectionChanged(TSharedPtr<FDMXEntityTreeNodeBase> InSelectedNodePtr, ESelectInfo::Type SelectInfo)
+{
+	if (SelectInfo != ESelectInfo::Direct)
+	{
+		UpdateSelectionFromNodes(EntitiesTreeWidget->GetSelectedItems());
+	}
+}
+
+bool SDMXEntityTreeViewBase::RefreshFilteredState(TSharedPtr<FDMXEntityTreeNodeBase> Node, bool bRecursive)
+{
+	FString FilterText = FText::TrimPrecedingAndTrailing(GetFilterText()).ToString();
+	TArray<FString> FilterTerms;
+	FilterText.ParseIntoArray(FilterTerms, TEXT(" "), /*CullEmpty =*/true);
+
+	struct RefreshFilteredState_Inner
+	{
+		static void RefreshFilteredState(TSharedPtr<FDMXEntityTreeNodeBase> InNode, const TArray<FString>& InFilterTerms, bool bInRecursive)
+		{
+			if (bInRecursive)
+			{
+				for (TSharedPtr<FDMXEntityTreeNodeBase> Child : InNode->GetChildren())
+				{
+					RefreshFilteredState(Child, InFilterTerms, bInRecursive);
+				}
+			}
+
+			FString DisplayStr = InNode->GetDisplayNameString();
+
+			bool bIsFilteredOut = false;
+			for (const FString& FilterTerm : InFilterTerms)
+			{
+				if (!DisplayStr.Contains(FilterTerm))
+				{
+					bIsFilteredOut = true;
+				}
+			}
+			// if we're not recursing, then assume this is for a new node and we need to update the parent
+			// otherwise, assume the parent was hit as part of the recursion
+			InNode->UpdateCachedFilterState(!bIsFilteredOut, /*bUpdateParent =*/!bInRecursive);
+		}
+	};
+
+	RefreshFilteredState_Inner::RefreshFilteredState(Node, FilterTerms, bRecursive);
+	return Node->IsFlaggedForFiltration();
+}
+
+void SDMXEntityTreeViewBase::UpdateSelectionFromNodes(const TArray<TSharedPtr<FDMXEntityTreeNodeBase>>& SelectedNodes)
+{
+	bUpdatingSelection = true;
+
+	// Notify that the selection has updated
+	OnSelectionChangedDelegate.ExecuteIfBound(GetSelectedEntities());
+
+	bUpdatingSelection = false;
+}
+
+void SDMXEntityTreeViewBase::OnFilterTextChanged(const FText& InFilterText)
+{
+	for (TSharedPtr<FDMXEntityTreeNodeBase> Node : RootNode->GetChildren())
+	{
+		RefreshFilteredState(Node, true);
+	}
+
+	// Clears selection to make UpdateTree automatically select the first visible node
+	EntitiesTreeWidget->ClearSelection();
+
+	UpdateTree(/*bRegenerateTreeNodes =*/false);
+
+	// If we reset the filter, recover nodes expansion states
+	UpdateNodesExpansion(RootNode.ToSharedRef(), GetFilterText().IsEmpty());
+}
+
+void SDMXEntityTreeViewBase::GetCollapsedNodes(TSet<TSharedPtr<FDMXEntityTreeNodeBase>>& OutCollapsedNodes, TSharedPtr<FDMXEntityTreeNodeBase> InParentNodePtr /*= nullptr*/) const
+{
+	if (!InParentNodePtr.IsValid())
+	{
+		InParentNodePtr = RootNode;
+	}
+
+	for (const TSharedPtr<FDMXEntityTreeNodeBase>& Node : InParentNodePtr->GetChildren())
+	{
+		if (Node->GetChildren().Num() > 0)
+		{
+			if (!EntitiesTreeWidget->IsItemExpanded(Node))
+			{
+				OutCollapsedNodes.Add(Node);
+			}
+			else // not collapsed. Check children
+			{
+				GetCollapsedNodes(OutCollapsedNodes, Node);
+			}
+		}
+	}
+}
+
+void SDMXEntityTreeViewBase::SetNodeExpansionState(TSharedPtr<FDMXEntityTreeNodeBase> InNodeToChange, const bool bIsExpanded)
+{
+	if (EntitiesTreeWidget.IsValid() && InNodeToChange.IsValid())
+	{
+		EntitiesTreeWidget->SetItemExpansion(InNodeToChange, bIsExpanded);
+	}
+}
+
+void SDMXEntityTreeViewBase::UpdateNodesExpansion(TSharedRef<FDMXEntityTreeNodeBase> InRootNode, bool bFilterIsEmpty)
+{
+	// Only category nodes have children and need expansion
+	if (InRootNode->GetNodeType() != FDMXEntityTreeNodeBase::ENodeType::EntityNode)
+	{
+		// If the filter is not empty, all nodes should be expanded
+		const bool bExpandNodes = !bFilterIsEmpty || InRootNode->GetExpansionState();
+		EntitiesTreeWidget->SetItemExpansion(InRootNode, bExpandNodes);
+
+		for (const TSharedPtr<FDMXEntityTreeNodeBase>& Child : InRootNode->GetChildren())
+		{
+			if (Child.IsValid() && Child->GetNodeType() != FDMXEntityTreeNodeBase::ENodeType::EntityNode)
+			{
+				UpdateNodesExpansion(Child.ToSharedRef(), bFilterIsEmpty);
+			}
+		}
+	}
+}
+
+void SDMXEntityTreeViewBase::OnActiveTabChanged(TSharedPtr<SDockTab> PreviouslyActive, TSharedPtr<SDockTab> NewlyActivated)
+{
+	if (IsInTab(NewlyActivated))
+	{
+		UpdateTree();
+
+		// Refresh selected entities' properties on the inspector panel by issuing a selection update.
+		// Some properties might have been changed on a previously selected tab.
+		UpdateSelectionFromNodes(EntitiesTreeWidget->GetSelectedItems());
+	}
+}
+
+bool SDMXEntityTreeViewBase::IsInTab(TSharedPtr<SDockTab> InDockTab) const
+{
+	// Too many hierarchy levels to do it with a recursive function. Crashes with Stack Overflow.
+	// Using loop instead.
+
+	if (InDockTab.IsValid())
+	{
+		// Tab content that should be a parent of this widget on some level
+		TSharedPtr<SWidget> TabContent = InDockTab->GetContent();
+		// Current parent being checked against
+		TSharedPtr<SWidget> CurrentParent = GetParentWidget();
+
+		while (CurrentParent.IsValid())
+		{
+			if (CurrentParent == TabContent)
+			{
+				return true;
+			}
+			CurrentParent = CurrentParent->GetParentWidget();
+		}
+
+		// reached top widget (parent is invalid) and none was the tab
+		return false;
+	}
+
+	return false;
+}
+
+#undef LOCTEXT_NAMESPACE
diff --git a/Engine/Plugins/VirtualProduction/DMX/DMXEngine/Source/DMXEditor/Private/Widgets/SDMXEntityTreeViewBase.h b/Engine/Plugins/VirtualProduction/DMX/DMXEngine/Source/DMXEditor/Private/Widgets/SDMXEntityTreeViewBase.h
new file mode 100644
--- /dev/null
+++ b/Engine/Plugins/VirtualProduction/DMX/DMXEngine/Source/DMXEditor/Private/Widgets/SDMXEntityTreeViewBase.h
@@ -0,0 +1,266 @@
+// Copyright Epic Games, Inc. All Rights Reserved.
+
+#pragma once
+
+#include "Widgets/DMXEntityTreeNode.h"
+
+#include "CoreMinimal.h"
+#include "EditorUndoClient.h"
+#include "Widgets/SCompoundWidget.h"
+
+class FDMXEditor;
+class UDMXLibrary;
+
+class FMenuBuilder;
+class FUICommandList;
+class ITableRow;
+class SComboButton;
+class SDockTab;
+class SSearchBox;
+class STableViewBase;
+template<typename ItemType> class STreeView;
+
+
+/** Base class for a tree view of entities, using DMXEntityTreeNode as its Nodes. */
+class SDMXEntityTreeViewBase
+	: public SCompoundWidget
+	, public FEditorUndoClient
+{
+public:
+	DECLARE_DELEGATE_OneParam(FDMXOnSelectionChanged, const TArray<UDMXEntity*>&);
+
+	SLATE_BEGIN_ARGS(SDMXEntityTreeViewBase)
+	{}
+		/** The DMX Editor that owns this widget */
+		SLATE_ARGUMENT(TWeakPtr<FDMXEditor>, DMXEditor)
+
+		/** Exectued when the list changed its selection */
+		SLATE_EVENT(FDMXOnSelectionChanged, OnSelectionChanged)
+
+		/** Exectued when entites were added to the DMXEditor's library */
+		SLATE_EVENT(FSimpleDelegate, OnEntitiesAdded)
+
+		/** Exectued when entites were reorderd in the list, and potentially in the library */
+		SLATE_EVENT(FSimpleDelegate, OnEntityOrderChanged)
+
+		/** Exectued when entites were removed from the DMXEditor's library */
+		SLATE_EVENT(FSimpleDelegate, OnEntitiesRemoved)
+
+	SLATE_END_ARGS()
+
+	/** Constructs the widget */
+	void Construct(const FArguments& InArgs);
+
+	/** Destructor */
+	virtual ~SDMXEntityTreeViewBase();
+
+	/** Refresh the tree control to reflect changes in the editor */
+	void UpdateTree(bool bRegenerateTreeNodes = true);
+
+	/**
+	 * Creates a new category node directly under the passed parent or just retrieves it if existent.
+	 * If InParentNodePtr is null, a root category node is created/retrieved.
+	 * Stores a value, so when something is dropped on the category the value can be retrieved.
+	 */
+	template <typename ValueType>
+	TSharedRef<FDMXEntityTreeCategoryNode> GetOrCreateCategoryNode(const FDMXEntityTreeCategoryNode::ECategoryType InCategoryType, const FText InCategoryName, ValueType Value, TSharedPtr<FDMXEntityTreeNodeBase> InParentNode = nullptr, const FText& InToolTip = FText::GetEmpty())
+	{
+		for (const TSharedPtr<FDMXEntityTreeNodeBase>& Node : InParentNode.IsValid() ? InParentNode->GetChildren() : RootNode->GetChildren())
+		{
+			if (Node.IsValid() && Node->GetNodeType() == FDMXEntityTreeNodeBase::ENodeType::CategoryNode)
+			{
+				TSharedPtr<FDMXEntityTreeCategoryNode> ExistingCategoryNode = StaticCastSharedPtr<FDMXEntityTreeCategoryNode>(Node);
+				if (ExistingCategoryNode->GetCategoryType() == InCategoryType && ExistingCategoryNode->GetDisplayNameText().CompareTo(InCategoryName) == 0)
+				{
+					EntitiesTreeWidget->SetItemExpansion(Node, true);
+					return ExistingCategoryNode.ToSharedRef();
+				}
+			}
+		}
+
+		// Didn't find an existing node. Add one.
+		TSharedRef<FDMXEntityTreeCategoryNode> NewNode = MakeShared<FDMXEntityTreeCategoryNode>(InCategoryType, InCategoryName, Value, InToolTip);
+		if (InParentNode.IsValid())
+		{
+			InParentNode->AddChild(NewNode);
+		}
+		else
+		{
+			RootNode->AddChild(NewNode);
+		}
+
+		constexpr bool bRefreshFilteredStateRecursive = false;
+		RefreshFilteredState(NewNode, bRefreshFilteredStateRecursive);
+
+		constexpr bool bExpandItem = true;
+		EntitiesTreeWidget->SetItemExpansion(NewNode, bExpandItem);
+		NewNode->SetExpansionState(bExpandItem);
+
+		return NewNode;
+	}
+
+	/** Returns the category node of the entity, or null if is not in the list */
+	TSharedPtr<FDMXEntityTreeCategoryNode> FindCategoryNodeOfEntity(UDMXEntity* Entity) const;
+
+	/** Helper method to recursively find a tree node for the given DMX Entity starting at the given tree node */
+	TSharedPtr<FDMXEntityTreeEntityNode> FindNodeByEntity(const UDMXEntity* Entity, TSharedPtr<FDMXEntityTreeNodeBase> StartNode = nullptr) const;
+
+	/** Helper method to recursively find a tree node with the given name starting at the given tree node */
+	TSharedPtr<FDMXEntityTreeNodeBase> FindNodeByName(const FText& Name, TSharedPtr<FDMXEntityTreeNodeBase> StartNode = nullptr) const;
+
+	/** Selects an item */
+	void SelectItemByNode(const TSharedRef<FDMXEntityTreeNodeBase>& Node, ESelectInfo::Type SelectInfo = ESelectInfo::Direct);
+
+	/** Selects an item by Entity */
+	void SelectItemByEntity(const UDMXEntity* InEntity, ESelectInfo::Type SelectInfo = ESelectInfo::Direct);
+
+	/** Selects items by an Array of Entites */
+	void SelectItemsByEntities(const TArray<UDMXEntity*>& InEntities, ESelectInfo::Type SelectInfo = ESelectInfo::Direct);
+
+	/** Selects an item by name */
+	void SelectItemByName(const FString& ItemName, ESelectInfo::Type SelectInfo = ESelectInfo::Direct);
+
+	/** Gets selected nodes */
+	TArray<TSharedPtr<FDMXEntityTreeEntityNode>> GetSelectedNodes() const;
+
+	/** Get only the valid selected entities */
+	TArray<UDMXEntity*> GetSelectedEntities() const;
+
+	/** Gets current filter from the FilterBox */
+	FText GetFilterText() const;
+
+	/** Gets the DMX Library object being edited */
+	UDMXLibrary* GetDMXLibrary() const;
+
+	/** Handler for when an entity from the list is dragged */
+	virtual FReply OnEntitiesDragged(TSharedPtr<FDMXEntityTreeNodeBase> InNodePtr, const FPointerEvent& MouseEvent);
+
+protected:
+	//~ Begin SWidget interface
+	virtual FReply OnKeyDown(const FGeometry& MyGeometry, const FKeyEvent& InKeyEvent);
+	//~ End SWidget Interface
+
+	//~ Begin FEditorUndoClient Interface
+	virtual void PostUndo(bool bSuccess) override;
+	virtual void PostRedo(bool bSuccess) override;
+	// End of FEditorUndoClient
+	
+	/** Generates the Add New Entity Button that is displayed above the Tree View */
+	virtual TSharedRef<SWidget> GenerateAddNewEntityButton() = 0;
+
+	/** Rebuils of the root node */
+	virtual void RebuildNodes(const TSharedPtr<FDMXEntityTreeRootNode>& RootNode) = 0;
+
+	/** Generates a row in the tree view */
+	virtual TSharedRef<ITableRow> OnGenerateRow(TSharedPtr<FDMXEntityTreeNodeBase> Node, const TSharedRef<STableViewBase>& OwnerTable) = 0;
+
+	/** Gets the children of specified Node */
+	virtual void OnGetChildren(TSharedPtr<FDMXEntityTreeNodeBase> InNode, TArray<TSharedPtr<FDMXEntityTreeNodeBase>>& OutChildren);
+
+	/** Called when the expansion of a Node changed */
+	virtual void OnExpansionChanged(TSharedPtr<FDMXEntityTreeNodeBase> Node, bool bInExpansionState);
+
+	/** Called when selection in the Tree changed */
+	virtual void OnSelectionChanged(TSharedPtr<FDMXEntityTreeNodeBase> Node, ESelectInfo::Type SelectInfo);
+
+	/** Called to display a context menu when right clicking an Entity */
+	virtual TSharedPtr<SWidget> OnContextMenuOpen() = 0;
+
+	///** Called when a TreeNode is scrolled into view */
+	//virtual void OnItemScrolledIntoView(TSharedPtr<FDMXEntityTreeNodeBase> Node, const TSharedPtr<ITableRow>& InWidget);
+
+	/** Cut selected node(s) */
+	virtual void OnCutSelectedNodes() = 0;
+	virtual bool CanCutNodes() const = 0;
+
+	/** Copy selected node(s) */
+	virtual void OnCopySelectedNodes() = 0;
+	virtual bool CanCopyNodes() const = 0;
+
+	/** Pastes previously copied node(s) */
+	virtual void OnPasteNodes() = 0;
+	virtual bool CanPasteNodes() const = 0;
+
+	/** Callbacks to duplicate the selected component */
+	virtual bool CanDuplicateNodes() const = 0;
+	virtual void OnDuplicateNodes() = 0;
+
+	/** Removes existing selected component nodes from the SCS */
+	virtual void OnDeleteNodes() = 0;
+	virtual bool CanDeleteNodes() const = 0;
+
+	/** Requests a rename on the selected Entity. */
+	virtual void OnRenameNode() = 0;
+	virtual bool CanRenameNode() const = 0;
+
+	/**
+	 * Compares the filter bar's text with the item's component name. Use
+	 * bRecursive to refresh the state of child nodes as well. Returns true if
+	 * the node is set to be filtered out
+	 */
+	bool RefreshFilteredState(TSharedPtr<FDMXEntityTreeNodeBase> Node, bool bRecursive);
+
+	/** The actual Tree widget */
+	TSharedPtr<STreeView<TSharedPtr<FDMXEntityTreeNodeBase>>> EntitiesTreeWidget;
+
+	/** Broadcast whent he selection updated. */
+	FDMXOnSelectionChanged OnSelectionChangedDelegate;
+
+	/** Broadcast when the entity list added an entity to the library */
+	FSimpleDelegate OnEntitiesAdded;
+
+	/** Broadcast when the entity list changed order of the library's entity array */
+	FSimpleDelegate OnEntityOrderChanged;
+
+	/** Broadcast when the entity list deleted an entity from the library */
+	FSimpleDelegate OnEntitiesRemoved;
+
+	/** Command list for handling actions such as copy and paste */
+	TSharedPtr<FUICommandList> CommandList;
+
+	/** Pointer back to the DMXEditor tool that owns us */
+	TWeakPtr<FDMXEditor> DMXEditor;
+
+private:
+	/** Update any associated selection from the passed in nodes */
+	void UpdateSelectionFromNodes(const TArray<TSharedPtr<FDMXEntityTreeNodeBase>>& SelectedNodes);
+
+	/** Callback when the filter is changed, forces the action tree(s) to filter */
+	void OnFilterTextChanged(const FText& InFilterText);
+
+	/** Returns the set of expandable nodes that are currently collapsed in the UI */
+	void GetCollapsedNodes(TSet<TSharedPtr<FDMXEntityTreeNodeBase>>& OutCollapsedNodes, TSharedPtr<FDMXEntityTreeNodeBase> InParentNode = nullptr) const;
+
+	/**
+	 * Set the expansion state of a node
+	 *
+	 * @param InNodeToChange	The node to be expanded/collapsed
+	 * @param bIsExpanded		True to expand the node, false to collapse it
+	 */
+	void SetNodeExpansionState(TSharedPtr<FDMXEntityTreeNodeBase> InNodeToChange, const bool bIsExpanded);
+
+	/** Expand all categories during filtering and resets nodes expansion state after filtering is cleared */
+	void UpdateNodesExpansion(TSharedRef<FDMXEntityTreeNodeBase> InRootNode, bool bFilterIsEmpty);
+
+	/** Called when the active tab in the editor changes */
+	void OnActiveTabChanged(TSharedPtr<SDockTab> PreviouslyActive, TSharedPtr<SDockTab> NewlyActivated);
+
+	/** Searches this widget's parents to see if it's a child of InDockTab */
+	bool IsInTab(TSharedPtr<SDockTab> InDockTab) const;
+
+	/**
+	 * The root node.
+	 * It's not added to the tree, but the main categories and all their children
+	 * (entity and sub-category nodes) belong to it to make recursive searching algorithms nicer.
+	 */
+	TSharedPtr<FDMXEntityTreeRootNode> RootNode;
+
+	/** The filter box that handles filtering entities */
+	TSharedPtr<SSearchBox> FilterBox;
+
+	/** Gate to prevent changing the selection while selection change is being broadcast. */
+	bool bUpdatingSelection = false;
+
+	/** Delegate handle bound to the FGlobalTabmanager::OnActiveTabChanged delegate */
+	FDelegateHandle OnActiveTabChangedDelegateHandle;
+};
diff --git a/Engine/Plugins/VirtualProduction/DMX/DMXEngine/Source/DMXRuntime/Private/Library/DMXEntityFixturePatch.cpp b/Engine/Plugins/VirtualProduction/DMX/DMXEngine/Source/DMXRuntime/Private/Library/DMXEntityFixturePatch.cpp
--- a/Engine/Plugins/VirtualProduction/DMX/DMXEngine/Source/DMXRuntime/Private/Library/DMXEntityFixturePatch.cpp
+++ b/Engine/Plugins/VirtualProduction/DMX/DMXEngine/Source/DMXRuntime/Private/Library/DMXEntityFixturePatch.cpp
@@ -227,9 +227,9 @@ bool UDMXEntityFixturePatch::IsValidEntity(FText& OutReason) const
 {
 	OutReason = FText::GetEmpty();
 
-	if (ParentFixtureTypeTemplate == nullptr)
+	if (!ParentFixtureTypeTemplate)
 	{
-		OutReason = LOCTEXT("InvalidReason_NullParentTemplate", "Fixture Template is null");
+		OutReason = LOCTEXT("InvalidReason_NullParentTemplate", "Fixture Type is null");
 	}
 	else if ((GetStartingChannel() + GetChannelSpan() - 1) > DMX_UNIVERSE_SIZE)
 	{
@@ -246,28 +246,27 @@ bool UDMXEntityFixturePatch::IsValidEntity(FText& OutReason) const
 			FText::FromString(ParentFixtureTypeTemplate->GetDisplayName())
 		);
 	}	
-	else
+	else if (!ParentFixtureTypeTemplate->Modes.IsValidIndex(ActiveMode))
 	{
-		int32 IdxExistingFunction = ParentFixtureTypeTemplate->Modes.IndexOfByPredicate([](const FDMXFixtureMode& Mode) {
-			return 
-				Mode.Functions.Num() > 0 ||
-				Mode.FixtureMatrixConfig.CellAttributes.Num() > 0;
-			});
-
-		if (IdxExistingFunction == INDEX_NONE)
-		{
-			OutReason = FText::Format(
-				LOCTEXT("InvalidReason_NoFunctionsDefined", "'{0}' cannot be assigned as its parent Fixture Type '{1}' does not define any Functions."),
-				FText::FromString(GetDisplayName()),
-				FText::FromString(ParentFixtureTypeTemplate->GetDisplayName()));
-		}
-		else if (GetChannelSpan() == 0)
-		{
-			OutReason = FText::Format(
-				LOCTEXT("InvalidReason_ParentChannelSpanIsZero", "'{0}' cannot be assigned as its parent Fixture Type '{1}' overrides channel span with 0."),
-				FText::FromString(GetDisplayName()),
-				FText::FromString(ParentFixtureTypeTemplate->GetDisplayName()));
-		}
+		OutReason = FText::Format(
+			LOCTEXT("InvalidReason_NoModesDefined", "'{0}' cannot be assigned as its parent Fixture Type '{1}' does no longer contain Active Mode."),
+			FText::FromString(GetDisplayName()),
+			FText::FromString(ParentFixtureTypeTemplate->GetDisplayName())
+		);
+	}
+	else if (ParentFixtureTypeTemplate->Modes[ActiveMode].Functions.Num() == 0)
+	{
+		OutReason = FText::Format(
+			LOCTEXT("InvalidReason_NoFunctionsDefined", "'{0}' cannot be assigned as its parent Fixture Type '{1}' does not define any Functions in the Active Mode."),
+			FText::FromString(GetDisplayName()),
+			FText::FromString(ParentFixtureTypeTemplate->GetDisplayName()));
+	}
+	else if (GetChannelSpan() == 0)
+	{
+		OutReason = FText::Format(
+			LOCTEXT("InvalidReason_ParentChannelSpanIsZero", "'{0}' cannot be assigned as its parent Fixture Type '{1}' overrides channel span with 0."),
+			FText::FromString(GetDisplayName()),
+			FText::FromString(ParentFixtureTypeTemplate->GetDisplayName()));
 	}
 
 	return OutReason.IsEmpty();
