DMX - Fix CIS non-unity compile

diff --git a/Engine/Plugins/VirtualProduction/DMX/DMXEngine/Source/DMXEditor/Private/Widgets/SDMXEntityTreeViewBase.cpp b/Engine/Plugins/VirtualProduction/DMX/DMXEngine/Source/DMXEditor/Private/Widgets/SDMXEntityTreeViewBase.cpp
--- a/Engine/Plugins/VirtualProduction/DMX/DMXEngine/Source/DMXEditor/Private/Widgets/SDMXEntityTreeViewBase.cpp
+++ b/Engine/Plugins/VirtualProduction/DMX/DMXEngine/Source/DMXEditor/Private/Widgets/SDMXEntityTreeViewBase.cpp
@@ -13,7 +13,6 @@
 #include "Widgets/Docking/SDockTab.h"
 #include "Widgets/Input/SSearchBox.h"
 #include "Widgets/Images/SImage.h"
-#include "Widgets/Views/STreeView.h"
 
 
 #define LOCTEXT_NAMESPACE "SDMXEntityTreeViewBase"
diff --git a/Engine/Plugins/VirtualProduction/DMX/DMXEngine/Source/DMXEditor/Private/Widgets/SDMXEntityTreeViewBase.h b/Engine/Plugins/VirtualProduction/DMX/DMXEngine/Source/DMXEditor/Private/Widgets/SDMXEntityTreeViewBase.h
--- a/Engine/Plugins/VirtualProduction/DMX/DMXEngine/Source/DMXEditor/Private/Widgets/SDMXEntityTreeViewBase.h
+++ b/Engine/Plugins/VirtualProduction/DMX/DMXEngine/Source/DMXEditor/Private/Widgets/SDMXEntityTreeViewBase.h
@@ -7,6 +7,7 @@
 #include "CoreMinimal.h"
 #include "EditorUndoClient.h"
 #include "Widgets/SCompoundWidget.h"
+#include "Widgets/Views/STreeView.h"
 
 class FDMXEditor;
 class UDMXLibrary;
