DMX - Refactor DMXEditorPropertyEditorCustomization classes into less dependent ones in separate files

diff --git a/Engine/Plugins/VirtualProduction/DMX/DMXEngine/Source/DMXEditor/Private/Customizations/DMXEditorPropertyEditorCustomization.cpp b/Engine/Plugins/VirtualProduction/DMX/DMXEngine/Source/DMXEditor/Private/Customizations/DMXEditorPropertyEditorCustomization.cpp
deleted file mode 100644
diff --git a/Engine/Plugins/VirtualProduction/DMX/DMXEngine/Source/DMXEditor/Private/Customizations/DMXEditorPropertyEditorCustomization.h b/Engine/Plugins/VirtualProduction/DMX/DMXEngine/Source/DMXEditor/Private/Customizations/DMXEditorPropertyEditorCustomization.h
deleted file mode 100644
diff --git a/Engine/Plugins/VirtualProduction/DMX/DMXEngine/Source/DMXEditor/Private/Customizations/DMXEntityFixturePatchDetails.cpp b/Engine/Plugins/VirtualProduction/DMX/DMXEngine/Source/DMXEditor/Private/Customizations/DMXEntityFixturePatchDetails.cpp
new file mode 100644
--- /dev/null
+++ b/Engine/Plugins/VirtualProduction/DMX/DMXEngine/Source/DMXEditor/Private/Customizations/DMXEntityFixturePatchDetails.cpp
@@ -0,0 +1,266 @@
+// Copyright Epic Games, Inc. All Rights Reserved.
+
+#include "Customizations/DMXEntityFixturePatchDetails.h"
+
+#include "DMXEditorUtils.h"
+#include "Library/DMXEntityFixturePatch.h"
+#include "Widgets/SDMXEntityDropdownMenu.h"
+
+#include "DetailLayoutBuilder.h"
+#include "DetailWidgetRow.h"
+#include "IPropertyUtilities.h"
+
+
+#define LOCTEXT_NAMESPACE "DMXEntityFixturePatchFixtureSettingsDetails"
+
+FDMXEntityFixturePatchDetails::FDMXEntityFixturePatchDetails(TWeakPtr<FDMXEditor> InDMXEditorPtr)
+	: DMXEditorPtr(InDMXEditorPtr)
+{}
+
+TSharedRef<IDetailCustomization> FDMXEntityFixturePatchDetails::MakeInstance(TWeakPtr<FDMXEditor> InDMXEditorPtr)
+{
+	return MakeShared<FDMXEntityFixturePatchDetails>(InDMXEditorPtr);
+}
+
+void FDMXEntityFixturePatchDetails::CustomizeDetails(IDetailLayoutBuilder& DetailBuilder)
+{
+	PropertyUtilities = DetailBuilder.GetPropertyUtilities();
+
+	AutoAssignAddressHandle = DetailBuilder.GetProperty(UDMXEntityFixturePatch::GetAutoAssignAddressPropertyNameChecked());
+	ParentFixtureTypeHandle = DetailBuilder.GetProperty(UDMXEntityFixturePatch::GetParentFixtureTypeTemplatePropertyNameChecked());
+	ActiveModeHandle = DetailBuilder.GetProperty(UDMXEntityFixturePatch::GetActiveModePropertyNameChecked());
+
+	// Bind to auto assign address changes to assign channels when it gets enabled
+	FSimpleDelegate OnAutoAssignAddressChangedDelegate = FSimpleDelegate::CreateSP(this, &FDMXEntityFixturePatchDetails::OnAutoAssignAddressChanged);
+	AutoAssignAddressHandle->SetOnPropertyValueChanged(OnAutoAssignAddressChangedDelegate);
+
+	// Handle mode changes of the parent fixture type
+	UDMXEntityFixtureType::GetDataTypeChangeDelegate().AddSP(this, &FDMXEntityFixturePatchDetails::OnModeChanged);
+
+	// Make a Fixture Types dropdown for the Fixture Type template property
+	DetailBuilder.EditDefaultProperty(ParentFixtureTypeHandle)->CustomWidget(false)
+		.NameContent()
+		[
+			ParentFixtureTypeHandle->CreatePropertyNameWidget()
+		]
+		.ValueContent()
+		.MinDesiredWidth(200.0f)
+		.MaxDesiredWidth(400.0f)
+		[
+			SNew(SDMXEntityPickerButton<UDMXEntityFixtureType>)
+			.DMXEditor(DMXEditorPtr)
+			.CurrentEntity(this, &FDMXEntityFixturePatchDetails::GetParentFixtureType)
+			.OnEntitySelected(this, &FDMXEntityFixturePatchDetails::OnParentFixtureTypeChanged)
+			.HasMultipleValues(this, &FDMXEntityFixturePatchDetails::IsParentFixtureTypeMultipleValues)
+		];
+
+	// Make a modes dropdown to select the active Fixture Type Mode, if a valid Fixture Type is selected
+	TSharedPtr<uint32> DefaultSelectedActiveMode = nullptr;
+	GenerateActiveModesSource();
+
+	int32 ActiveMode;
+	if (ensure(ActiveModeHandle->GetValue(ActiveMode) == FPropertyAccess::Success))
+	{
+		const bool bActiveModeExists = ActiveModesSource.ContainsByPredicate([ActiveMode](TSharedPtr<uint32> Option) 
+			{
+				return Option.IsValid() && *Option == ActiveMode;
+			});
+
+		if (!bActiveModeExists)
+		{
+			SetActiveMode(0);
+		}
+	}
+
+	DetailBuilder.EditDefaultProperty(ActiveModeHandle)->CustomWidget(false)
+		.NameContent()
+		[
+			ActiveModeHandle->CreatePropertyNameWidget()
+		]
+		.ValueContent()
+		.MaxDesiredWidth(160.0f)
+		[
+			SAssignNew(ActiveModeComboBox, SComboBox<TSharedPtr<uint32>>)
+			.IsEnabled(this, &FDMXEntityFixturePatchDetails::IsActiveModeEditable)
+			.OptionsSource(&ActiveModesSource)
+			.OnGenerateWidget(this, &FDMXEntityFixturePatchDetails::GenerateActiveModeWidget)
+			.OnSelectionChanged(this, &FDMXEntityFixturePatchDetails::OnActiveModeChanged)
+			.InitiallySelectedItem(DefaultSelectedActiveMode)
+			[
+				SNew(STextBlock)
+				.MinDesiredWidth(50.0f)
+				.Text(this, &FDMXEntityFixturePatchDetails::GetCurrentActiveModeLabel)
+				.Font(DetailBuilder.GetDetailFont())
+			]
+		];
+}
+
+TSharedRef<SWidget> FDMXEntityFixturePatchDetails::GenerateActiveModeWidget(const TSharedPtr<uint32> InMode) const
+{
+	UObject* Object = nullptr;
+
+	if (ParentFixtureTypeHandle->GetValue(Object) == FPropertyAccess::Success && Object != nullptr)
+	{
+		if (UDMXEntityFixtureType* Patch = Cast<UDMXEntityFixtureType>(Object))
+		{
+			if (InMode.IsValid() && 
+				Patch->Modes.IsValidIndex(*InMode))
+			{
+				return 
+					SNew(STextBlock)
+					.Text(FText::FromString(Patch->Modes[*InMode].ModeName));
+			}
+		}
+	}
+
+	return SNullWidget::NullWidget;
+}
+
+void FDMXEntityFixturePatchDetails::OnParentFixtureTypeChanged(UDMXEntity* NewTemplate) const
+{
+	ParentFixtureTypeHandle->SetValue(Cast<UDMXEntityFixtureType>(NewTemplate));
+}
+
+void FDMXEntityFixturePatchDetails::OnModeChanged(const UDMXEntityFixtureType* FixtureType, const FDMXFixtureMode& Mode)
+{
+	PropertyUtilities->ForceRefresh();
+}
+
+void FDMXEntityFixturePatchDetails::OnActiveModeChanged(const TSharedPtr<uint32> InSelectedMode, ESelectInfo::Type SelectInfo)
+{
+	if (InSelectedMode.IsValid())
+	{
+		ActiveModeHandle->SetValue(*InSelectedMode);
+	}
+}
+
+void FDMXEntityFixturePatchDetails::OnAutoAssignAddressChanged()
+{
+	bool bAutoAssignAddress;
+	if (ensure(AutoAssignAddressHandle->GetValue(bAutoAssignAddress) == FPropertyAccess::Success))
+	{
+		if (bAutoAssignAddress)
+		{
+			TArray<UObject*> OuterObjects;
+			AutoAssignAddressHandle->GetOuterObjects(OuterObjects);
+
+			TArray<UDMXEntityFixturePatch*> FixturePatches;
+			for (UObject* Object : OuterObjects)
+			{
+				UDMXEntityFixturePatch* Patch = CastChecked<UDMXEntityFixturePatch>(Object);
+				FixturePatches.Add(Patch);
+			}
+
+			FDMXEditorUtils::AutoAssignedAddresses(FixturePatches);
+		}
+	}
+}
+
+void FDMXEntityFixturePatchDetails::GenerateActiveModesSource()
+{
+	ActiveModesSource.Reset();
+
+	UObject* Object = nullptr;
+	if (ParentFixtureTypeHandle->GetValue(Object) == FPropertyAccess::Success)
+	{
+		if (UDMXEntityFixtureType* Fixture = Cast<UDMXEntityFixtureType>(Object))
+		{
+			const uint32 NumModes = Fixture->Modes.Num();
+			for (uint32 ModeIndex = 0; ModeIndex < NumModes; ++ModeIndex)
+			{
+				ActiveModesSource.Add(MakeShared<uint32>(ModeIndex));
+			}
+		}
+	}
+}
+
+TWeakObjectPtr<UDMXEntityFixtureType> FDMXEntityFixturePatchDetails::GetParentFixtureType() const
+{
+	UObject* Object;
+	if (ParentFixtureTypeHandle->GetValue(Object) == FPropertyAccess::Success)
+	{
+		return Cast<UDMXEntityFixtureType>(Object);
+	}
+	return nullptr;
+}
+
+
+bool FDMXEntityFixturePatchDetails::IsParentFixtureTypeMultipleValues() const
+{
+	UObject* Object;
+	return ParentFixtureTypeHandle->GetValue(Object) == FPropertyAccess::MultipleValues;
+}
+
+bool FDMXEntityFixturePatchDetails::IsActiveModeEditable() const
+{
+	UObject* Object = nullptr;
+	if (ParentFixtureTypeHandle->GetValue(Object) == FPropertyAccess::Success && Object)
+	{
+		if (UDMXEntityFixtureType* Fixture = Cast<UDMXEntityFixtureType>(Object))
+		{
+			return Fixture->Modes.Num() > 0;
+		}
+
+		return false;
+	}
+
+	return false;
+}
+
+FText FDMXEntityFixturePatchDetails::GetCurrentActiveModeLabel() const
+{
+	static const FText MultipleValuesLabel = LOCTEXT("MultipleValuesLabel", "Multiple Values");
+	static const FText NullTypeLabel = LOCTEXT("NullFixtureTypeLabel", "No Fixture Type selected");
+	static const FText MultipleTypesLabel = LOCTEXT("MultipleFixtureTypesLabel", "Multiple Types Selected");
+	static const FText NoModesLabel = LOCTEXT("NoModesLabel", "No modes in Fixture Type");
+
+	UObject* Object = nullptr;
+	const FPropertyAccess::Result FixtureTemplateAccessResult = ParentFixtureTypeHandle->GetValue(Object);
+	UDMXEntityFixtureType* FixtureTemplate = Cast<UDMXEntityFixtureType>(Object);
+
+	// Is only one type of Fixture Type selected?
+	if (FixtureTemplateAccessResult == FPropertyAccess::Success)
+	{
+		// Is this type valid?
+		if (FixtureTemplate != nullptr)
+		{
+			// We can try to get the mode, although it could be a different one for each of the templates
+			int32 ModeValue = 0;
+			if (ActiveModeHandle->GetValue(ModeValue) == FPropertyAccess::Success)
+			{
+				const TArray<FDMXFixtureMode>& CurrentModes = FixtureTemplate->Modes;
+				if (CurrentModes.Num() > 0 && CurrentModes.IsValidIndex(ModeValue))
+				{
+					return FText::FromString(
+						CurrentModes[ModeValue].ModeName
+					);
+				}
+				else
+				{
+					return NoModesLabel;
+				}
+			}
+
+			return MultipleValuesLabel;
+		}
+
+		return NullTypeLabel;
+	}
+
+	return MultipleTypesLabel;
+}
+
+void FDMXEntityFixturePatchDetails::SetActiveMode(int32 ModeIndex)
+{
+	TArray<UObject*> OuterObjects;
+	ActiveModeHandle->GetOuterObjects(OuterObjects);
+
+	TArray<UDMXEntityFixturePatch*> FixturePatches;
+	for (UObject* Object : OuterObjects)
+	{
+		UDMXEntityFixturePatch* Patch = CastChecked<UDMXEntityFixturePatch>(Object);
+		Patch->SetActiveModeIndex(ModeIndex);
+	}
+}
+
+#undef LOCTEXT_NAMESPACE
diff --git a/Engine/Plugins/VirtualProduction/DMX/DMXEngine/Source/DMXEditor/Private/Customizations/DMXEntityFixturePatchDetails.h b/Engine/Plugins/VirtualProduction/DMX/DMXEngine/Source/DMXEditor/Private/Customizations/DMXEntityFixturePatchDetails.h
new file mode 100644
--- /dev/null
+++ b/Engine/Plugins/VirtualProduction/DMX/DMXEngine/Source/DMXEditor/Private/Customizations/DMXEntityFixturePatchDetails.h
@@ -0,0 +1,90 @@
+// Copyright Epic Games, Inc. All Rights Reserved.
+
+#pragma once
+
+#include "CoreMinimal.h"
+#include "IDetailCustomization.h"
+#include "Types/SlateEnums.h"
+
+struct FDMXFixtureMode;
+class FDMXEditor;
+class UDMXEntity;
+class UDMXEntityFixtureType;
+
+class IPropertyHandle;
+class IPropertyUtilities;
+template<typename OptionType> class SComboBox;
+class SWidget;
+
+
+/** Details customization for the DMMXFixturePatch class */
+class FDMXEntityFixturePatchDetails
+	: public IDetailCustomization
+{
+public:
+	/** Constructor */
+	FDMXEntityFixturePatchDetails(TWeakPtr<FDMXEditor> InDMXEditorPtr);
+	
+	/** Creates a detail customization instance */
+	static TSharedRef<IDetailCustomization> MakeInstance(TWeakPtr<FDMXEditor> InDMXEditorPtr);
+
+protected:
+	//~ Begin IDetailCustomization interface
+	virtual void CustomizeDetails(IDetailLayoutBuilder& DetailBuilder) override;
+	//~ End IDetailCustomization interface
+
+private:
+	/** Generates the customized ActiveMode property widget */
+	TSharedRef<SWidget> GenerateActiveModeWidget(const TSharedPtr<uint32> InMode) const;
+
+	/** Called when the ParentFixtureType property changed */
+	void OnParentFixtureTypeChanged(UDMXEntity* NewTemplate) const;
+	
+	/** Called when the ActiveMode property changed */
+	void OnActiveModeChanged(const TSharedPtr<uint32> InSelectedMode, ESelectInfo::Type SelectInfo);
+
+	/** Called when a Mode in the Modes Array property changed */
+	void OnModeChanged(const UDMXEntityFixtureType* FixtureType, const FDMXFixtureMode& Mode);
+
+	/** Called when the bAutoAssignAddress property changed */
+	void OnAutoAssignAddressChanged();
+
+	/** Generates the Array of available mode in the ActiveModeSource member array */
+	void GenerateActiveModesSource();
+
+	/** Returns the Fixture Type the patch uses */
+	TWeakObjectPtr<UDMXEntityFixtureType> GetParentFixtureType() const;
+	
+	/** Returns true if the ParentFixtureType property has multiple values */
+	bool IsParentFixtureTypeMultipleValues() const;
+
+	/** Returns true if the ActiveMode property is editable */
+	bool IsActiveModeEditable() const;
+
+	/** Returns a label for the currently active mode */
+	FText GetCurrentActiveModeLabel() const;
+
+	/** Sets the ActiveMode to specified value */
+	void SetActiveMode(int32 ModeIndex);
+
+	/** Handle for the ParentFixtureType property */
+	TSharedPtr<IPropertyHandle> ParentFixtureTypeHandle;
+
+	/** Handle for the ActiveMode property */
+	TSharedPtr<IPropertyHandle> ActiveModeHandle;
+
+	/** Handle for the bAutoAssignAddress property */
+	TSharedPtr<IPropertyHandle> AutoAssignAddressHandle;
+
+	/** ComboBox that displays the ActiveMode property */
+	TSharedPtr<SComboBox<TSharedPtr<uint32>>> ActiveModeComboBox;
+
+	/** Array of available mode indicies */
+	TArray<TSharedPtr<uint32>> ActiveModesSource;
+
+	/** The property utilities for the customization */
+	TSharedPtr<IPropertyUtilities> PropertyUtilities;
+
+	/** Weak reference to the DMX editor */
+	TWeakPtr<FDMXEditor> DMXEditorPtr;
+};
diff --git a/Engine/Plugins/VirtualProduction/DMX/DMXEngine/Source/DMXEditor/Private/Customizations/DMXEntityFixtureTypeFixtureSettingsDetails.cpp b/Engine/Plugins/VirtualProduction/DMX/DMXEngine/Source/DMXEditor/Private/Customizations/DMXEntityFixtureTypeFixtureSettingsDetails.cpp
--- a/Engine/Plugins/VirtualProduction/DMX/DMXEngine/Source/DMXEditor/Private/Customizations/DMXEntityFixtureTypeFixtureSettingsDetails.cpp
+++ b/Engine/Plugins/VirtualProduction/DMX/DMXEngine/Source/DMXEditor/Private/Customizations/DMXEntityFixtureTypeFixtureSettingsDetails.cpp
@@ -10,6 +10,11 @@
 
 #define LOCTEXT_NAMESPACE "DMXEntityFixtureTypeFixtureSettingsDetails"
 
+TSharedRef<IDetailCustomization> FDMXEntityFixtureTypeFixtureSettingsDetails::MakeInstance(TWeakPtr<FDMXEditor> InDMXEditorPtr)
+{
+	return MakeShared<FDMXEntityFixtureTypeFixtureSettingsDetails>(InDMXEditorPtr);
+}
+
 void FDMXEntityFixtureTypeFixtureSettingsDetails::CustomizeDetails(IDetailLayoutBuilder& DetailBuilder)
 {
 	DetailBuilder.HideProperty(GET_MEMBER_NAME_CHECKED(UDMXEntityFixtureType, Modes));
diff --git a/Engine/Plugins/VirtualProduction/DMX/DMXEngine/Source/DMXEditor/Private/Customizations/DMXEntityFixtureTypeFixtureSettingsDetails.h b/Engine/Plugins/VirtualProduction/DMX/DMXEngine/Source/DMXEditor/Private/Customizations/DMXEntityFixtureTypeFixtureSettingsDetails.h
--- a/Engine/Plugins/VirtualProduction/DMX/DMXEngine/Source/DMXEditor/Private/Customizations/DMXEntityFixtureTypeFixtureSettingsDetails.h
+++ b/Engine/Plugins/VirtualProduction/DMX/DMXEngine/Source/DMXEditor/Private/Customizations/DMXEntityFixtureTypeFixtureSettingsDetails.h
@@ -18,7 +18,12 @@ public:
 		: DMXEditorPtr(InDMXEditorPtr)
 	{}
 
+	/** Creates an instance of this details customization */
+	static TSharedRef<IDetailCustomization> MakeInstance(TWeakPtr<FDMXEditor> InDMXEditorPtr);
+
+	//~ Begin IDetailCustomization interface
 	virtual void CustomizeDetails(IDetailLayoutBuilder& DetailBuilder) override;
+	//~ End IDetailCustomization interface
 
 protected:
 	/** Weak reference to the DMX editor */
diff --git a/Engine/Plugins/VirtualProduction/DMX/DMXEngine/Source/DMXEditor/Private/Customizations/DMXEntityFixtureTypeFunctionPropertiesDetails.cpp b/Engine/Plugins/VirtualProduction/DMX/DMXEngine/Source/DMXEditor/Private/Customizations/DMXEntityFixtureTypeFunctionPropertiesDetails.cpp
--- a/Engine/Plugins/VirtualProduction/DMX/DMXEngine/Source/DMXEditor/Private/Customizations/DMXEntityFixtureTypeFunctionPropertiesDetails.cpp
+++ b/Engine/Plugins/VirtualProduction/DMX/DMXEngine/Source/DMXEditor/Private/Customizations/DMXEntityFixtureTypeFunctionPropertiesDetails.cpp
@@ -17,6 +17,11 @@
 
 #define LOCTEXT_NAMESPACE "DMXEntityFixtureTypeFunctionPropertiesDetails"
 
+TSharedRef<IDetailCustomization> FDMXEntityFixtureTypeFunctionPropertiesDetails::MakeInstance(TWeakPtr<FDMXEditor> InDMXEditorPtr)
+{
+	return MakeShared<FDMXEntityFixtureTypeFunctionPropertiesDetails>(InDMXEditorPtr);
+}
+
 void FDMXEntityFixtureTypeFunctionPropertiesDetails::CustomizeDetails(IDetailLayoutBuilder& DetailBuilder)
 {	
 	DetailBuilder.HideCategory("Entity Properties");
diff --git a/Engine/Plugins/VirtualProduction/DMX/DMXEngine/Source/DMXEditor/Private/Customizations/DMXEntityFixtureTypeFunctionPropertiesDetails.h b/Engine/Plugins/VirtualProduction/DMX/DMXEngine/Source/DMXEditor/Private/Customizations/DMXEntityFixtureTypeFunctionPropertiesDetails.h
--- a/Engine/Plugins/VirtualProduction/DMX/DMXEngine/Source/DMXEditor/Private/Customizations/DMXEntityFixtureTypeFunctionPropertiesDetails.h
+++ b/Engine/Plugins/VirtualProduction/DMX/DMXEngine/Source/DMXEditor/Private/Customizations/DMXEntityFixtureTypeFunctionPropertiesDetails.h
@@ -27,7 +27,12 @@ public:
 		: DMXEditorPtr(InDMXEditorPtr)
 	{}
 
+	/** Creates an instance of this details customization */
+	static TSharedRef<IDetailCustomization> MakeInstance(TWeakPtr<FDMXEditor> InDMXEditorPtr);
+
+	//~ Begin IDetailCustomization interface
 	virtual void CustomizeDetails(IDetailLayoutBuilder& DetailBuilder) override;
+	//~ End IDetailCustomization interface
 
 protected:
 	/** Called when modes were selected */
diff --git a/Engine/Plugins/VirtualProduction/DMX/DMXEngine/Source/DMXEditor/Private/Customizations/DMXEntityFixtureTypeFunctionsDetails.cpp b/Engine/Plugins/VirtualProduction/DMX/DMXEngine/Source/DMXEditor/Private/Customizations/DMXEntityFixtureTypeFunctionsDetails.cpp
--- a/Engine/Plugins/VirtualProduction/DMX/DMXEngine/Source/DMXEditor/Private/Customizations/DMXEntityFixtureTypeFunctionsDetails.cpp
+++ b/Engine/Plugins/VirtualProduction/DMX/DMXEngine/Source/DMXEditor/Private/Customizations/DMXEntityFixtureTypeFunctionsDetails.cpp
@@ -18,6 +18,11 @@
 
 #define LOCTEXT_NAMESPACE "DMXEntityFixtureTypeFunctionsDetails"
 
+TSharedRef<IDetailCustomization> FDMXEntityFixtureTypeFunctionsDetails::MakeInstance(TWeakPtr<FDMXEditor> InDMXEditorPtr)
+{
+	return MakeShared<FDMXEntityFixtureTypeFunctionsDetails>(InDMXEditorPtr);
+}
+
 void FDMXEntityFixtureTypeFunctionsDetails::CustomizeDetails(IDetailLayoutBuilder& DetailBuilder)
 {
 	DetailBuilder.HideCategory("Entity Properties");
diff --git a/Engine/Plugins/VirtualProduction/DMX/DMXEngine/Source/DMXEditor/Private/Customizations/DMXEntityFixtureTypeFunctionsDetails.h b/Engine/Plugins/VirtualProduction/DMX/DMXEngine/Source/DMXEditor/Private/Customizations/DMXEntityFixtureTypeFunctionsDetails.h
--- a/Engine/Plugins/VirtualProduction/DMX/DMXEngine/Source/DMXEditor/Private/Customizations/DMXEntityFixtureTypeFunctionsDetails.h
+++ b/Engine/Plugins/VirtualProduction/DMX/DMXEngine/Source/DMXEditor/Private/Customizations/DMXEntityFixtureTypeFunctionsDetails.h
@@ -24,7 +24,13 @@ public:
 		: DMXEditorPtr(InDMXEditorPtr)
 	{}
 
+	/** Creates an instance of this details customization */
+	static TSharedRef<IDetailCustomization> MakeInstance(TWeakPtr<FDMXEditor> InDMXEditorPtr);
+	
+	//~ Begin IDetailCustomization interface
 	virtual void CustomizeDetails(IDetailLayoutBuilder& DetailBuilder) override;
+	//~ End IDetailCustomization interface
+
 protected:
 	/** Called when num modes changed */
 	void OnNumModesChanged();
diff --git a/Engine/Plugins/VirtualProduction/DMX/DMXEngine/Source/DMXEditor/Private/Customizations/DMXEntityFixtureTypeModePropertiesDetails.cpp b/Engine/Plugins/VirtualProduction/DMX/DMXEngine/Source/DMXEditor/Private/Customizations/DMXEntityFixtureTypeModePropertiesDetails.cpp
--- a/Engine/Plugins/VirtualProduction/DMX/DMXEngine/Source/DMXEditor/Private/Customizations/DMXEntityFixtureTypeModePropertiesDetails.cpp
+++ b/Engine/Plugins/VirtualProduction/DMX/DMXEngine/Source/DMXEditor/Private/Customizations/DMXEntityFixtureTypeModePropertiesDetails.cpp
@@ -18,6 +18,11 @@
 
 #define LOCTEXT_NAMESPACE "DMXEntityFixtureTypeModePropertiesDetails"
 
+TSharedRef<IDetailCustomization> FDMXEntityFixtureTypeModePropertiesDetails::MakeInstance(TWeakPtr<FDMXEditor> InDMXEditorPtr)
+{
+	return MakeShared<FDMXEntityFixtureTypeModePropertiesDetails>(InDMXEditorPtr);
+}
+
 void FDMXEntityFixtureTypeModePropertiesDetails::CustomizeDetails(IDetailLayoutBuilder& DetailBuilder)
 {
 	DetailBuilder.HideCategory("Entity Properties");
diff --git a/Engine/Plugins/VirtualProduction/DMX/DMXEngine/Source/DMXEditor/Private/Customizations/DMXEntityFixtureTypeModePropertiesDetails.h b/Engine/Plugins/VirtualProduction/DMX/DMXEngine/Source/DMXEditor/Private/Customizations/DMXEntityFixtureTypeModePropertiesDetails.h
--- a/Engine/Plugins/VirtualProduction/DMX/DMXEngine/Source/DMXEditor/Private/Customizations/DMXEntityFixtureTypeModePropertiesDetails.h
+++ b/Engine/Plugins/VirtualProduction/DMX/DMXEngine/Source/DMXEditor/Private/Customizations/DMXEntityFixtureTypeModePropertiesDetails.h
@@ -27,7 +27,12 @@ public:
 		: DMXEditorPtr(InDMXEditorPtr)
 	{}
 
+	/** Creates an instance of this details customization */
+	static TSharedRef<IDetailCustomization> MakeInstance(TWeakPtr<FDMXEditor> InDMXEditorPtr);
+	
+	//~ Begin IDetailCustomization interface
 	virtual void CustomizeDetails(IDetailLayoutBuilder& DetailBuilder) override;
+	//~ End IDetailCustomization interface
 
 protected:
 	/** Called when modes were selected */
diff --git a/Engine/Plugins/VirtualProduction/DMX/DMXEngine/Source/DMXEditor/Private/Customizations/DMXEntityFixtureTypeModesDetails.cpp b/Engine/Plugins/VirtualProduction/DMX/DMXEngine/Source/DMXEditor/Private/Customizations/DMXEntityFixtureTypeModesDetails.cpp
--- a/Engine/Plugins/VirtualProduction/DMX/DMXEngine/Source/DMXEditor/Private/Customizations/DMXEntityFixtureTypeModesDetails.cpp
+++ b/Engine/Plugins/VirtualProduction/DMX/DMXEngine/Source/DMXEditor/Private/Customizations/DMXEntityFixtureTypeModesDetails.cpp
@@ -415,6 +415,11 @@ void SDMXModeItemListViewBox::OnRenameItem()
 	}
 }
 
+TSharedRef<IDetailCustomization> FDMXEntityFixtureTypeModesDetails::MakeInstance(TWeakPtr<FDMXEditor> InDMXEditorPtr)
+{
+	return MakeShared<FDMXEntityFixtureTypeModesDetails>(InDMXEditorPtr);
+}
+
 void FDMXEntityFixtureTypeModesDetails::CustomizeDetails(IDetailLayoutBuilder& DetailBuilder)
 {
 	DetailBuilder.HideCategory("Entity Properties");
diff --git a/Engine/Plugins/VirtualProduction/DMX/DMXEngine/Source/DMXEditor/Private/Customizations/DMXEntityFixtureTypeModesDetails.h b/Engine/Plugins/VirtualProduction/DMX/DMXEngine/Source/DMXEditor/Private/Customizations/DMXEntityFixtureTypeModesDetails.h
--- a/Engine/Plugins/VirtualProduction/DMX/DMXEngine/Source/DMXEditor/Private/Customizations/DMXEntityFixtureTypeModesDetails.h
+++ b/Engine/Plugins/VirtualProduction/DMX/DMXEngine/Source/DMXEditor/Private/Customizations/DMXEntityFixtureTypeModesDetails.h
@@ -38,7 +38,12 @@ public:
 		: DMXEditorPtr(InDMXEditorPtr)
 	{}
 
+	/** Creates an instance of this details customization */
+	static TSharedRef<IDetailCustomization> MakeInstance(TWeakPtr<FDMXEditor> InDMXEditorPtr);
+
+	//~ Begin IDetailCustomization interface
 	virtual void CustomizeDetails(IDetailLayoutBuilder& DetailBuilder) override;
+	//~ End IDetailCustomization interface
 
 protected:
 	//~ Begin FEditorUndoClient Interface
diff --git a/Engine/Plugins/VirtualProduction/DMX/DMXEngine/Source/DMXEditor/Private/Customizations/DMXEntityReferenceCustomization.cpp b/Engine/Plugins/VirtualProduction/DMX/DMXEngine/Source/DMXEditor/Private/Customizations/DMXEntityReferenceCustomization.cpp
new file mode 100644
--- /dev/null
+++ b/Engine/Plugins/VirtualProduction/DMX/DMXEngine/Source/DMXEditor/Private/Customizations/DMXEntityReferenceCustomization.cpp
@@ -0,0 +1,228 @@
+// Copyright Epic Games, Inc. All Rights Reserved.
+
+#include "Customizations/DMXEntityReferenceCustomization.h"
+
+#include "Library/DMXEntity.h"
+#include "Library/DMXEntityReference.h"
+#include "Library/DMXLibrary.h"
+#include "Widgets/SDMXEntityDropdownMenu.h"
+
+#include "DetailWidgetRow.h"
+#include "IDetailPropertyRow.h"
+#include "IDetailChildrenBuilder.h"
+#include "Widgets/Text/STextBlock.h"
+
+
+#define LOCTEXT_NAMESPACE "DMXEntityReferenceCustomization"
+
+TSharedRef<IPropertyTypeCustomization> FDMXEntityReferenceCustomization::MakeInstance()
+{
+	return MakeShared<FDMXEntityReferenceCustomization>();
+}
+
+void FDMXEntityReferenceCustomization::CustomizeHeader(TSharedRef<IPropertyHandle> InPropertyHandle, FDetailWidgetRow& InHeaderRow, IPropertyTypeCustomizationUtils& CustomizationUtils)
+{
+	StructHandle = InPropertyHandle;
+
+	// If we should display any children, leave the standard header. Otherwise, create custom picker
+	TSharedRef<SWidget> ValueContent = !ShouldDisplayLibrary()
+		? CreateEntityPickerWidget(InPropertyHandle)
+		: InPropertyHandle->CreatePropertyValueWidget(false);
+
+	InHeaderRow
+		.NameContent()
+		[
+			InPropertyHandle->CreatePropertyNameWidget()
+		]
+		.ValueContent()
+		.MinDesiredWidth(200.0f)
+		.MaxDesiredWidth(400.0f)
+		[
+			ValueContent
+		];
+}
+
+void FDMXEntityReferenceCustomization::CustomizeChildren(TSharedRef<IPropertyHandle> InPropertyHandle, IDetailChildrenBuilder& InChildBuilder, IPropertyTypeCustomizationUtils& CustomizationUtils)
+{
+	if (!ShouldDisplayLibrary())
+	{
+		// Don't add any child properties
+		return;
+	}
+
+	// Retrieve structure's child properties
+	uint32 NumChildren;
+	InPropertyHandle->GetNumChildren(NumChildren);
+	TMap<FName, TSharedPtr< IPropertyHandle > > PropertyHandles;
+	for (uint32 ChildIndex = 0; ChildIndex < NumChildren; ++ChildIndex)
+	{
+		TSharedRef<IPropertyHandle> ChildHandle = InPropertyHandle->GetChildHandle(ChildIndex).ToSharedRef();
+		const FName PropertyName = ChildHandle->GetProperty()->GetFName();
+		PropertyHandles.Add(PropertyName, ChildHandle);
+	}
+
+	// Add the properties
+	for (const TPair<FName, TSharedPtr<IPropertyHandle>>& PropertyHandlePair : PropertyHandles)
+	{
+		if ((PropertyHandlePair.Key == GET_MEMBER_NAME_CHECKED(FDMXEntityReference, DMXLibrary) && ShouldDisplayLibrary())
+			|| (PropertyHandlePair.Key != GET_MEMBER_NAME_CHECKED(FDMXEntityReference, DMXLibrary) && PropertyHandlePair.Key != "EntityId")) // EntityId is private, so GET_MEMBER... won't work
+		{
+			InChildBuilder.AddProperty(PropertyHandlePair.Value.ToSharedRef());
+		}
+	}
+
+	// Add the picker
+	InChildBuilder.AddCustomRow(LOCTEXT("EntityReferencePickerSearchText", "Entity"))
+		.NameContent()
+		[
+			SNew(STextBlock)
+			.Font(CustomizationUtils.GetRegularFont())
+			.Text(this, &FDMXEntityReferenceCustomization::GetPickerPropertyLabel)
+		]
+		.ValueContent()
+		.MinDesiredWidth(200.0f)
+		.MaxDesiredWidth(400.0f)
+		[
+			CreateEntityPickerWidget(InPropertyHandle)
+		];
+}
+
+bool FDMXEntityReferenceCustomization::ShouldDisplayLibrary() const
+{
+	TArray<const void*> RawDataArr;
+	StructHandle->AccessRawData(RawDataArr);
+
+	for (const void* RawData : RawDataArr)
+	{
+		const FDMXEntityReference* EntityRefPtr = reinterpret_cast<const FDMXEntityReference*>(RawData);
+		if (!EntityRefPtr->bDisplayLibraryPicker)
+		{
+			return false;
+		}
+	}
+
+	return true;
+}
+
+TSharedRef<SWidget> FDMXEntityReferenceCustomization::CreateEntityPickerWidget(TSharedPtr<IPropertyHandle> InPropertyHandle) const
+{
+	return SNew(SDMXEntityPickerButton<UDMXEntity>)
+		.CurrentEntity(this, &FDMXEntityReferenceCustomization::GetCurrentEntity)
+		.HasMultipleValues(this, &FDMXEntityReferenceCustomization::GetEntityIsMultipleValues)
+		.OnEntitySelected(this, &FDMXEntityReferenceCustomization::OnEntitySelected)
+		.EntityTypeFilter(this, &FDMXEntityReferenceCustomization::GetEntityType)
+		.DMXLibrary(this, &FDMXEntityReferenceCustomization::GetDMXLibrary)
+		.IsEnabled(this, &FDMXEntityReferenceCustomization::GetPickerEnabled);
+}
+
+FText FDMXEntityReferenceCustomization::GetPickerPropertyLabel() const
+{
+	if (TSubclassOf<UDMXEntity> EntityType = GetEntityType())
+	{
+		return FDMXEditorUtils::GetEntityTypeNameText(EntityType, false);
+	}
+
+	return LOCTEXT("GenericTypeEntityLabel", "Entity");
+}
+
+bool FDMXEntityReferenceCustomization::GetPickerEnabled() const
+{
+	return GetEntityType() != nullptr;
+}
+
+TWeakObjectPtr<UDMXEntity> FDMXEntityReferenceCustomization::GetCurrentEntity() const
+{
+	if (GetEntityIsMultipleValues()) return nullptr;
+
+	void* StructPtr = nullptr;
+	if (StructHandle->GetValueData(StructPtr) == FPropertyAccess::Success && StructPtr != nullptr)
+	{
+		const FDMXEntityReference* EntityRef = reinterpret_cast<FDMXEntityReference*>(StructPtr);
+		return EntityRef->GetEntity();
+	}
+
+	return nullptr;
+}
+
+bool FDMXEntityReferenceCustomization::GetEntityIsMultipleValues() const
+{
+	TArray<void*> RawData;
+	StructHandle->AccessRawData(RawData);
+	if (RawData[0] == nullptr)
+	{
+		return true;
+	}
+
+	bool bFirstEntitySet = false;
+	UDMXEntity* FirstEntityPtr = nullptr;
+
+	for (const void* StructPtr : RawData)
+	{
+		const FDMXEntityReference* EntityRefPtr = reinterpret_cast<const FDMXEntityReference*>(StructPtr);
+
+		if (bFirstEntitySet)
+		{
+			if (EntityRefPtr->GetEntity() != FirstEntityPtr)
+			{
+				return true;
+			}
+		}
+		else
+		{
+			FirstEntityPtr = EntityRefPtr->GetEntity();
+			bFirstEntitySet = true;
+		}
+	}
+
+	return false;
+}
+
+void FDMXEntityReferenceCustomization::OnEntitySelected(UDMXEntity* NewEntity) const
+{
+	FDMXEntityReference NewStructValues;
+	NewStructValues.SetEntity(NewEntity);
+
+	// Export new values to text format that can be imported later into the actual struct properties
+	FString TextValue;
+	FDMXEntityReference::StaticStruct()->ExportText(TextValue, &NewStructValues, nullptr, nullptr, EPropertyPortFlags::PPF_None, nullptr);
+
+	// Set values on edited property handle from exported text
+	ensure(StructHandle->SetValueFromFormattedString(TextValue, EPropertyValueSetFlags::DefaultFlags) == FPropertyAccess::Success);
+}
+
+TSubclassOf<UDMXEntity> FDMXEntityReferenceCustomization::GetEntityType() const
+{
+	TArray<void*> RawData;
+	StructHandle->AccessRawData(RawData);
+	if (RawData[0] == nullptr)
+	{
+		return nullptr;
+	}
+
+	const TSubclassOf<UDMXEntity> FirstEntityType = reinterpret_cast<FDMXEntityReference*>(RawData[0])->GetEntityType();
+
+	for (const void* StructPtr : RawData)
+	{
+		const FDMXEntityReference* EntityRefPtr = reinterpret_cast<const FDMXEntityReference*>(StructPtr);
+		if (EntityRefPtr->GetEntityType() != FirstEntityType)
+		{
+			// Different types are selected
+			return nullptr;
+		}
+	}
+
+	return FirstEntityType;
+}
+
+TWeakObjectPtr<UDMXLibrary> FDMXEntityReferenceCustomization::GetDMXLibrary() const
+{
+	TSharedPtr<IPropertyHandle> LibraryHandle = StructHandle->GetChildHandle(GET_MEMBER_NAME_CHECKED(FDMXEntityReference, DMXLibrary));
+	UObject* Object = nullptr;
+	if (LibraryHandle->GetValue(Object) == FPropertyAccess::Success)
+	{
+		return Cast<UDMXLibrary>(Object);
+	}
+	return nullptr;
+}
+
+#undef LOCTEXT_NAMESPACE
diff --git a/Engine/Plugins/VirtualProduction/DMX/DMXEngine/Source/DMXEditor/Private/Customizations/DMXEntityReferenceCustomization.h b/Engine/Plugins/VirtualProduction/DMX/DMXEngine/Source/DMXEditor/Private/Customizations/DMXEntityReferenceCustomization.h
new file mode 100644
--- /dev/null
+++ b/Engine/Plugins/VirtualProduction/DMX/DMXEngine/Source/DMXEditor/Private/Customizations/DMXEntityReferenceCustomization.h
@@ -0,0 +1,59 @@
+// Copyright Epic Games, Inc. All Rights Reserved.
+
+#pragma once
+
+#include "IPropertyTypeCustomization.h"
+#include "CoreMinimal.h"
+
+#include "DMXProtocolTypes.h"
+
+class UDMXEntity;
+class UDMXLibrary;
+
+/**
+ * Customization for Entity Reference structs
+ */
+class FDMXEntityReferenceCustomization
+	: public IPropertyTypeCustomization
+{
+public:
+	/** Creates an instance of the property type customization */
+	static TSharedRef<IPropertyTypeCustomization> MakeInstance();
+
+	//~ IPropertyTypeCustomization interface begin
+	virtual void CustomizeHeader(TSharedRef<IPropertyHandle> InPropertyHandle, FDetailWidgetRow& InHeaderRow, IPropertyTypeCustomizationUtils& CustomizationUtils) override;
+	virtual void CustomizeChildren(TSharedRef<IPropertyHandle> InPropertyHandle, IDetailChildrenBuilder& InChildBuilder, IPropertyTypeCustomizationUtils& CustomizationUtils) override;
+	//~ IPropertyTypeCustomization interface end
+
+private:
+	/** Returns whether the library field should be displayed */
+	bool ShouldDisplayLibrary() const;
+
+	/** Creates the widget where the entity can be picked */
+	TSharedRef<SWidget> CreateEntityPickerWidget(TSharedPtr<IPropertyHandle> InPropertyHandle) const;
+
+	/** Returns the label of the entity property */
+	FText GetPickerPropertyLabel() const;
+
+	/** Returns true if the picker is enabled */
+	bool GetPickerEnabled() const;
+	
+	/** Returns the currently selected entity */
+	TWeakObjectPtr<UDMXEntity> GetCurrentEntity() const;
+
+	/** Returns true if multiple values are customized */
+	bool GetEntityIsMultipleValues() const;
+
+	/** Called when an entity was selected */
+	void OnEntitySelected(UDMXEntity* NewEntity) const;
+
+	/** Returns the type of entities that can be selected */
+	TSubclassOf<UDMXEntity> GetEntityType() const;
+
+	/** Returns the DMX library where the entity resides */
+	TWeakObjectPtr<UDMXLibrary> GetDMXLibrary() const;
+
+private:
+	/** Handle for the entity struct that is customized */
+	TSharedPtr<IPropertyHandle> StructHandle;
+};
diff --git a/Engine/Plugins/VirtualProduction/DMX/DMXEngine/Source/DMXEditor/Private/Customizations/DMXFixtureFunctionCustomization.cpp b/Engine/Plugins/VirtualProduction/DMX/DMXEngine/Source/DMXEditor/Private/Customizations/DMXFixtureFunctionCustomization.cpp
new file mode 100644
--- /dev/null
+++ b/Engine/Plugins/VirtualProduction/DMX/DMXEngine/Source/DMXEditor/Private/Customizations/DMXFixtureFunctionCustomization.cpp
@@ -0,0 +1,344 @@
+// Copyright Epic Games, Inc. All Rights Reserved.
+
+#include "Customizations/DMXFixtureFunctionCustomization.h"
+
+#include "DMXEditor.h"
+#include "Library/DMXEntityFixtureType.h"
+
+#include "DetailWidgetRow.h"
+#include "IDetailChildrenBuilder.h"
+#include "DetailLayoutBuilder.h"
+#include "Widgets/Input/SEditableTextBox.h"
+#include "Widgets/Input/SNumericEntryBox.h"
+
+
+#define LOCTEXT_NAMESPACE "DMXFixtureTypeFunctionCustomization"
+
+TSharedRef<IPropertyTypeCustomization> FDMXFixtureFunctionCustomization::MakeInstance()
+{
+	return MakeShared<FDMXFixtureFunctionCustomization>();
+}
+
+void FDMXFixtureFunctionCustomization::CustomizeHeader(TSharedRef<IPropertyHandle> InStructPropertyHandle, class FDetailWidgetRow& InHeaderRow, IPropertyTypeCustomizationUtils& InStructCustomizationUtils)
+{
+	constexpr bool bDisplayDefaultPropertyButtons = false;
+
+	InHeaderRow
+		.NameContent()
+		[
+			InStructPropertyHandle->CreatePropertyNameWidget()
+		]
+		.ValueContent()
+		[
+			InStructPropertyHandle->CreatePropertyValueWidget(bDisplayDefaultPropertyButtons)
+		];
+}
+
+void FDMXFixtureFunctionCustomization::CustomizeChildren(TSharedRef<IPropertyHandle> InStructPropertyHandle, IDetailChildrenBuilder& InStructBuilder, IPropertyTypeCustomizationUtils& InStructCustomizationUtils)
+{
+	// Retrieve structure's child properties
+	uint32 NumChildren;
+	InStructPropertyHandle->GetNumChildren(NumChildren);
+	
+	TMap<FName, TSharedRef<IPropertyHandle>> PropertyNamePropertyHandleMap;
+	for (uint32 ChildIndex = 0; ChildIndex < NumChildren; ++ChildIndex)
+	{
+		TSharedRef<IPropertyHandle> ChildHandle = InStructPropertyHandle->GetChildHandle(ChildIndex).ToSharedRef();
+		const FName PropertyName = ChildHandle->GetProperty()->GetFName();
+		PropertyNamePropertyHandleMap.Add(PropertyName, ChildHandle);
+	}
+
+	// Remember relevant property handles
+	FunctionNameHandle = PropertyNamePropertyHandleMap.FindChecked(GET_MEMBER_NAME_CHECKED(FDMXFixtureFunction, FunctionName));
+	DataTypeHandle = PropertyNamePropertyHandleMap.FindChecked(GET_MEMBER_NAME_CHECKED(FDMXFixtureFunction, DataType));
+	DefaultValueHandle = PropertyNamePropertyHandleMap.FindChecked(GET_MEMBER_NAME_CHECKED(FDMXFixtureFunction, DefaultValue));
+	UseLSBHandle = PropertyNamePropertyHandleMap.FindChecked(GET_MEMBER_NAME_CHECKED(FDMXFixtureFunction, bUseLSBMode));
+
+	// Add properties to the layout
+	for (const TPair<FName, TSharedRef<IPropertyHandle>>& PropertyNamePropertyHandlePair : PropertyNamePropertyHandleMap)
+	{
+		if (PropertyNamePropertyHandlePair.Key == GET_MEMBER_NAME_CHECKED(FDMXFixtureFunction, FunctionName))
+		{
+			// Customize the FunctionName property
+			BuildFunctionNameWidget(InStructBuilder);
+		}
+		else if (PropertyNamePropertyHandlePair.Key == GET_MEMBER_NAME_CHECKED(FDMXFixtureFunction, DefaultValue))
+		{
+			// Customize the DefaultValue property (add a field for each byte of the default value, which will be displayed depending on the selected DataType)
+			InStructBuilder.AddProperty(PropertyNamePropertyHandlePair.Value);
+			AddDefaultValueFields(InStructBuilder);
+		}
+		else
+		{
+			// Add other properties
+			InStructBuilder.AddProperty(PropertyNamePropertyHandlePair.Value);
+		}
+	}
+}
+
+void FDMXFixtureFunctionCustomization::BuildFunctionNameWidget(IDetailChildrenBuilder& InStructBuilder)
+{
+	InStructBuilder
+		.AddCustomRow(LOCTEXT("FunctionNameWidget", "FunctionNameWidget"))
+		.NameContent()
+		[
+			SNew(STextBlock)
+			.Text(LOCTEXT("FixtureFunctionNameLabel", "Function Name"))
+			.Font(IDetailLayoutBuilder::GetDetailFont())
+		]
+		.ValueContent()
+		.MaxDesiredWidth(250.0f)
+		[
+			SAssignNew(FunctionNameEditableTextBox, SEditableTextBox)
+			.Text(this, &FDMXFixtureFunctionCustomization::GetFunctionName)
+			.OnTextChanged(this, &FDMXFixtureFunctionCustomization::OnFunctionNameChanged)
+			.OnTextCommitted(this, &FDMXFixtureFunctionCustomization::OnFunctionNameCommitted)
+			.Font(IDetailLayoutBuilder::GetDetailFont())
+		];
+}
+
+TArray<FString> FDMXFixtureFunctionCustomization::GetExistingFunctionNames() const
+{
+	TArray<FString> ExistingFunctionNames;
+
+	// Get the parent functions array property to be able to read the other functions names inside the current Mode
+	TSharedRef<IPropertyHandle> FunctionStructHandle = FunctionNameHandle->GetParentHandle().ToSharedRef();
+	TSharedRef<IPropertyHandle> FunctionsArrayHandle = FunctionStructHandle->GetParentHandle().ToSharedRef();
+
+	uint32 NumFunctions;
+	FunctionsArrayHandle->GetNumChildren(NumFunctions);
+
+	for (uint32 FunctionIndex = 0; FunctionIndex < NumFunctions; FunctionIndex++)
+	{
+		TSharedRef<IPropertyHandle> FunctionHandle = FunctionsArrayHandle->GetChildHandle(FunctionIndex).ToSharedRef();
+		TSharedPtr<IPropertyHandle> NameHandle = FunctionHandle->GetChildHandle(GET_MEMBER_NAME_CHECKED(FDMXFixtureFunction, FunctionName));
+
+		FString FunctionName;
+		NameHandle->GetValue(FunctionName);
+
+		ExistingFunctionNames.Add(FunctionName);
+	}
+
+	return ExistingFunctionNames;
+}
+
+void FDMXFixtureFunctionCustomization::OnFunctionNameChanged(const FText& InNewText)
+{
+	check(FunctionNameEditableTextBox.IsValid() && FunctionNameHandle.IsValid() && FunctionNameHandle->IsValidHandle());
+
+	if (FText::TrimPrecedingAndTrailing(InNewText).IsEmpty())
+	{
+		FunctionNameEditableTextBox->SetError(LOCTEXT("FunctionNameError_Empty", "The name can't be blank!"));
+		return;
+	}
+
+	FString CurrentName;
+	FunctionNameHandle->GetValue(CurrentName);
+
+	const FString& NewName = InNewText.ToString();
+	if (CurrentName.Equals(NewName))
+	{
+		FunctionNameEditableTextBox->SetError(FText::GetEmpty());
+		return;
+	}
+
+	const TArray<FString> ExistingNames = GetExistingFunctionNames();
+	if (ExistingNames.Contains(NewName))
+	{
+		FunctionNameEditableTextBox->SetError(LOCTEXT("FixtureFunctionName_Existent", "This name is already used by another function in this mode!"));
+	}
+	else
+	{
+		FunctionNameEditableTextBox->SetError(FText::GetEmpty());
+	}
+}
+
+void FDMXFixtureFunctionCustomization::OnFunctionNameCommitted(const FText& InNewText, ETextCommit::Type InTextCommit)
+{
+	if (InTextCommit != ETextCommit::OnCleared && !FunctionNameEditableTextBox->HasError())
+	{
+		const FString& NewName = InNewText.ToString();
+		SetFunctionName(NewName);
+	}
+	FunctionNameEditableTextBox->SetError(FText::GetEmpty());
+}
+
+FText FDMXFixtureFunctionCustomization::GetFunctionName() const
+{
+	FString Name;
+	FunctionNameHandle->GetValue(Name);
+
+	return FText::FromString(Name);
+}
+
+void FDMXFixtureFunctionCustomization::SetFunctionName(const FString& NewName)
+{
+	if (FunctionNameHandle.IsValid())
+	{
+		FunctionNameHandle->SetValue(NewName);
+	}
+}
+
+void FDMXFixtureFunctionCustomization::AddDefaultValueFields(IDetailChildrenBuilder& InStructBuilder)
+{
+	const FMargin Padding = FMargin(2.0f, 0.0f, 0.0f, 0.0f);
+
+	InStructBuilder
+		.AddCustomRow(LOCTEXT("DefaultValueRowName", "Default Value (Bytes)"))
+		.NameContent()
+		[
+			SNew(STextBlock)
+			.Text(LOCTEXT("DefaultValueLabel", "Default Value (Bytes)"))
+			.Font(IDetailLayoutBuilder::GetDetailFont())
+			.ToolTipText(LOCTEXT("DefaultValueTooltip", "The default value for the channels of the Attribute"))
+		]
+		.ValueContent()
+		.MinDesiredWidth(340.0f)
+		.MaxDesiredWidth(340.0f)
+		[
+			SNew(SHorizontalBox)
+			+ SHorizontalBox::Slot()
+			.FillWidth(1.0f)
+			.VAlign(VAlign_Fill)
+			[
+				CreateSingleDefaultValueChannelField(1, SNumericEntryBox<uint8>::RedLabelBackgroundColor)
+			]
+
+			+ SHorizontalBox::Slot()
+			.FillWidth(1.0f)
+			.VAlign(VAlign_Fill)
+			.Padding(Padding)
+			[
+				CreateSingleDefaultValueChannelField(2, SNumericEntryBox<uint8>::GreenLabelBackgroundColor)
+			]
+
+			+ SHorizontalBox::Slot()
+			.FillWidth(1.0f)
+			.VAlign(VAlign_Fill)
+			.Padding(Padding)
+			[
+				CreateSingleDefaultValueChannelField(3, SNumericEntryBox<uint8>::BlueLabelBackgroundColor)
+			]
+
+			+ SHorizontalBox::Slot()
+			.FillWidth(1.0f)
+			.VAlign(VAlign_Fill)
+			.Padding(Padding)
+			[
+				CreateSingleDefaultValueChannelField(4)
+			]
+		]
+		.Visibility(MakeAttributeSP(this, &FDMXFixtureFunctionCustomization::GetDefaultValueChannelVisibility, uint8(1)));
+
+}
+
+TSharedRef<SWidget> FDMXFixtureFunctionCustomization::CreateSingleDefaultValueChannelField(uint8 Channel, const FLinearColor& LabelColor)
+{
+	return SNew(SNumericEntryBox<uint8>)
+		.MinValue(0)
+		.MaxValue(255)
+		.MaxSliderValue(255)
+		.Value(this, &FDMXFixtureFunctionCustomization::GetDefaultValueChannelValue, Channel)
+		.OnValueChanged(this, &FDMXFixtureFunctionCustomization::OnDefaultValueChannelValueChanged, Channel)
+		.OnValueCommitted(this, &FDMXFixtureFunctionCustomization::OnDefaultValueChannelValueCommitted)
+		.Font(IDetailLayoutBuilder::GetDetailFont())
+		.AllowSpin(true)
+		.Visibility(MakeAttributeSP(this, &FDMXFixtureFunctionCustomization::GetDefaultValueChannelVisibility, Channel))
+		.LabelPadding(0)
+		.Label()
+		[
+			SNumericEntryBox<uint8>::BuildLabel(
+				FText::AsNumber(Channel),
+				FLinearColor::White,
+				LabelColor)
+		];
+}
+
+TOptional<uint8> FDMXFixtureFunctionCustomization::GetDefaultValueChannelValue(uint8 Channel) const
+{
+	TOptional<uint8> RetVal;
+	int64 Value;
+	bool bUseLSBValue;
+	if (DefaultValueHandle->GetValue(Value) == FPropertyAccess::Success
+		&& UseLSBHandle->GetValue(bUseLSBValue) == FPropertyAccess::Success)
+	{
+		uint8 BytesOffset = 0;
+
+		if (!bUseLSBValue)
+		{
+			void* DataTypePtr = nullptr;
+			if (DataTypeHandle->GetValueData(DataTypePtr) == FPropertyAccess::Success)
+			{
+				const EDMXFixtureSignalFormat& DataType = *(EDMXFixtureSignalFormat*)DataTypePtr;
+				const uint8 NumChannels = UDMXEntityFixtureType::NumChannelsToOccupy(DataType);
+				BytesOffset = (NumChannels - Channel) * 8;
+			}
+		}
+		else
+		{
+			BytesOffset = (Channel - 1) * 8;
+		}
+
+		RetVal = Value >> BytesOffset & 0xff;
+	}
+
+	return RetVal;
+}
+
+EVisibility FDMXFixtureFunctionCustomization::GetDefaultValueChannelVisibility(uint8 Channel) const
+{
+	void* DataTypePtr = nullptr;
+	if (DataTypeHandle->GetValueData(DataTypePtr) == FPropertyAccess::Success)
+	{
+		EDMXFixtureSignalFormat* DataType = (EDMXFixtureSignalFormat*)DataTypePtr;
+		if ((*DataType >= EDMXFixtureSignalFormat::E32Bit && Channel <= 4)
+			|| (*DataType >= EDMXFixtureSignalFormat::E24Bit && Channel <= 3)
+			|| (*DataType >= EDMXFixtureSignalFormat::E16Bit && Channel <= 2)
+			|| (*DataType >= EDMXFixtureSignalFormat::E8Bit && Channel <= 1))
+		{
+			return EVisibility::Visible;
+		}
+	}
+
+	return EVisibility::Collapsed;
+}
+
+void FDMXFixtureFunctionCustomization::OnDefaultValueChannelValueChanged(uint8 NewValue, uint8 Channel)
+{
+	int64 DefaultValue;
+	bool bUseLSBValue;
+	if (DefaultValueHandle->GetValue(DefaultValue) == FPropertyAccess::Success
+		&& UseLSBHandle->GetValue(bUseLSBValue) == FPropertyAccess::Success)
+	{
+		uint8* ValueBytes = reinterpret_cast<uint8*>(&DefaultValue);
+
+		if (!bUseLSBValue)
+		{
+			void* DataTypePtr = nullptr;
+			if (DataTypeHandle->GetValueData(DataTypePtr) == FPropertyAccess::Success)
+			{
+				const EDMXFixtureSignalFormat& DataType = *(EDMXFixtureSignalFormat*)DataTypePtr;
+				const uint8 NumChannels = UDMXEntityFixtureType::NumChannelsToOccupy(DataType);
+				ValueBytes[NumChannels - Channel] = NewValue;
+			}
+		}
+		else
+		{
+			ValueBytes[Channel - 1] = NewValue;
+		}
+
+		DefaultValueHandle->SetValue(DefaultValue, EPropertyValueSetFlags::InteractiveChange);
+	}
+}
+
+void FDMXFixtureFunctionCustomization::OnDefaultValueChannelValueCommitted(uint8 NewValue, ETextCommit::Type CommitType)
+{
+	int64 DefaultValue;
+	if (DefaultValueHandle->GetValue(DefaultValue) == FPropertyAccess::Success)
+	{
+		// Set the value anew but without the InteractiveChange flag, to trigger a transaction when as the value is committed.
+		DefaultValueHandle->SetValue(DefaultValue);
+	}
+}
+
+#undef LOCTEXT_NAMESPACE
diff --git a/Engine/Plugins/VirtualProduction/DMX/DMXEngine/Source/DMXEditor/Private/Customizations/DMXFixtureFunctionCustomization.h b/Engine/Plugins/VirtualProduction/DMX/DMXEngine/Source/DMXEditor/Private/Customizations/DMXFixtureFunctionCustomization.h
new file mode 100644
--- /dev/null
+++ b/Engine/Plugins/VirtualProduction/DMX/DMXEngine/Source/DMXEditor/Private/Customizations/DMXFixtureFunctionCustomization.h
@@ -0,0 +1,88 @@
+// Copyright Epic Games, Inc. All Rights Reserved.
+
+#pragma once
+
+#include "CoreMinimal.h"
+#include "IPropertyTypeCustomization.h"
+
+class FDMXEditor;
+class UDMXEntityFixtureType;
+
+struct EVisibility;
+class SEditableTextBox;
+
+
+/** 
+ * Customization for the DMXFixtureFunction struct 
+ */
+class FDMXFixtureFunctionCustomization
+	: public IPropertyTypeCustomization
+{
+public:
+	static TSharedRef<IPropertyTypeCustomization> MakeInstance();
+
+protected:
+	//~ Begin IPropertyTypeCustomization interface
+	virtual void CustomizeHeader(TSharedRef<IPropertyHandle> InStructPropertyHandle, class FDetailWidgetRow& InHeaderRow, IPropertyTypeCustomizationUtils& InStructCustomizationUtils) override;
+	virtual void CustomizeChildren(TSharedRef<IPropertyHandle> InStructPropertyHandle, IDetailChildrenBuilder& InStructBuilder, IPropertyTypeCustomizationUtils& InStructCustomizationUtils) override;
+	//~ End IPropertyTypeCustomization interface
+	
+private:
+	////////////////////////////////////////
+	// FunctionName property related methods
+
+	/** Builds the widget that displays the FunctionName property */
+	void BuildFunctionNameWidget(IDetailChildrenBuilder& InStructBuilder);
+	
+	/** Find the existing names for the function type being edited within the Fixture Type the function resides in */
+	TArray<FString> GetExistingFunctionNames() const;
+
+	/** Called when the name of a function changed */
+	void OnFunctionNameChanged(const FText& InNewText);
+
+	/** Called when the name of a mode changed */
+	void OnFunctionNameCommitted(const FText& InNewText, ETextCommit::Type InTextCommit);
+
+	/** Returns the function name */
+	FText GetFunctionName() const;
+
+	/** Changes the function name on the fixture properties */
+	void SetFunctionName(const FString& NewName);
+
+	////////////////////////////////////////
+	// DefaultValue property related methods
+
+	/** Adds fields for each byte of the default value, which will be displayed depending on the selected DataType */
+	void AddDefaultValueFields(IDetailChildrenBuilder& InStructBuilder);
+
+	/** Creates a single default value field */
+	TSharedRef<SWidget> CreateSingleDefaultValueChannelField(uint8 Channel, const FLinearColor& LabelColor = FLinearColor(0.0f, 0.0f, 0.0f, 0.5f));
+
+	/** Gets the value of the specified default value channel */
+	TOptional<uint8> GetDefaultValueChannelValue(uint8 Channel) const;
+
+	/** Returns the visibility of specified channel of the default value */
+	EVisibility GetDefaultValueChannelVisibility(uint8 Channel) const;
+
+	/** Called when a default value channel changed */
+	void OnDefaultValueChannelValueChanged(uint8 NewValue, uint8 Channel);
+
+	/** Called when value was committed to a default value channel */
+	void OnDefaultValueChannelValueCommitted(uint8 NewValue, ETextCommit::Type CommitType);
+
+private:
+	/** Handle for the FunctionName property */
+	TSharedPtr<IPropertyHandle> FunctionNameHandle;
+
+	/** Handle for the DataType property */
+	TSharedPtr<IPropertyHandle> DataTypeHandle;
+
+	/** Handle for the DefaultValue property */
+	TSharedPtr<IPropertyHandle> DefaultValueHandle;
+
+	/** Handle for the bUseLSB property */
+	TSharedPtr<IPropertyHandle> UseLSBHandle;
+
+	/** The text box to edit the function name */
+	TSharedPtr<SEditableTextBox> FunctionNameEditableTextBox;
+};
diff --git a/Engine/Plugins/VirtualProduction/DMX/DMXEngine/Source/DMXEditor/Private/Customizations/DMXFixtureModeCustomization.cpp b/Engine/Plugins/VirtualProduction/DMX/DMXEngine/Source/DMXEditor/Private/Customizations/DMXFixtureModeCustomization.cpp
new file mode 100644
--- /dev/null
+++ b/Engine/Plugins/VirtualProduction/DMX/DMXEngine/Source/DMXEditor/Private/Customizations/DMXFixtureModeCustomization.cpp
@@ -0,0 +1,214 @@
+// Copyright Epic Games, Inc. All Rights Reserved.
+
+#include "Customizations/DMXFixtureModeCustomization.h"
+
+#include "DMXEditor.h"
+#include "Library/DMXEntityFixtureType.h"
+
+#include "DetailWidgetRow.h"
+#include "IDetailChildrenBuilder.h"
+#include "DetailLayoutBuilder.h"
+#include "Widgets/Input/SEditableTextBox.h"
+
+
+#define LOCTEXT_NAMESPACE "DMXFixtureTypeModeCustomization"
+
+TSharedRef<IPropertyTypeCustomization> FDMXFixtureModeCustomization::MakeInstance()
+{
+	return MakeShared<FDMXFixtureModeCustomization>();
+}
+
+void FDMXFixtureModeCustomization::CustomizeHeader(TSharedRef<IPropertyHandle> InStructPropertyHandle, class FDetailWidgetRow& InHeaderRow, IPropertyTypeCustomizationUtils& InStructCustomizationUtils)
+{
+	constexpr bool bDisplayDefaultPropertyButtons = false;
+
+	InHeaderRow
+		.NameContent()
+		[
+			InStructPropertyHandle->CreatePropertyNameWidget()
+		]
+		.ValueContent()
+		[
+			InStructPropertyHandle->CreatePropertyValueWidget(bDisplayDefaultPropertyButtons)
+		];
+}
+
+void FDMXFixtureModeCustomization::CustomizeChildren(TSharedRef<IPropertyHandle> InStructPropertyHandle, IDetailChildrenBuilder& InStructBuilder, IPropertyTypeCustomizationUtils& InStructCustomizationUtils)
+{
+	// Remember the outer fixture type
+	TArray<UObject*> Outers;
+	InStructPropertyHandle->GetOuterObjects(Outers);
+
+	if (Outers.Num() == 1 && Outers[0]->IsA<UDMXEntityFixtureType>())
+	{
+		OuterFixtureType = Cast<UDMXEntityFixtureType>(Outers[0]);
+	}
+
+	// Retrieve structure's child properties
+	uint32 NumChildren;
+	InStructPropertyHandle->GetNumChildren(NumChildren);
+	
+	TMap<FName, TSharedRef<IPropertyHandle>> PropertyNamePropertyHandleMap;
+	for (uint32 ChildIndex = 0; ChildIndex < NumChildren; ++ChildIndex)
+	{
+		TSharedRef<IPropertyHandle> ChildHandle = InStructPropertyHandle->GetChildHandle(ChildIndex).ToSharedRef();
+		const FName PropertyName = ChildHandle->GetProperty()->GetFName();
+
+		PropertyNamePropertyHandleMap.Add(PropertyName, ChildHandle);
+	}
+
+	// Remember the ModeName property
+	ModeNameHandle = PropertyNamePropertyHandleMap.FindChecked(GET_MEMBER_NAME_CHECKED(FDMXFixtureMode, ModeName));
+
+	// Add properties to the layout
+	for (const TPair<FName, TSharedRef<IPropertyHandle>>& PropertyNamePropertyHandlePair : PropertyNamePropertyHandleMap)
+	{
+		// Don't add the Functions array, it has its own view in the editor.
+		if (PropertyNamePropertyHandlePair.Key == GET_MEMBER_NAME_CHECKED(FDMXFixtureMode, Functions))
+		{
+			continue;
+		}
+	
+		// Hide the reset to default button for all properties
+		PropertyNamePropertyHandlePair.Value->MarkResetToDefaultCustomized();
+
+		// Add the FixtureMatrixConfig struct and get its visibility depending on the bMatrixVisible property
+		if (PropertyNamePropertyHandlePair.Key == GET_MEMBER_NAME_CHECKED(FDMXFixtureMode, FixtureMatrixConfig))
+		{
+			IDetailPropertyRow& FixtureMatrixPropertyRow = InStructBuilder.AddProperty(PropertyNamePropertyHandlePair.Value);
+			FixtureMatrixPropertyRow.Visibility(TAttribute<EVisibility>::Create(TAttribute<EVisibility>::FGetter::CreateSP(this, &FDMXFixtureModeCustomization::GetFixtureMatrixPropertyVisibility)));	
+			FixtureMatrixPropertyRow.ShouldAutoExpand(true);
+			FixtureMatrixPropertyRow.OverrideResetToDefault(FResetToDefaultOverride());
+			continue;
+		}
+
+		if (PropertyNamePropertyHandlePair.Key == GET_MEMBER_NAME_CHECKED(FDMXFixtureMode, ModeName))
+		{
+			// Customize the ModeName property
+			BuildModeNameWidget(InStructBuilder);
+			continue;
+		}
+
+		// Add other properties without further customizing them
+		IDetailPropertyRow& PropertyRow = InStructBuilder.AddProperty(PropertyNamePropertyHandlePair.Value);
+		PropertyRow.OverrideResetToDefault(FResetToDefaultOverride());
+	}
+}
+
+void FDMXFixtureModeCustomization::BuildModeNameWidget(IDetailChildrenBuilder& InStructBuilder)
+{
+	InStructBuilder
+		.AddCustomRow(LOCTEXT("ModeNameWidget", "ModeNameWidget"))
+		.NameContent()
+		[
+			SNew(STextBlock)
+			.Text(LOCTEXT("FixtureModeNameLabel", "Mode Name"))
+			.Font(IDetailLayoutBuilder::GetDetailFont())
+		]
+		.ValueContent()
+		.MaxDesiredWidth(250.0f)
+		[
+			SAssignNew(ModeNameEditableTextBox, SEditableTextBox)
+			.Text(this, &FDMXFixtureModeCustomization::GetModeName)
+			.OnTextChanged(this, &FDMXFixtureModeCustomization::OnModeNameChanged)
+			.OnTextCommitted(this, &FDMXFixtureModeCustomization::OnModeNameCommitted)
+			.Font(IDetailLayoutBuilder::GetDetailFont())
+		];
+}
+
+TArray<FString> FDMXFixtureModeCustomization::GetExistingModeNames() const
+{
+	TArray<FString> ExistingModeNames;
+
+	// Get the parent Modes array property to be able to read the other Modes names inside the current Mode
+	TSharedRef<IPropertyHandle> ModeStructHandle = ModeNameHandle->GetParentHandle().ToSharedRef();
+	TSharedRef<IPropertyHandle> ModesArrayHandle = ModeStructHandle->GetParentHandle().ToSharedRef();
+
+	uint32 NumModes;
+	ModesArrayHandle->GetNumChildren(NumModes);
+
+	for (uint32 ModeIndex = 0; ModeIndex < NumModes; ModeIndex++)
+	{
+		TSharedRef<IPropertyHandle> ModeHandle = ModesArrayHandle->GetChildHandle(ModeIndex).ToSharedRef();
+		TSharedPtr<IPropertyHandle> NameHandle = ModeHandle->GetChildHandle(GET_MEMBER_NAME_CHECKED(FDMXFixtureMode, ModeName));
+
+		FString ModeName;
+		NameHandle->GetValue(ModeName);
+
+		ExistingModeNames.Add(ModeName);
+	}
+
+	return ExistingModeNames;
+}
+
+void FDMXFixtureModeCustomization::OnModeNameChanged(const FText& InNewText)
+{
+	check(ModeNameEditableTextBox.IsValid() && ModeNameHandle.IsValid() && ModeNameHandle->IsValidHandle());
+
+	if (FText::TrimPrecedingAndTrailing(InNewText).IsEmpty())
+	{
+		ModeNameEditableTextBox->SetError(LOCTEXT("ModeNameError_Empty", "The name can't be blank!"));
+		return;
+	}
+
+	FString CurrentName;
+	ModeNameHandle->GetValue(CurrentName);
+
+	const FString& NewName = InNewText.ToString();
+	if (CurrentName.Equals(NewName))
+	{
+		ModeNameEditableTextBox->SetError(FText::GetEmpty());
+		return;
+	}
+
+	const TArray<FString> ExistingNames = GetExistingModeNames();
+	if (ExistingNames.Contains(NewName))
+	{
+		ModeNameEditableTextBox->SetError(LOCTEXT("FixtureModeName_Existent", "This name is already used by another Mode in this mode!"));
+	}
+	else
+	{
+		ModeNameEditableTextBox->SetError(FText::GetEmpty());
+	}
+}
+
+void FDMXFixtureModeCustomization::OnModeNameCommitted(const FText& InNewText, ETextCommit::Type InTextCommit)
+{
+	if (InTextCommit != ETextCommit::OnCleared && !ModeNameEditableTextBox->HasError())
+	{
+		const FString& NewName = InNewText.ToString();
+		SetModeName(NewName);
+	}
+	ModeNameEditableTextBox->SetError(FText::GetEmpty());
+}
+
+FText FDMXFixtureModeCustomization::GetModeName() const
+{
+	FString Name;
+	ModeNameHandle->GetValue(Name);
+
+	return FText::FromString(Name);
+}
+
+void FDMXFixtureModeCustomization::SetModeName(const FString& NewName)
+{
+	if (ModeNameHandle.IsValid())
+	{
+		ModeNameHandle->SetValue(NewName);
+	}
+}
+
+EVisibility FDMXFixtureModeCustomization::GetFixtureMatrixPropertyVisibility() const
+{
+	if (UDMXEntityFixtureType* FixtureType = OuterFixtureType.Get())
+	{
+		if (FixtureType->bFixtureMatrixEnabled)
+		{
+			return EVisibility::Visible;
+		}
+	}
+
+	return EVisibility::Collapsed;
+}
+
+#undef LOCTEXT_NAMESPACE
diff --git a/Engine/Plugins/VirtualProduction/DMX/DMXEngine/Source/DMXEditor/Private/Customizations/DMXFixtureModeCustomization.h b/Engine/Plugins/VirtualProduction/DMX/DMXEngine/Source/DMXEditor/Private/Customizations/DMXFixtureModeCustomization.h
new file mode 100644
--- /dev/null
+++ b/Engine/Plugins/VirtualProduction/DMX/DMXEngine/Source/DMXEditor/Private/Customizations/DMXFixtureModeCustomization.h
@@ -0,0 +1,59 @@
+// Copyright Epic Games, Inc. All Rights Reserved.
+
+#pragma once
+
+#include "CoreMinimal.h"
+#include "IPropertyTypeCustomization.h"
+
+class FDMXEditor;
+class UDMXEntityFixtureType;
+
+class SEditableTextBox;
+
+
+/** 
+ * Customization for the DMXFixtureMode struct 
+ */
+class FDMXFixtureModeCustomization
+	: public IPropertyTypeCustomization
+{
+public:
+	static TSharedRef<IPropertyTypeCustomization> MakeInstance();
+
+protected:
+	//~ Begin IPropertyTypeCustomization interface
+	virtual void CustomizeHeader(TSharedRef<IPropertyHandle> InStructPropertyHandle, class FDetailWidgetRow& InHeaderRow, IPropertyTypeCustomizationUtils& InStructCustomizationUtils) override;
+	virtual void CustomizeChildren(TSharedRef<IPropertyHandle> InStructPropertyHandle, IDetailChildrenBuilder& InStructBuilder, IPropertyTypeCustomizationUtils& InStructCustomizationUtils) override;
+	//~ End IPropertyTypeCustomization interface
+
+private:
+	/** Builds the widget that displays the Mode name */
+	void BuildModeNameWidget(IDetailChildrenBuilder& InStructBuilder);
+	
+	/** Find the existing names for the Mode type being edited within the Fixture Type the Mode resides in */
+	TArray<FString> GetExistingModeNames() const;
+
+	/** Called when the name of a Mode changed */
+	void OnModeNameChanged(const FText& InNewText);
+
+	/** Called when the name of a mode changed */
+	void OnModeNameCommitted(const FText& InNewText, ETextCommit::Type InTextCommit);
+
+	/** Returns the Mode name */
+	FText GetModeName() const;
+
+	/** Changes the Mode name on the fixture properties */
+	void SetModeName(const FString& NewName);
+
+	/** Returns the visibility of the Fixture Matrix property */
+	EVisibility GetFixtureMatrixPropertyVisibility() const;
+
+	/** Handle to the Mode Name property */
+	TSharedPtr<IPropertyHandle> ModeNameHandle;
+
+	/** The text box to edit the Mode name */
+	TSharedPtr<SEditableTextBox> ModeNameEditableTextBox;
+
+	/** The Fixture Type that owns the Mode */
+	TWeakObjectPtr<UDMXEntityFixtureType> OuterFixtureType;
+};
diff --git a/Engine/Plugins/VirtualProduction/DMX/DMXEngine/Source/DMXEditor/Private/Customizations/DMXLibraryPortReferencesCustomization.cpp b/Engine/Plugins/VirtualProduction/DMX/DMXEngine/Source/DMXEditor/Private/Customizations/DMXLibraryPortReferencesCustomization.cpp
--- a/Engine/Plugins/VirtualProduction/DMX/DMXEngine/Source/DMXEditor/Private/Customizations/DMXLibraryPortReferencesCustomization.cpp
+++ b/Engine/Plugins/VirtualProduction/DMX/DMXEngine/Source/DMXEditor/Private/Customizations/DMXLibraryPortReferencesCustomization.cpp
@@ -169,6 +169,10 @@ namespace
 	};
 }
 
+TSharedRef<IPropertyTypeCustomization> FDMXLibraryPortReferencesCustomization::MakeInstance()
+{
+	return MakeShared<FDMXLibraryPortReferencesCustomization>();
+}
 
 void FDMXLibraryPortReferencesCustomization::CustomizeHeader(TSharedRef<IPropertyHandle> StructPropertyHandle, FDetailWidgetRow& HeaderRow, IPropertyTypeCustomizationUtils& StructCustomizationUtils)
 {
diff --git a/Engine/Plugins/VirtualProduction/DMX/DMXEngine/Source/DMXEditor/Private/Customizations/DMXLibraryPortReferencesCustomization.h b/Engine/Plugins/VirtualProduction/DMX/DMXEngine/Source/DMXEditor/Private/Customizations/DMXLibraryPortReferencesCustomization.h
--- a/Engine/Plugins/VirtualProduction/DMX/DMXEngine/Source/DMXEditor/Private/Customizations/DMXLibraryPortReferencesCustomization.h
+++ b/Engine/Plugins/VirtualProduction/DMX/DMXEngine/Source/DMXEditor/Private/Customizations/DMXLibraryPortReferencesCustomization.h
@@ -19,7 +19,7 @@ class FDMXLibraryPortReferencesCustomization
 {
 public:
 	/** Creates an instance of the property type customization */
-	static TSharedRef<IPropertyTypeCustomization> MakeInstance() { return MakeShared<FDMXLibraryPortReferencesCustomization>(); }
+	static TSharedRef<IPropertyTypeCustomization> MakeInstance();
 
 	// ~Begin IPropertyTypeCustomization Interface
 	virtual void CustomizeHeader(TSharedRef<IPropertyHandle> StructPropertyHandle, FDetailWidgetRow& HeaderRow, IPropertyTypeCustomizationUtils& StructCustomizationUtils) override;
diff --git a/Engine/Plugins/VirtualProduction/DMX/DMXEngine/Source/DMXEditor/Private/Customizations/DMXNameListCustomization.h b/Engine/Plugins/VirtualProduction/DMX/DMXEngine/Source/DMXEditor/Private/Customizations/DMXNameListCustomization.h
new file mode 100644
--- /dev/null
+++ b/Engine/Plugins/VirtualProduction/DMX/DMXEngine/Source/DMXEditor/Private/Customizations/DMXNameListCustomization.h
@@ -0,0 +1,160 @@
+// Copyright Epic Games, Inc. All Rights Reserved.
+
+#pragma once
+
+#include "DMXProtocolTypes.h"
+
+#include "CoreMinimal.h"
+#include "DetailWidgetRow.h"
+#include "IPropertyTypeCustomization.h"
+#include "IPropertyUtilities.h"
+#include "Widgets/SNameListPicker.h"
+
+
+/**  
+ * Customization for types that should be displayed as a dropdown of options from a FName array 
+ */
+template<typename TStructType>
+class FDMXNameListCustomization
+	: public IPropertyTypeCustomization
+{
+private:
+	using TNameListType = FDMXNameListCustomization<TStructType>;
+
+public:
+	/** Construction requires a delegate that returns the source list of possible names */
+	FDMXNameListCustomization()
+	{}
+
+	/** Creates an instance of this property type customization */
+	static TSharedRef<IPropertyTypeCustomization> MakeInstance()
+	{
+		return MakeShared<FDMXNameListCustomization<TStructType>>();
+	}
+
+	//~ IPropertyTypeCustomization interface begin
+	virtual void CustomizeHeader(TSharedRef<IPropertyHandle> InPropertyHandle, FDetailWidgetRow& InHeaderRow, IPropertyTypeCustomizationUtils& CustomizationUtils) override
+	{
+		StructPropertyHandle = InPropertyHandle;
+		TSharedPtr<IPropertyUtilities> PropertyUtils = CustomizationUtils.GetPropertyUtilities();
+
+		check(CastFieldChecked<FStructProperty>(StructPropertyHandle->GetProperty())->Struct == TStructType::StaticStruct());
+
+		InHeaderRow
+			.NameContent()
+			[
+				StructPropertyHandle->CreatePropertyNameWidget()
+			]
+			.ValueContent()
+			.MinDesiredWidth(125.0f)
+			.MaxDesiredWidth(0.0f)
+			[
+				SNew(SNameListPicker)
+				.Font(CustomizationUtils.GetRegularFont())
+				.HasMultipleValues(this, &TNameListType::HasMultipleValues)
+				.OptionsSource(MakeAttributeLambda(&TStructType::GetPossibleValues))
+				.UpdateOptionsDelegate(&TStructType::OnValuesChanged)
+				.IsValid(this, &TNameListType::HideWarningIcon)
+				.Value(this, &TNameListType::GetValue)
+				.bCanBeNone(TStructType::bCanBeNone)
+				.bDisplayWarningIcon(true)
+				.OnValueChanged(this, &TNameListType::SetValue)
+			]
+		.IsEnabled(MakeAttributeLambda([=] { return !InPropertyHandle->IsEditConst() && PropertyUtils->IsPropertyEditingEnabled(); }));
+	}
+
+	virtual void CustomizeChildren(TSharedRef<IPropertyHandle> InPropertyHandle, IDetailChildrenBuilder& InChildBuilder, IPropertyTypeCustomizationUtils& CustomizationUtils) override {};
+	//~ IPropertyTypeCustomization interface end
+
+private:
+	FName GetValue() const
+	{
+		TArray<const void*> RawData;
+		StructPropertyHandle->AccessRawData(RawData);
+
+		for (const void* RawPtr : RawData)
+		{
+			if (RawPtr != nullptr)
+			{
+				// The types we use with this customization must have a cast constructor to FName
+				return reinterpret_cast<const TStructType*>(RawPtr)->GetName();
+			}
+		}
+
+		return FName();
+	}
+
+	void SetValue(FName NewValue)
+	{
+		FStructProperty* StructProperty = CastFieldChecked<FStructProperty>(StructPropertyHandle->GetProperty());
+
+		TArray<void*> RawData;
+		StructPropertyHandle->AccessRawData(RawData);
+
+		for (void* SingleRawData : RawData)
+		{
+			TStructType* PreviousValue = reinterpret_cast<TStructType*>(SingleRawData);
+			TStructType NewProtocolName;
+			NewProtocolName.SetFromName(NewValue);
+
+			// Export new value to text format that can be imported later
+			FString TextValue;
+			StructProperty->Struct->ExportText(TextValue, &NewProtocolName, PreviousValue, nullptr, EPropertyPortFlags::PPF_None, nullptr);
+
+			// Set values on edited property handle from exported text
+			ensure(StructPropertyHandle->SetValueFromFormattedString(TextValue, EPropertyValueSetFlags::DefaultFlags) == FPropertyAccess::Result::Success);
+		}
+	}
+
+	bool HasMultipleValues() const
+	{
+		TArray<const void*> RawData;
+		StructPropertyHandle->AccessRawData(RawData);
+
+		TOptional<TStructType> CompareAgainst;
+		for (const void* RawPtr : RawData)
+		{
+			if (RawPtr == nullptr)
+			{
+				if (CompareAgainst.IsSet())
+				{
+					return false;
+				}
+			}
+			else
+			{
+				const TStructType* ThisValue = reinterpret_cast<const TStructType*>(RawPtr);
+
+				if (!CompareAgainst.IsSet())
+				{
+					CompareAgainst = *ThisValue;
+				}
+				else if (!(*ThisValue == CompareAgainst.GetValue()))
+				{
+					return true;
+				}
+			}
+		}
+
+		return false;
+	}
+
+	bool HideWarningIcon() const
+	{
+		if (HasMultipleValues())
+		{
+			return true;
+		}
+
+		const FName CurrentValue = GetValue();
+		if (CurrentValue.IsEqual(FDMXNameListItem::None))
+		{
+			return true;
+		}
+
+		return TStructType::IsValid(GetValue());
+	}
+
+private:
+	TSharedPtr<IPropertyHandle> StructPropertyHandle;
+};
diff --git a/Engine/Plugins/VirtualProduction/DMX/DMXEngine/Source/DMXEditor/Private/Customizations/DMXPixelMappingDistributionCustomization.cpp b/Engine/Plugins/VirtualProduction/DMX/DMXEngine/Source/DMXEditor/Private/Customizations/DMXPixelMappingDistributionCustomization.cpp
new file mode 100644
--- /dev/null
+++ b/Engine/Plugins/VirtualProduction/DMX/DMXEngine/Source/DMXEditor/Private/Customizations/DMXPixelMappingDistributionCustomization.cpp
@@ -0,0 +1,91 @@
+// Copyright Epic Games, Inc. All Rights Reserved.
+
+#include "Customizations/DMXPixelMappingDistributionCustomization.h"
+
+#include "DMXEditorStyle.h"
+#include "DMXProtocolTypes.h"
+
+#include "DetailWidgetRow.h"
+#include "PropertyHandle.h"
+#include "Widgets/Images/SImage.h"
+#include "Widgets/Input/SButton.h"
+#include "Widgets/Layout/SUniformGridPanel.h"
+
+
+TSharedRef<IPropertyTypeCustomization> FDMXPixelMappingDistributionCustomization::MakeInstance()
+{
+	return MakeShared<FDMXPixelMappingDistributionCustomization>();
+}
+
+void FDMXPixelMappingDistributionCustomization::CustomizeHeader(TSharedRef<IPropertyHandle> InStructPropertyHandle, FDetailWidgetRow& InHeaderRow, IPropertyTypeCustomizationUtils& CustomizationUtils)
+{
+	EnumPropertyHandle = InStructPropertyHandle;
+
+	TSharedRef<SUniformGridPanel> DistributionGridPanel = 
+		SNew(SUniformGridPanel)
+		.SlotPadding(FMargin(1.f));
+
+	for (int32 XIndex = 0; XIndex < DistributionGridNumXPanels; ++XIndex)
+	{
+		for (int32 YIndex = 0; YIndex < DistributionGridNumYPanels; ++YIndex)
+		{
+			FString BrushPath = FString::Printf(TEXT("DMXEditor.PixelMapping.DistributionGrid.%d.%d"), XIndex, YIndex);
+
+			TSharedPtr<SButton> Button = SNew(SButton)
+				.ButtonColorAndOpacity(TAttribute<FSlateColor>::Create(TAttribute<FSlateColor>::FGetter::CreateSP(this, &FDMXPixelMappingDistributionCustomization::GetButtonColorAndOpacity, XIndex, YIndex)))
+				.OnClicked(FOnClicked::CreateSP(this, &FDMXPixelMappingDistributionCustomization::OnGridButtonClicked, XIndex, YIndex))
+				[
+					SNew(SImage)
+					.Image(FDMXEditorStyle::Get().GetBrush(*BrushPath))
+				];
+
+				DistributionGridPanel->AddSlot(XIndex, YIndex)
+				[
+					Button.ToSharedRef()
+				];
+		}
+	}
+
+	InHeaderRow
+		.NameContent()
+		[
+			EnumPropertyHandle->CreatePropertyNameWidget()
+		]
+		.ValueContent()
+		.MinDesiredWidth(200.0f)
+		.MaxDesiredWidth(400.0f)
+		[
+			DistributionGridPanel
+		];
+}
+
+FReply FDMXPixelMappingDistributionCustomization::OnGridButtonClicked(int32 GridIndexX, int32 GridIndexY)
+{
+	if (EnumPropertyHandle.IsValid())
+	{
+		uint8 ChoosenDistribution = (GridIndexX * DistributionGridNumXPanels + GridIndexY);
+
+		EnumPropertyHandle->NotifyPreChange();
+		EnumPropertyHandle->SetValue(ChoosenDistribution);
+		EnumPropertyHandle->NotifyPostChange(EPropertyChangeType::ValueSet);
+	}
+
+	return FReply::Handled();
+}
+
+FSlateColor FDMXPixelMappingDistributionCustomization::GetButtonColorAndOpacity(int32 GridIndexX, int32 GridIndexY)
+{
+	if (EnumPropertyHandle.IsValid())
+	{
+		uint8 CurrentDistribution = 0;
+		if (EnumPropertyHandle->GetValue(CurrentDistribution) == FPropertyAccess::Result::Success)
+		{
+			if (CurrentDistribution == (GridIndexX * DistributionGridNumXPanels + GridIndexY))
+			{
+				return FLinearColor(0.2f, 0.2f, 0.2f, 1.f);
+			}
+		}
+	}
+
+	return FLinearColor::Transparent;
+}
diff --git a/Engine/Plugins/VirtualProduction/DMX/DMXEngine/Source/DMXEditor/Private/Customizations/DMXPixelMappingDistributionCustomization.h b/Engine/Plugins/VirtualProduction/DMX/DMXEngine/Source/DMXEditor/Private/Customizations/DMXPixelMappingDistributionCustomization.h
new file mode 100644
--- /dev/null
+++ b/Engine/Plugins/VirtualProduction/DMX/DMXEngine/Source/DMXEditor/Private/Customizations/DMXPixelMappingDistributionCustomization.h
@@ -0,0 +1,41 @@
+// Copyright Epic Games, Inc. All Rights Reserved.
+
+#pragma once
+
+#include "DMXProtocolTypes.h"
+
+#include "IPropertyTypeCustomization.h"
+
+
+/**
+ * Customization for the DMXPixelMappingDistribution enum
+ */
+class FDMXPixelMappingDistributionCustomization
+	: public IPropertyTypeCustomization
+{
+public:
+	/** Creates an instance of this property type customization */
+	static TSharedRef<IPropertyTypeCustomization> MakeInstance();
+
+	//~ IPropertyTypeCustomization interface begin
+	virtual void CustomizeHeader(TSharedRef<IPropertyHandle> InStructPropertyHandle, FDetailWidgetRow& InHeaderRow, IPropertyTypeCustomizationUtils& CustomizationUtils) override;
+	virtual void CustomizeChildren(TSharedRef<IPropertyHandle> InStructPropertyHandle, IDetailChildrenBuilder& InStructBuilder, IPropertyTypeCustomizationUtils& InStructCustomizationUtils) override {}
+	//~ IPropertyTypeCustomization interface end
+
+protected:
+	/** Gets the color and opacity of the enum value buttons */
+	FSlateColor GetButtonColorAndOpacity(int32 GridIndexX, int32 GridIndexY);
+
+	/** Called when an enum value button in the grid was clicked */
+	FReply OnGridButtonClicked(int32 GridIndexX, int32 GridIndexY);
+
+private:
+	/** Number of panels in X direction */
+	static const uint8 DistributionGridNumXPanels = 4;
+
+	/** Number of buttons in Y direction */
+	static const uint8 DistributionGridNumYPanels = 4;
+
+	/** Handle for the enum property that is customized */
+	TSharedPtr<IPropertyHandle> EnumPropertyHandle;
+};
diff --git a/Engine/Plugins/VirtualProduction/DMX/DMXEngine/Source/DMXEditor/Private/Customizations/TakeRecorderDMXLibrarySourceEditorCustomization.cpp b/Engine/Plugins/VirtualProduction/DMX/DMXEngine/Source/DMXEditor/Private/Customizations/TakeRecorderDMXLibrarySourceEditorCustomization.cpp
new file mode 100644
--- /dev/null
+++ b/Engine/Plugins/VirtualProduction/DMX/DMXEngine/Source/DMXEditor/Private/Customizations/TakeRecorderDMXLibrarySourceEditorCustomization.cpp
@@ -0,0 +1,117 @@
+// Copyright Epic Games, Inc. All Rights Reserved.
+
+#include "Customizations/TakeRecorderDMXLibrarySourceEditorCustomization.h"
+#include "Sequencer/TakeRecorderDMXLibrarySource.h"
+#include "Library/DMXLibrary.h"
+
+#include "DetailCategoryBuilder.h"
+#include "DetailLayoutBuilder.h"
+#include "DetailWidgetRow.h"
+#include "IDetailPropertyRow.h"
+#include "Input/Reply.h"
+#include "PropertyHandle.h"
+#include "ScopedTransaction.h"
+#include "Widgets/Input/SButton.h"
+#include "Widgets/Text/STextBlock.h"
+#include "Widgets/SNullWidget.h"
+
+#define LOCTEXT_NAMESPACE "FTakeRecorderDMXLibrarySourceEditorCustomization"
+
+TSharedRef<IPropertyTypeCustomization> FTakeRecorderDMXLibrarySourceEditorCustomization::MakeInstance()
+{
+	return MakeShared<FTakeRecorderDMXLibrarySourceEditorCustomization>();
+}
+
+void FTakeRecorderDMXLibrarySourceEditorCustomization::CustomizeHeader(TSharedRef<IPropertyHandle> PropertyHandle, FDetailWidgetRow& HeaderRow, IPropertyTypeCustomizationUtils& CustomizationUtils)
+{
+	StructHandle = PropertyHandle;
+
+	// Cache the DMX Library TakeRecorder objects being displayed by the DetailsView that created this customization
+	TArray<UObject*> CustomizedObjects;
+	PropertyHandle->GetOuterObjects(CustomizedObjects);
+	CustomizedDMXRecorders.Reserve(CustomizedObjects.Num());
+
+	for (UObject* Object : CustomizedObjects)
+	{
+		if (Object == nullptr || !Object->IsValidLowLevelFast())
+		{
+			continue;
+		}
+
+		UTakeRecorderDMXLibrarySource* DMXSource = Cast<UTakeRecorderDMXLibrarySource>(Object);
+		if (DMXSource != nullptr)
+		{
+			CustomizedDMXRecorders.Emplace(DMXSource);
+		}
+	}
+	check(CustomizedDMXRecorders.Num());
+
+	// Create the "Add all patches" button row
+	HeaderRow
+		.NameContent()
+		[
+			SNullWidget::NullWidget // Empty label 'cause the button already has its label
+		]
+		.ValueContent()
+		.MinDesiredWidth(0.0f)
+		.MaxDesiredWidth(0.0f)
+		.HAlign(HAlign_Left)
+		[
+			SNew(SButton)
+			.IsEnabled(this, &FTakeRecorderDMXLibrarySourceEditorCustomization::GetAddAllEnabled)
+			.OnClicked(this, &FTakeRecorderDMXLibrarySourceEditorCustomization::HandleOnClicked)
+			.ToolTipText(LOCTEXT("ToolTip", "Add all Patches in the DMX Library to be recorded"))
+			.Content()
+			[
+				SNew(STextBlock)
+				.Font(IDetailLayoutBuilder::GetDetailFont())
+				.Text(LOCTEXT("AddAllPatches", "Add all Fixture Patches"))
+			]
+		];
+}
+
+bool FTakeRecorderDMXLibrarySourceEditorCustomization::GetAddAllEnabled() const
+{
+	// All selected DMX recorders must have valid DMX Libraries to enable the button
+	for (TWeakObjectPtr<UTakeRecorderDMXLibrarySource> DMXRecorder : CustomizedDMXRecorders)
+	{
+		if (!DMXRecorder.IsValid() || !DMXRecorder->IsValidLowLevelFast()
+			|| DMXRecorder->DMXLibrary == nullptr || !DMXRecorder->DMXLibrary->IsValidLowLevelFast())
+		{
+			return false;
+		}
+	}
+
+	return true;
+}
+
+FReply FTakeRecorderDMXLibrarySourceEditorCustomization::HandleOnClicked()
+{
+	check(StructHandle.IsValid() && StructHandle->IsValidHandle());
+
+	// Scope to end the Transaction before calling NotifyFinishedChangingProperties,
+	// like it's done when using IPropertyHandle::SetValue
+	{
+		const FScopedTransaction Transaction(LOCTEXT("Transaction", "Add all Fixture Patches for recording"));
+
+		// Notify objects that own this struct so they can call Modify to record their current states
+		StructHandle->NotifyPreChange();
+
+		// Add all patches on each selected DMX Take Recorder
+		for (TWeakObjectPtr<UTakeRecorderDMXLibrarySource> DMXRecorder : CustomizedDMXRecorders)
+		{
+			if (DMXRecorder.IsValid() && DMXRecorder->IsValidLowLevelFast())
+			{
+				DMXRecorder->AddAllPatches();
+			}
+		}
+
+		StructHandle->NotifyPostChange(EPropertyChangeType::ValueSet);
+	}
+
+	StructHandle->NotifyFinishedChangingProperties();
+
+	return FReply::Handled();
+}
+
+#undef LOCTEXT_NAMESPACE
diff --git a/Engine/Plugins/VirtualProduction/DMX/DMXEngine/Source/DMXEditor/Private/Customizations/TakeRecorderDMXLibrarySourceEditorCustomization.h b/Engine/Plugins/VirtualProduction/DMX/DMXEngine/Source/DMXEditor/Private/Customizations/TakeRecorderDMXLibrarySourceEditorCustomization.h
new file mode 100644
--- /dev/null
+++ b/Engine/Plugins/VirtualProduction/DMX/DMXEngine/Source/DMXEditor/Private/Customizations/TakeRecorderDMXLibrarySourceEditorCustomization.h
@@ -0,0 +1,40 @@
+// Copyright Epic Games, Inc. All Rights Reserved.
+
+#pragma once
+
+#include "CoreMinimal.h"
+#include "IPropertyTypeCustomization.h"
+#include "UObject/WeakObjectPtrTemplates.h"
+
+class FReply;
+class IPropertyHandle;
+class UTakeRecorderDMXLibrarySource;
+
+class FTakeRecorderDMXLibrarySourceEditorCustomization
+	: public IPropertyTypeCustomization
+{
+public:
+	/** Creates an instance of this property type customization */
+	static TSharedRef<IPropertyTypeCustomization> MakeInstance();
+
+	//~ Begin IPropertyTypeCustomization interface
+	virtual void CustomizeHeader(TSharedRef<IPropertyHandle> PropertyHandle, FDetailWidgetRow& HeaderRow, IPropertyTypeCustomizationUtils& CustomizationUtils) override;
+	virtual void CustomizeChildren(TSharedRef<IPropertyHandle> PropertyHandle, IDetailChildrenBuilder& ChildBuilder, IPropertyTypeCustomizationUtils& CustomizationUtils) override {}
+	//~ End IPropertyTypeCustomization interface
+
+private:
+	/** Checks whether all selected DMX Take Recorders have valid DMX Libraries */
+	bool GetAddAllEnabled() const;
+	/** Add all Patches to be recorded on each DMX Take Recorder */
+	FReply HandleOnClicked();
+
+private:
+	/**
+	 * Cache of the TakeRecorder sources being customized by the DetailsView that
+	 * created this customization
+	 */
+	TArray<TWeakObjectPtr<UTakeRecorderDMXLibrarySource>> CustomizedDMXRecorders;
+
+	/** Used to generate Change events on the root DMX Take Recorder objects */
+	TSharedPtr<IPropertyHandle> StructHandle;
+};
diff --git a/Engine/Plugins/VirtualProduction/DMX/DMXEngine/Source/DMXEditor/Private/DMXEditorModule.cpp b/Engine/Plugins/VirtualProduction/DMX/DMXEngine/Source/DMXEditor/Private/DMXEditorModule.cpp
--- a/Engine/Plugins/VirtualProduction/DMX/DMXEngine/Source/DMXEditor/Private/DMXEditorModule.cpp
+++ b/Engine/Plugins/VirtualProduction/DMX/DMXEngine/Source/DMXEditor/Private/DMXEditorModule.cpp
@@ -6,25 +6,29 @@
 #include "DMXEditor.h"
 #include "DMXEditorStyle.h"
 #include "DMXEditorTabNames.h"
-#include "DMXProtocolTypes.h"
 #include "DMXProtocolBlueprintLibrary.h"
-#include "Library/DMXLibrary.h"
-#include "Library/DMXEntityReference.h"
-#include "Library/DMXEntity.h"
-#include "Game/DMXComponent.h"
-#include "Commands/DMXEditorCommands.h"
+#include "DMXProtocolTypes.h"
 #include "AssetTools/AssetTypeActions_DMXEditorLibrary.h"
-#include "Customizations/DMXEditorPropertyEditorCustomization.h"
+#include "Customizations/DMXEntityReferenceCustomization.h"
+#include "Customizations/DMXFixtureFunctionCustomization.h"
+#include "Customizations/DMXFixtureModeCustomization.h"
 #include "Customizations/DMXLibraryPortReferencesCustomization.h"
+#include "Customizations/DMXNameListCustomization.h"
+#include "Customizations/DMXPixelMappingDistributionCustomization.h"
+#include "Customizations/TakeRecorderDMXLibrarySourceEditorCustomization.h"
+#include "Commands/DMXEditorCommands.h"
+#include "Game/DMXComponent.h"
+#include "Library/DMXEntityReference.h"
+#include "Library/DMXEntity.h"
+#include "Library/DMXEntityFixtureType.h"
+#include "Library/DMXLibrary.h"
 #include "Sequencer/DMXLibraryTrackEditor.h"
 #include "Sequencer/TakeRecorderDMXLibrarySource.h"
-#include "Sequencer/Customizations/TakeRecorderDMXLibrarySourceEditorCustomization.h"
 #include "Widgets/Monitors/SDMXActivityMonitor.h"
 #include "Widgets/Monitors/SDMXChannelsMonitor.h"
 #include "Widgets/OutputConsole/SDMXOutputConsole.h"
 #include "Widgets/PatchTool/SDMXPatchTool.h"
 
-#include "Templates/SharedPointer.h"
 #include "LevelEditor.h"
 #include "AssetToolsModule.h"
 #include "PropertyEditorModule.h"
@@ -34,7 +38,7 @@
 #include "Framework/Commands/UICommandList.h"
 #include "Framework/Docking/TabManager.h"
 #include "Framework/MultiBox/MultiBoxBuilder.h"
-#include "Toolkits/GlobalEditorCommonCommands.h"
+#include "Templates/SharedPointer.h"
 #include "Widgets/SBoxPanel.h"
 #include "Widgets/SOverlay.h"
 #include "Widgets/Docking/SDockTab.h"
@@ -217,21 +221,33 @@ void FDMXEditorModule::RegisterAssetTypeActions()
 
 void FDMXEditorModule::RegisterPropertyTypeCustomizations()
 {
-	// Customizations for the name lists of our custom types, like Fixture Categories
-	RegisterCustomPropertyTypeLayout(FDMXProtocolName::StaticStruct()->GetFName(), FOnGetPropertyTypeCustomizationInstance::CreateStatic(&FDMXCustomizationFactory::MakeInstance<FNameListCustomization<FDMXProtocolName>>));
-	RegisterCustomPropertyTypeLayout(FDMXFixtureCategory::StaticStruct()->GetFName(), FOnGetPropertyTypeCustomizationInstance::CreateStatic(&FDMXCustomizationFactory::MakeInstance<FNameListCustomization<FDMXFixtureCategory>>));
-	RegisterCustomPropertyTypeLayout(FDMXAttributeName::StaticStruct()->GetFName(), FOnGetPropertyTypeCustomizationInstance::CreateStatic(&FDMXCustomizationFactory::MakeInstance<FNameListCustomization<FDMXAttributeName>>));
-	RegisterCustomPropertyTypeLayout("EDMXPixelMappingDistribution", FOnGetPropertyTypeCustomizationInstance::CreateStatic(&FDMXCustomizationFactory::MakeInstance<FDMXPixelMappingDistributionCustomization>));
+	FPropertyEditorModule& PropertyModule = FModuleManager::GetModuleChecked<FPropertyEditorModule>("PropertyEditor");
+
+	// Property type customization for the DMXFixtureMode struct
+	FOnGetPropertyTypeCustomizationInstance FixtureModePropertyTypeCustomizationGetter = FOnGetPropertyTypeCustomizationInstance::CreateStatic(&FDMXFixtureModeCustomization::MakeInstance);
+	PropertyModule.RegisterCustomPropertyTypeLayout(FDMXFixtureMode::StaticStruct()->GetFName(), FixtureModePropertyTypeCustomizationGetter);
+
+	// Property type customization for the DMXFixtureFunction struct
+	FOnGetPropertyTypeCustomizationInstance FixtureFunctionPropertyTypeCustomizationGetter = FOnGetPropertyTypeCustomizationInstance::CreateStatic(&FDMXFixtureFunctionCustomization::MakeInstance);
+	PropertyModule.RegisterCustomPropertyTypeLayout(FDMXFixtureFunction::StaticStruct()->GetFName(), FixtureFunctionPropertyTypeCustomizationGetter);
+
+	// Property type customization for the EDMXPixelMappingDistribution enum
+	RegisterCustomPropertyTypeLayout("EDMXPixelMappingDistribution", FOnGetPropertyTypeCustomizationInstance::CreateStatic(&FDMXPixelMappingDistributionCustomization::MakeInstance));
+
+	// Property type customization for Name List structs
+	RegisterCustomPropertyTypeLayout(FDMXProtocolName::StaticStruct()->GetFName(), FOnGetPropertyTypeCustomizationInstance::CreateStatic(&FDMXNameListCustomization<FDMXProtocolName>::MakeInstance));
+	RegisterCustomPropertyTypeLayout(FDMXFixtureCategory::StaticStruct()->GetFName(), FOnGetPropertyTypeCustomizationInstance::CreateStatic(&FDMXNameListCustomization<FDMXFixtureCategory>::MakeInstance));
+	RegisterCustomPropertyTypeLayout(FDMXAttributeName::StaticStruct()->GetFName(), FOnGetPropertyTypeCustomizationInstance::CreateStatic(&FDMXNameListCustomization<FDMXAttributeName>::MakeInstance));
 
-	// Customizations for the Entity Reference types
-	RegisterCustomPropertyTypeLayout(FDMXEntityFixtureTypeRef::StaticStruct()->GetFName(), FOnGetPropertyTypeCustomizationInstance::CreateStatic(&FDMXCustomizationFactory::MakeInstance<FDMXEntityReferenceCustomization>));
-	RegisterCustomPropertyTypeLayout(FDMXEntityFixturePatchRef::StaticStruct()->GetFName(), FOnGetPropertyTypeCustomizationInstance::CreateStatic(&FDMXCustomizationFactory::MakeInstance<FDMXEntityReferenceCustomization>));
+	// Customizations for the Entity Reference structs
+	RegisterCustomPropertyTypeLayout(FDMXEntityFixtureTypeRef::StaticStruct()->GetFName(), FOnGetPropertyTypeCustomizationInstance::CreateStatic(&FDMXEntityReferenceCustomization::MakeInstance));
+	RegisterCustomPropertyTypeLayout(FDMXEntityFixturePatchRef::StaticStruct()->GetFName(), FOnGetPropertyTypeCustomizationInstance::CreateStatic(&FDMXEntityReferenceCustomization::MakeInstance));
 
 	// Customization for the DMXLibraryPortReferences struct type
-	RegisterCustomPropertyTypeLayout(FDMXLibraryPortReferences::StaticStruct()->GetFName(), FOnGetPropertyTypeCustomizationInstance::CreateStatic(&FDMXCustomizationFactory::MakeInstance<FDMXLibraryPortReferencesCustomization>));
+	RegisterCustomPropertyTypeLayout(FDMXLibraryPortReferences::StaticStruct()->GetFName(), FOnGetPropertyTypeCustomizationInstance::CreateStatic(&FDMXLibraryPortReferencesCustomization::MakeInstance));
 
 	// DMXLibrary TakeRecorder AddAllPatchesButton customization
-	RegisterCustomPropertyTypeLayout(FAddAllPatchesButton::StaticStruct()->GetFName(), FOnGetPropertyTypeCustomizationInstance::CreateStatic(&FDMXCustomizationFactory::MakeInstance<FDMXLibraryRecorderAddAllPatchesButtonCustomization>));
+	RegisterCustomPropertyTypeLayout(FAddAllPatchesButton::StaticStruct()->GetFName(), FOnGetPropertyTypeCustomizationInstance::CreateStatic(&FTakeRecorderDMXLibrarySourceEditorCustomization::MakeInstance));
 }
 
 void FDMXEditorModule::RegisterSequencerTypes()
diff --git a/Engine/Plugins/VirtualProduction/DMX/DMXEngine/Source/DMXEditor/Private/Sequencer/Customizations/TakeRecorderDMXLibrarySourceEditorCustomization.cpp b/Engine/Plugins/VirtualProduction/DMX/DMXEngine/Source/DMXEditor/Private/Sequencer/Customizations/TakeRecorderDMXLibrarySourceEditorCustomization.cpp
deleted file mode 100644
diff --git a/Engine/Plugins/VirtualProduction/DMX/DMXEngine/Source/DMXEditor/Private/Sequencer/Customizations/TakeRecorderDMXLibrarySourceEditorCustomization.h b/Engine/Plugins/VirtualProduction/DMX/DMXEngine/Source/DMXEditor/Private/Sequencer/Customizations/TakeRecorderDMXLibrarySourceEditorCustomization.h
deleted file mode 100644
diff --git a/Engine/Plugins/VirtualProduction/DMX/DMXEngine/Source/DMXEditor/Private/Widgets/FixturePatch/SDMXFixturePatchInspector.cpp b/Engine/Plugins/VirtualProduction/DMX/DMXEngine/Source/DMXEditor/Private/Widgets/FixturePatch/SDMXFixturePatchInspector.cpp
--- a/Engine/Plugins/VirtualProduction/DMX/DMXEngine/Source/DMXEditor/Private/Widgets/FixturePatch/SDMXFixturePatchInspector.cpp
+++ b/Engine/Plugins/VirtualProduction/DMX/DMXEngine/Source/DMXEditor/Private/Widgets/FixturePatch/SDMXFixturePatchInspector.cpp
@@ -4,8 +4,7 @@
 
 #include "DMXEditor.h"
 #include "DMXFixturePatchSharedData.h"
-
-#include "Customizations/DMXEditorPropertyEditorCustomization.h"
+#include "Customizations/DMXEntityFixturePatchDetails.h"
 #include "Library/DMXEntityFixturePatch.h"
 
 
@@ -27,8 +26,8 @@ void SDMXFixturePatchInspector::Construct(const FArguments& InArgs)
 	}
 
 	// Register customization for UOBJECT
-	FOnGetDetailCustomizationInstance FixturePatchDetails = FOnGetDetailCustomizationInstance::CreateStatic(&FDMXCustomizationFactory::MakeInstance<FDMXFixturePatchesDetails>, InArgs._DMXEditor);
-	GetPropertyView()->RegisterInstancedCustomPropertyLayout(UDMXEntityFixturePatch::StaticClass(), FixturePatchDetails);
+	FOnGetDetailCustomizationInstance FixturePatchDetailCustomizationGetter = FOnGetDetailCustomizationInstance::CreateStatic(&FDMXEntityFixturePatchDetails::MakeInstance, InArgs._DMXEditor);
+	GetPropertyView()->RegisterInstancedCustomPropertyLayout(UDMXEntityFixturePatch::StaticClass(), FixturePatchDetailCustomizationGetter);
 }
 
 void SDMXFixturePatchInspector::OnFixturePatchesSelected()
diff --git a/Engine/Plugins/VirtualProduction/DMX/DMXEngine/Source/DMXEditor/Private/Widgets/SDMXEntityInspector.cpp b/Engine/Plugins/VirtualProduction/DMX/DMXEngine/Source/DMXEditor/Private/Widgets/SDMXEntityInspector.cpp
--- a/Engine/Plugins/VirtualProduction/DMX/DMXEngine/Source/DMXEditor/Private/Widgets/SDMXEntityInspector.cpp
+++ b/Engine/Plugins/VirtualProduction/DMX/DMXEngine/Source/DMXEditor/Private/Widgets/SDMXEntityInspector.cpp
@@ -6,7 +6,6 @@
 #include "DMXEditorLog.h"
 #include "DMXProtocolConstants.h"
 
-#include "Customizations/DMXEditorPropertyEditorCustomization.h"
 #include "Customizations/DMXEntityFixtureTypeFixtureSettingsDetails.h"
 #include "Customizations/DMXEntityFixtureTypeModesDetails.h"
 #include "Customizations/DMXEntityFixtureTypeModePropertiesDetails.h"
@@ -140,33 +139,25 @@ void SDMXEntityInspectorFixtureTypes::Construct(const FArguments& InArgs, EDMXFi
 	switch (Layout)
 	{
 	case EDMXFixtureTypeLayout::FixtureSettings:
-		FixtureTypeDetails = FOnGetDetailCustomizationInstance::CreateStatic(&FDMXCustomizationFactory::MakeInstance<FDMXEntityFixtureTypeFixtureSettingsDetails>, InArgs._DMXEditor);
+		FixtureTypeDetails = FOnGetDetailCustomizationInstance::CreateStatic(&FDMXEntityFixtureTypeFixtureSettingsDetails::MakeInstance, InArgs._DMXEditor);
 		break;
 	case EDMXFixtureTypeLayout::Modes:
-		FixtureTypeDetails = FOnGetDetailCustomizationInstance::CreateStatic(&FDMXCustomizationFactory::MakeInstance<FDMXEntityFixtureTypeModesDetails>, InArgs._DMXEditor);
+		FixtureTypeDetails = FOnGetDetailCustomizationInstance::CreateStatic(&FDMXEntityFixtureTypeModesDetails::MakeInstance, InArgs._DMXEditor);
 		break;
 	case EDMXFixtureTypeLayout::ModeProperties:
-		FixtureTypeDetails = FOnGetDetailCustomizationInstance::CreateStatic(&FDMXCustomizationFactory::MakeInstance<FDMXEntityFixtureTypeModePropertiesDetails>, InArgs._DMXEditor);
+		FixtureTypeDetails = FOnGetDetailCustomizationInstance::CreateStatic(&FDMXEntityFixtureTypeModePropertiesDetails::MakeInstance, InArgs._DMXEditor);
 		break;
 	case EDMXFixtureTypeLayout::Functions:
-		FixtureTypeDetails = FOnGetDetailCustomizationInstance::CreateStatic(&FDMXCustomizationFactory::MakeInstance<FDMXEntityFixtureTypeFunctionsDetails>, InArgs._DMXEditor);
+		FixtureTypeDetails = FOnGetDetailCustomizationInstance::CreateStatic(&FDMXEntityFixtureTypeFunctionsDetails::MakeInstance, InArgs._DMXEditor);
 		break;
 	case EDMXFixtureTypeLayout::FunctionProperties:
-		FixtureTypeDetails = FOnGetDetailCustomizationInstance::CreateStatic(&FDMXCustomizationFactory::MakeInstance<FDMXEntityFixtureTypeFunctionPropertiesDetails>, InArgs._DMXEditor);
+		FixtureTypeDetails = FOnGetDetailCustomizationInstance::CreateStatic(&FDMXEntityFixtureTypeFunctionPropertiesDetails::MakeInstance, InArgs._DMXEditor);
 		break;
 	default:
 		check(0); // Unhandled layout
 		return;
 	}
 	GetPropertyView()->RegisterInstancedCustomPropertyLayout(UDMXEntityFixtureType::StaticClass(), FixtureTypeDetails);
-
-	// Register customization for Fixture Mode USTRUCT
-	FOnGetPropertyTypeCustomizationInstance FixtureModeDetails = FOnGetPropertyTypeCustomizationInstance::CreateStatic(&FDMXCustomizationFactory::MakeInstance<FDMXFixtureModeDetails>, InArgs._DMXEditor);
-	GetPropertyView()->RegisterInstancedCustomPropertyTypeLayout(FDMXFixtureMode::StaticStruct()->GetFName(), FixtureModeDetails);
-
-	// Register customization for Fixture Function USTRUCT
-	FOnGetPropertyTypeCustomizationInstance FixtureFunctionDetails = FOnGetPropertyTypeCustomizationInstance::CreateStatic(&FDMXCustomizationFactory::MakeInstance<FDMXFixtureFunctionDetails>, InArgs._DMXEditor);
-	GetPropertyView()->RegisterInstancedCustomPropertyTypeLayout(FDMXFixtureFunction::StaticStruct()->GetFName(), FixtureFunctionDetails);
 }
 
 #undef LOCTEXT_NAMESPACE
diff --git a/Engine/Plugins/VirtualProduction/DMX/DMXEngine/Source/DMXRuntime/Private/Library/DMXEntityFixtureType.cpp b/Engine/Plugins/VirtualProduction/DMX/DMXEngine/Source/DMXRuntime/Private/Library/DMXEntityFixtureType.cpp
--- a/Engine/Plugins/VirtualProduction/DMX/DMXEngine/Source/DMXRuntime/Private/Library/DMXEntityFixtureType.cpp
+++ b/Engine/Plugins/VirtualProduction/DMX/DMXEngine/Source/DMXRuntime/Private/Library/DMXEntityFixtureType.cpp
@@ -636,18 +636,18 @@ void UDMXEntityFixtureType::PostEditChangeChainProperty(FPropertyChangedChainEve
 	if (PropertyName == GET_MEMBER_NAME_CHECKED(UDMXEntityFixtureType, Modes))
 	{
 		if (PropertyChangedEvent.ChangeType == EPropertyChangeType::ArrayRemove || PropertyChangedEvent.ChangeType == EPropertyChangeType::ArrayClear)
-	{		
-		if (ParentLibrary != nullptr)
-		{
-			ParentLibrary->ForEachEntityOfType<UDMXEntityFixturePatch>([this](UDMXEntityFixturePatch* Patch)
+		{		
+			if (ParentLibrary != nullptr)
 			{
-				if (Patch->GetFixtureType() == this)
+				ParentLibrary->ForEachEntityOfType<UDMXEntityFixturePatch>([this](UDMXEntityFixturePatch* Patch)
 				{
-					Patch->ValidateActiveMode();					
-				}
-			});
+					if (Patch->GetFixtureType() == this)
+					{
+						Patch->ValidateActiveMode();					
+					}
+				});
+			}
 		}
-	}
 
 		int32 ChangedModeIndex = PropertyChangedEvent.GetArrayIndex(PropertyName.ToString());
 		if (Modes.IsValidIndex(ChangedModeIndex))
diff --git a/Engine/Plugins/VirtualProduction/DMX/DMXEngine/Source/DMXRuntime/Public/Library/DMXEntityFixtureType.h b/Engine/Plugins/VirtualProduction/DMX/DMXEngine/Source/DMXRuntime/Public/Library/DMXEntityFixtureType.h
--- a/Engine/Plugins/VirtualProduction/DMX/DMXEngine/Source/DMXRuntime/Public/Library/DMXEntityFixtureType.h
+++ b/Engine/Plugins/VirtualProduction/DMX/DMXEngine/Source/DMXRuntime/Public/Library/DMXEntityFixtureType.h
@@ -56,7 +56,7 @@ struct DMXRUNTIME_API FDMXFixtureFunction
 	UPROPERTY(EditAnywhere, BlueprintReadOnly, meta = (DisplayPriority = "20"), Category = "DMX")
 	FString Description;
 
-	/** Initial value for this function when no value is set */
+	/** The Default Value of the function, imported from GDTF. The plugin doesn't make use of this value, but it can be used in blueprints */
 	UPROPERTY(EditAnywhere, BlueprintReadOnly, meta = (DisplayPriority = "30"), Category = "DMX")
 	int64 DefaultValue;
 
@@ -79,14 +79,14 @@ struct DMXRUNTIME_API FDMXFixtureFunction
 
 	/**
 	 * Least Significant Byte mode makes the individual bytes (channels) of the function be
-	 * interpreted with the first bytes being the lowest part of the number.
+	 * interpreted with the first bytes being the lowest part of the number (endianness).
 	 * 
 	 * E.g., given a 16 bit function with two channel values set to [0, 1],
-	 * they would be interpreted as the binary number 00000001 00000000, which means 256.
+	 * they would be interpreted as the binary number 0x01 0x00, which means 256.
 	 * The first byte (0) became the lowest part in binary form and the following byte (1), the highest.
 	 * 
 	 * Most Fixtures use MSB (Most Significant Byte) mode, which interprets bytes as highest first.
-	 * In MSB mode, the example above would be interpreted in binary as 00000000 00000001, which means 1.
+	 * In MSB mode, the example above would be interpreted in binary as 0x00 0x01, which means 1.
 	 * The first byte (0) became the highest part in binary form and the following byte (1), the lowest.
 	 */
 	UPROPERTY(EditAnywhere, BlueprintReadOnly, meta = (DisplayName = "Use LSB Mode", DisplayPriority = "29"), Category = "DMX")
