DMX - Restore the DMX Level Editor Toolbar Menu

diff --git a/Engine/Plugins/VirtualProduction/DMX/DMXEngine/Source/DMXEditor/DMXEditor.Build.cs b/Engine/Plugins/VirtualProduction/DMX/DMXEngine/Source/DMXEditor/DMXEditor.Build.cs
--- a/Engine/Plugins/VirtualProduction/DMX/DMXEngine/Source/DMXEditor/DMXEditor.Build.cs
+++ b/Engine/Plugins/VirtualProduction/DMX/DMXEngine/Source/DMXEditor/DMXEditor.Build.cs
@@ -15,6 +15,7 @@ public class DMXEditor : ModuleRules
 				"DMXProtocol",
 				"DMXProtocolEditor",
 				"DMXRuntime",
+				"ToolMenus"
 			}
 		);
 
diff --git a/Engine/Plugins/VirtualProduction/DMX/DMXEngine/Source/DMXEditor/Private/DMXEditorModule.cpp b/Engine/Plugins/VirtualProduction/DMX/DMXEngine/Source/DMXEditor/Private/DMXEditorModule.cpp
--- a/Engine/Plugins/VirtualProduction/DMX/DMXEngine/Source/DMXEditor/Private/DMXEditorModule.cpp
+++ b/Engine/Plugins/VirtualProduction/DMX/DMXEngine/Source/DMXEditor/Private/DMXEditorModule.cpp
@@ -9,6 +9,7 @@
 #include "DMXProtocolBlueprintLibrary.h"
 #include "DMXProtocolTypes.h"
 #include "AssetTools/AssetTypeActions_DMXEditorLibrary.h"
+#include "Commands/DMXEditorCommands.h"
 #include "Customizations/DMXEntityReferenceCustomization.h"
 #include "Customizations/DMXFixtureFunctionCustomization.h"
 #include "Customizations/DMXFixtureModeCustomization.h"
@@ -16,11 +17,10 @@
 #include "Customizations/DMXNameListCustomization.h"
 #include "Customizations/DMXPixelMappingDistributionCustomization.h"
 #include "Customizations/TakeRecorderDMXLibrarySourceEditorCustomization.h"
-#include "Commands/DMXEditorCommands.h"
 #include "Game/DMXComponent.h"
-#include "Library/DMXEntityReference.h"
 #include "Library/DMXEntity.h"
 #include "Library/DMXEntityFixtureType.h"
+#include "Library/DMXEntityReference.h"
 #include "Library/DMXLibrary.h"
 #include "Sequencer/DMXLibraryTrackEditor.h"
 #include "Sequencer/TakeRecorderDMXLibrarySource.h"
@@ -29,12 +29,13 @@
 #include "Widgets/OutputConsole/SDMXOutputConsole.h"
 #include "Widgets/PatchTool/SDMXPatchTool.h"
 
-#include "LevelEditor.h"
 #include "AssetToolsModule.h"
-#include "PropertyEditorModule.h"
 #include "AssetRegistryModule.h"
 #include "AssetToolsModule.h"
 #include "ISequencerModule.h"
+#include "LevelEditor.h"
+#include "PropertyEditorModule.h"
+#include "ToolMenus.h"
 #include "Framework/Commands/UICommandList.h"
 #include "Framework/Docking/TabManager.h"
 #include "Framework/MultiBox/MultiBoxBuilder.h"
@@ -140,50 +141,40 @@ void FDMXEditorModule::BindDMXEditorCommands()
 
 void FDMXEditorModule::ExtendLevelEditorToolbar()
 {
-	FLevelEditorModule& LevelEditorModule = FModuleManager::LoadModuleChecked<FLevelEditorModule>("LevelEditor");
-	TSharedPtr<FUICommandList> CommandBindings = LevelEditorModule.GetGlobalLevelEditorActions();
+	UToolMenu* Menu = UToolMenus::Get()->ExtendMenu("LevelEditor.LevelEditorToolBar.ModesToolBar");
 
-	TSharedPtr<FExtender> ToolbarExtender = MakeShared<FExtender>();
-	ToolbarExtender->AddToolBarExtension(
-		"Settings", 
-		EExtensionHook::After,
-		CommandBindings,
-		FToolBarExtensionDelegate::CreateRaw(this, &FDMXEditorModule::GenerateToolbarExtension)
-	);
-
-	LevelEditorModule.GetToolBarExtensibilityManager()->AddExtender(ToolbarExtender);
-}
+	FToolMenuSection& Section = Menu->FindOrAddSection("DMX");
 
-void FDMXEditorModule::GenerateToolbarExtension(FToolBarBuilder& InOutBuilder)
-{
-	auto OnGetToolbarButtonBrushLambda = [this]() -> const FSlateIcon
-	{
-		bool bSendEnabled = UDMXProtocolBlueprintLibrary::IsSendDMXEnabled();
-		bool bReceiveEnabled = UDMXProtocolBlueprintLibrary::IsReceiveDMXEnabled();
-
-		if (bSendEnabled && bReceiveEnabled)
-		{
-			return FSlateIcon(FDMXEditorStyle::GetStyleSetName(), "DMXEditor.LevelEditor.MenuIcon_snd-rcv");
-		}
-		else if (bSendEnabled)
-		{
-			return FSlateIcon(FDMXEditorStyle::GetStyleSetName(), "DMXEditor.LevelEditor.MenuIcon_snd");
-		}
-		else if (bReceiveEnabled)
-		{
-			return FSlateIcon(FDMXEditorStyle::GetStyleSetName(), "DMXEditor.LevelEditor.MenuIcon_rcv");
-		}
-
-		return FSlateIcon(FDMXEditorStyle::GetStyleSetName(), "DMXEditor.LevelEditor.MenuIcon_none");
-	};
-
-	InOutBuilder.AddComboButton(
+	FToolMenuEntry DMXEntry = FToolMenuEntry::InitComboButton(
+		"DMXMenu",
 		FUIAction(),
-		FOnGetContent::CreateRaw(this, &FDMXEditorModule::GenerateDMXLevelEditorToolbarMenu),
-		LOCTEXT("InputInfo_Label", "DMX"),
-		LOCTEXT("InputInfo_ToolTip", "DMX Tools"),
-		TAttribute<FSlateIcon>::Create(OnGetToolbarButtonBrushLambda)
-	);
+		FOnGetContent::CreateStatic(&FDMXEditorModule::GenerateDMXLevelEditorToolbarMenu),
+		LOCTEXT("LevelEditorToolbarDMXButtonLabel", "DMX"),
+		LOCTEXT("LevelEditorToolbarDMXButtonTooltip", "DMX Tools"),
+		TAttribute<FSlateIcon>::CreateLambda([]()
+			{
+				bool bSendEnabled = UDMXProtocolBlueprintLibrary::IsSendDMXEnabled();
+				bool bReceiveEnabled = UDMXProtocolBlueprintLibrary::IsReceiveDMXEnabled();
+
+				if (bSendEnabled && bReceiveEnabled)
+				{
+					return FSlateIcon(FDMXEditorStyle::GetStyleSetName(), "DMXEditor.LevelEditor.MenuIcon_snd-rcv");
+				}
+				else if (bSendEnabled)
+				{
+					return FSlateIcon(FDMXEditorStyle::GetStyleSetName(), "DMXEditor.LevelEditor.MenuIcon_snd");
+				}
+				else if (bReceiveEnabled)
+				{
+					return FSlateIcon(FDMXEditorStyle::GetStyleSetName(), "DMXEditor.LevelEditor.MenuIcon_rcv");
+				}
+
+				return FSlateIcon(FDMXEditorStyle::GetStyleSetName(), "DMXEditor.LevelEditor.MenuIcon_none");
+			})
+		);
+
+	DMXEntry.StyleNameOverride = "CalloutToolbar";
+	Section.AddEntry(DMXEntry);
 }
 
 TSharedRef<SWidget> FDMXEditorModule::GenerateDMXLevelEditorToolbarMenu()
diff --git a/Engine/Plugins/VirtualProduction/DMX/DMXEngine/Source/DMXEditor/Public/DMXEditorModule.h b/Engine/Plugins/VirtualProduction/DMX/DMXEngine/Source/DMXEditor/Public/DMXEditorModule.h
--- a/Engine/Plugins/VirtualProduction/DMX/DMXEngine/Source/DMXEditor/Public/DMXEditorModule.h
+++ b/Engine/Plugins/VirtualProduction/DMX/DMXEngine/Source/DMXEditor/Public/DMXEditorModule.h
@@ -77,11 +77,8 @@ private:
 	/** Creates the level editor toolbar extender */
 	void ExtendLevelEditorToolbar();
 
-	/** Generates the DMX level editor toolbar extension */
-	void GenerateToolbarExtension(class FToolBarBuilder& InOutBuilder);
-
 	/** Generates the level editor toolbar DMX Menu */
-	TSharedRef<class SWidget> GenerateDMXLevelEditorToolbarMenu();
+	static TSharedRef<class SWidget> GenerateDMXLevelEditorToolbarMenu();
 
 protected:
 	/** Registers asset types categories */
