DMX - Change comment to refer to UE5 instead of UE4

diff --git a/Engine/Plugins/VirtualProduction/DMX/DMXEngine/Source/DMXEditor/Private/Factories/DMXGDTFImporter.cpp b/Engine/Plugins/VirtualProduction/DMX/DMXEngine/Source/DMXEditor/Private/Factories/DMXGDTFImporter.cpp
--- a/Engine/Plugins/VirtualProduction/DMX/DMXEngine/Source/DMXEditor/Private/Factories/DMXGDTFImporter.cpp
+++ b/Engine/Plugins/VirtualProduction/DMX/DMXEngine/Source/DMXEditor/Private/Factories/DMXGDTFImporter.cpp
@@ -554,7 +554,7 @@ bool FDMXGDTFImporter::ParseXML()
     int32 EndXMLStringIndex = CharString.Find(EndXMLString, ESearchCase::CaseSensitive, ESearchDir::FromEnd);
     FString ResultXMLString = CharString.Mid(BeginXMLStringIndex, EndXMLStringIndex - BeginXMLStringIndex + EndXMLString.Len());
 
-    // @TODO. UE4 XML Parser can't handle symbol < or > inside attribute and just break the parsing
+    // @TODO. UE5 XML Parser can't handle symbol < or > inside attribute and just break the parsing
     // This small algorithm remove < and > from the XML line
     TArray<FString> XMLStringLines;
     ResultXMLString.ParseIntoArrayLines(XMLStringLines, /*bCullEmpty*/false);
