DMX - Rework DMX Fixture Type editor for better GDTF support, better UX and to ease maintaining and extending the editor

diff --git a/Engine/Plugins/VirtualProduction/DMX/DMXEngine/Source/DMXBlueprintGraph/Private/DMXBlueprintGraphModule.cpp b/Engine/Plugins/VirtualProduction/DMX/DMXEngine/Source/DMXBlueprintGraph/Private/DMXBlueprintGraphModule.cpp
--- a/Engine/Plugins/VirtualProduction/DMX/DMXEngine/Source/DMXBlueprintGraph/Private/DMXBlueprintGraphModule.cpp
+++ b/Engine/Plugins/VirtualProduction/DMX/DMXEngine/Source/DMXBlueprintGraph/Private/DMXBlueprintGraphModule.cpp
@@ -26,14 +26,13 @@ void FDMXBlueprintGraphModule::StartupModule()
 
 	RegisterObjectCustomizations();
 
-	DataTypeChangeDelegate = UDMXEntityFixtureType::GetDataTypeChangeDelegate().AddRaw(this, &FDMXBlueprintGraphModule::OnDataTypeChanged);
+	UDMXEntityFixtureType::GetOnFixtureTypeChanged().AddRaw(this, &FDMXBlueprintGraphModule::OnFixtureTypeChanged);
 }
 
 void FDMXBlueprintGraphModule::ShutdownModule()
 {
 	FEdGraphUtilities::UnregisterVisualPinFactory(DMXGraphPanelPinFactory);
 
-
 	if (FModuleManager::Get().IsModuleLoaded("PropertyEditor"))
 	{
 		FPropertyEditorModule& PropertyModule = FModuleManager::GetModuleChecked<FPropertyEditorModule>("PropertyEditor");
@@ -47,11 +46,7 @@ void FDMXBlueprintGraphModule::ShutdownModule()
 		PropertyModule.NotifyCustomizationModuleChanged();
 	}
 
-	// Remove the delegate
-	if (DataTypeChangeDelegate.IsValid())
-	{
-		DataTypeChangeDelegate.Reset();
-	}
+	UDMXEntityFixtureType::GetOnFixtureTypeChanged().RemoveAll(this);
 }
 
 void FDMXBlueprintGraphModule::RegisterObjectCustomizations()
@@ -73,13 +68,13 @@ void FDMXBlueprintGraphModule::RegisterCustomClassLayout(FName ClassName, FOnGet
 	PropertyModule.RegisterCustomClassLayout(ClassName, DetailLayoutDelegate);
 }
 
-void FDMXBlueprintGraphModule::OnDataTypeChanged(const UDMXEntityFixtureType* InFixtureType, const FDMXFixtureMode& InMode)
+void FDMXBlueprintGraphModule::OnFixtureTypeChanged(const UDMXEntityFixtureType* InFixtureType)
 {
 	for (TObjectIterator<UK2Node_GetDMXAttributeValues> It(RF_Transient | RF_ClassDefaultObject, /** bIncludeDerivedClasses */ true, /** InternalExcludeFlags */ EInternalObjectFlags::PendingKill); It; ++It)
 	{
 		if (It->HasValidBlueprint())
 		{
-			It->OnDataTypeChanged(InFixtureType, InMode);
+			It->OnFixtureTypeChanged(InFixtureType);
 		}
 	}
 }
diff --git a/Engine/Plugins/VirtualProduction/DMX/DMXEngine/Source/DMXBlueprintGraph/Private/K2Node_GetDMXAttributeValues.cpp b/Engine/Plugins/VirtualProduction/DMX/DMXEngine/Source/DMXBlueprintGraph/Private/K2Node_GetDMXAttributeValues.cpp
--- a/Engine/Plugins/VirtualProduction/DMX/DMXEngine/Source/DMXBlueprintGraph/Private/K2Node_GetDMXAttributeValues.cpp
+++ b/Engine/Plugins/VirtualProduction/DMX/DMXEngine/Source/DMXBlueprintGraph/Private/K2Node_GetDMXAttributeValues.cpp
@@ -552,7 +552,7 @@ UDMXEntityFixtureType* UK2Node_GetDMXAttributeValues::GetParentFixtureType() con
 	return nullptr;
 }
 
-void UK2Node_GetDMXAttributeValues::OnDataTypeChanged(const UDMXEntityFixtureType* InFixtureType, const FDMXFixtureMode& InMode)
+void UK2Node_GetDMXAttributeValues::OnFixtureTypeChanged(const UDMXEntityFixtureType* InFixtureType)
 {
 	// Check if there are any pins exists
 	if (Pins.Num() == 0)
@@ -564,8 +564,8 @@ void UK2Node_GetDMXAttributeValues::OnDataTypeChanged(const UDMXEntityFixtureTyp
 	{
 		if (const FDMXFixtureMode* ActiveFixtureMode = GetActiveFixtureMode())
 		{
-			// Reset Attributes if there is a match in fixture types objects and name of the modes
-			if (InFixtureType && (InFixtureType == ParentFixtureType) && (InMode.ModeName == ActiveFixtureMode->ModeName))
+			// Reset Attributes if there is a match in fixture types objects
+			if (InFixtureType && (InFixtureType == ParentFixtureType))
 			{
 				ResetAttributes();
 
@@ -579,4 +579,10 @@ void UK2Node_GetDMXAttributeValues::OnDataTypeChanged(const UDMXEntityFixtureTyp
 	}
 }
 
+void UK2Node_GetDMXAttributeValues::OnDataTypeChanged(const UDMXEntityFixtureType* InFixtureType, const FDMXFixtureMode& InMode)
+{
+	// DEPRECATED 5.0
+	OnFixtureTypeChanged(InFixtureType);
+}
+
 #undef LOCTEXT_NAMESPACE
diff --git a/Engine/Plugins/VirtualProduction/DMX/DMXEngine/Source/DMXBlueprintGraph/Public/DMXBlueprintGraphModule.h b/Engine/Plugins/VirtualProduction/DMX/DMXEngine/Source/DMXBlueprintGraph/Public/DMXBlueprintGraphModule.h
--- a/Engine/Plugins/VirtualProduction/DMX/DMXEngine/Source/DMXBlueprintGraph/Public/DMXBlueprintGraphModule.h
+++ b/Engine/Plugins/VirtualProduction/DMX/DMXEngine/Source/DMXBlueprintGraph/Public/DMXBlueprintGraphModule.h
@@ -46,23 +46,18 @@ private:
 
 
 	/**
-	 * FDataTypeChangeDelegate delegate call function
+	 * Called when a fixture type changed
 	 *
 	 * @param	InFixtureType	Fixture type UObject
 	 * @param	InMode			Changed mode
 	 *
 	 */
-	void OnDataTypeChanged(const UDMXEntityFixtureType* InFixtureType, const FDMXFixtureMode& InMode);
-
-private:
+	void OnFixtureTypeChanged(const UDMXEntityFixtureType* InFixtureType);
 
 	/** FDMXProtocolName and Custom nodes Graph Pin customizations */
 	TSharedPtr<FDMXGraphPanelPinFactory> DMXGraphPanelPinFactory;
 
 	/** List of registered class that we must unregister when the module shuts down */
 	TSet< FName > RegisteredClassNames;
-
-	/** UDMXEntityFixtureType.GetDataTypeChangeDelegate() delegate handle */
-	FDelegateHandle DataTypeChangeDelegate;
 };
 
diff --git a/Engine/Plugins/VirtualProduction/DMX/DMXEngine/Source/DMXBlueprintGraph/Public/K2Node_GetDMXAttributeValues.h b/Engine/Plugins/VirtualProduction/DMX/DMXEngine/Source/DMXBlueprintGraph/Public/K2Node_GetDMXAttributeValues.h
--- a/Engine/Plugins/VirtualProduction/DMX/DMXEngine/Source/DMXBlueprintGraph/Public/K2Node_GetDMXAttributeValues.h
+++ b/Engine/Plugins/VirtualProduction/DMX/DMXEngine/Source/DMXBlueprintGraph/Public/K2Node_GetDMXAttributeValues.h
@@ -55,12 +55,15 @@ public:
 	void RemoveOutputPin(UEdGraphPin* Pin);
 
 	/**
-	 * Called when DataType is chnaged in any existing UDMXEntityFixtureType UObject
+	 * Called when the Fixture Type changed 
 	 *
-	 * @param	InFixtureType		Fixture type UObject
-	 * @param	InDMXFixtureModes	Changed Mode
+	 * @param	InFixtureType		The changed Fixture Type
 	 *
 	 */
+	void OnFixtureTypeChanged(const UDMXEntityFixtureType* InFixtureType);
+
+	/** DEPRECATED 5.0 */
+	UE_DEPRECATED(5.0, "Replaced in favor of the more generic OnFixtureTypeChanged.")
 	void OnDataTypeChanged(const UDMXEntityFixtureType* InFixtureType, const FDMXFixtureMode& InMode);
 
 public:
diff --git a/Engine/Plugins/VirtualProduction/DMX/DMXEngine/Source/DMXEditor/Private/Customizations/DMXEntityFixturePatchDetails.cpp b/Engine/Plugins/VirtualProduction/DMX/DMXEngine/Source/DMXEditor/Private/Customizations/DMXEntityFixturePatchDetails.cpp
--- a/Engine/Plugins/VirtualProduction/DMX/DMXEngine/Source/DMXEditor/Private/Customizations/DMXEntityFixturePatchDetails.cpp
+++ b/Engine/Plugins/VirtualProduction/DMX/DMXEngine/Source/DMXEditor/Private/Customizations/DMXEntityFixturePatchDetails.cpp
@@ -35,7 +35,7 @@ void FDMXEntityFixturePatchDetails::CustomizeDetails(IDetailLayoutBuilder& Detai
 	AutoAssignAddressHandle->SetOnPropertyValueChanged(OnAutoAssignAddressChangedDelegate);
 
 	// Handle mode changes of the parent fixture type
-	UDMXEntityFixtureType::GetDataTypeChangeDelegate().AddSP(this, &FDMXEntityFixturePatchDetails::OnModeChanged);
+	UDMXEntityFixtureType::GetOnFixtureTypeChanged().AddSP(this, &FDMXEntityFixturePatchDetails::OnFixtureTypeChanged);
 
 	// Make a Fixture Types dropdown for the Fixture Type template property
 	DetailBuilder.EditDefaultProperty(ParentFixtureTypeHandle)->CustomWidget(false)
@@ -121,7 +121,7 @@ void FDMXEntityFixturePatchDetails::OnParentFixtureTypeChanged(UDMXEntity* NewTe
 	ParentFixtureTypeHandle->SetValue(Cast<UDMXEntityFixtureType>(NewTemplate));
 }
 
-void FDMXEntityFixturePatchDetails::OnModeChanged(const UDMXEntityFixtureType* FixtureType, const FDMXFixtureMode& Mode)
+void FDMXEntityFixturePatchDetails::OnFixtureTypeChanged(const UDMXEntityFixtureType* FixtureType)
 {
 	PropertyUtilities->ForceRefresh();
 }
diff --git a/Engine/Plugins/VirtualProduction/DMX/DMXEngine/Source/DMXEditor/Private/Customizations/DMXEntityFixturePatchDetails.h b/Engine/Plugins/VirtualProduction/DMX/DMXEngine/Source/DMXEditor/Private/Customizations/DMXEntityFixturePatchDetails.h
--- a/Engine/Plugins/VirtualProduction/DMX/DMXEngine/Source/DMXEditor/Private/Customizations/DMXEntityFixturePatchDetails.h
+++ b/Engine/Plugins/VirtualProduction/DMX/DMXEngine/Source/DMXEditor/Private/Customizations/DMXEntityFixturePatchDetails.h
@@ -44,7 +44,7 @@ private:
 	void OnActiveModeChanged(const TSharedPtr<uint32> InSelectedMode, ESelectInfo::Type SelectInfo);
 
 	/** Called when a Mode in the Modes Array property changed */
-	void OnModeChanged(const UDMXEntityFixtureType* FixtureType, const FDMXFixtureMode& Mode);
+	void OnFixtureTypeChanged(const UDMXEntityFixtureType* FixtureType);
 
 	/** Called when the bAutoAssignAddress property changed */
 	void OnAutoAssignAddressChanged();
diff --git a/Engine/Plugins/VirtualProduction/DMX/DMXEngine/Source/DMXEditor/Private/Customizations/DMXEntityFixtureTypeDetails.cpp b/Engine/Plugins/VirtualProduction/DMX/DMXEngine/Source/DMXEditor/Private/Customizations/DMXEntityFixtureTypeDetails.cpp
new file mode 100644
--- /dev/null
+++ b/Engine/Plugins/VirtualProduction/DMX/DMXEngine/Source/DMXEditor/Private/Customizations/DMXEntityFixtureTypeDetails.cpp
@@ -0,0 +1,51 @@
+// Copyright Epic Games, Inc. All Rights Reserved.
+
+#include "Customizations/DMXEntityFixtureTypeDetails.h"
+
+#include "Library/DMXEntityFixtureType.h"
+#include "Library/DMXImport.h"
+
+#include "CoreMinimal.h"
+#include "DetailLayoutBuilder.h"
+#include "IPropertyUtilities.h"
+#include "PropertyHandle.h"
+
+
+#define LOCTEXT_NAMESPACE "DMXEntityFixtureTypeDetails"
+
+TSharedRef<IDetailCustomization> FDMXEntityFixtureTypeDetails::MakeInstance()
+{
+	return MakeShared<FDMXEntityFixtureTypeDetails>();
+}
+
+void FDMXEntityFixtureTypeDetails::CustomizeDetails(IDetailLayoutBuilder& DetailBuilder)
+{
+	PropertyUtilities = DetailBuilder.GetPropertyUtilities();
+
+	DetailBuilder.HideProperty(GET_MEMBER_NAME_CHECKED(UDMXEntityFixtureType, Modes));
+
+	DMXImportHandle = DetailBuilder.GetProperty(GET_MEMBER_NAME_CHECKED(UDMXEntityFixtureType, DMXImport));
+	DMXImportHandle->SetOnPropertyValueChanged(FSimpleDelegate::CreateSP(this, &FDMXEntityFixtureTypeDetails::OnDMXImportChanged));
+};
+
+void FDMXEntityFixtureTypeDetails::OnDMXImportChanged()
+{
+	UObject* DMXImportObject = nullptr;
+	if (DMXImportHandle->GetValue(DMXImportObject) == FPropertyAccess::Success)
+	{
+		if (UDMXImport* DMXImport = Cast<UDMXImport>(DMXImportObject))
+		{
+			const TArray<TWeakObjectPtr<UObject>>& SelectedObjects = PropertyUtilities->GetSelectedObjects();
+
+			for (TWeakObjectPtr<UObject> WeakFixtureTypeObject : SelectedObjects)
+			{
+				if (UDMXEntityFixtureType* FixtureType = Cast<UDMXEntityFixtureType>(WeakFixtureTypeObject.Get()))
+				{
+					FixtureType->SetModesFromDMXImport(DMXImport);
+				}
+			}
+		}
+	}
+}
+
+#undef LOCTEXT_NAMESPACE
diff --git a/Engine/Plugins/VirtualProduction/DMX/DMXEngine/Source/DMXEditor/Private/DMXEditor.cpp b/Engine/Plugins/VirtualProduction/DMX/DMXEngine/Source/DMXEditor/Private/DMXEditor.cpp
--- a/Engine/Plugins/VirtualProduction/DMX/DMXEngine/Source/DMXEditor/Private/DMXEditor.cpp
+++ b/Engine/Plugins/VirtualProduction/DMX/DMXEngine/Source/DMXEditor/Private/DMXEditor.cpp
@@ -19,7 +19,6 @@
 #include "Widgets/Docking/SDockTab.h"
 
 #include "Widgets/SDMXEntityEditor.h"
-#include "Widgets/SDMXEntityInspector.h"
 #include "Widgets/FixturePatch/SDMXFixturePatchEditor.h"
 #include "Widgets/FixtureType/SDMXFixtureTypeEditor.h"
 #include "Widgets/OutputConsole/SDMXOutputConsole.h"
@@ -307,8 +306,7 @@ TSharedRef<SDMXLibraryEditorTab> FDMXEditor::CreateDMXLibraryEditorTab()
 
 TSharedRef<SDMXFixtureTypeEditor> FDMXEditor::CreateFixtureTypeEditor()
 {
-	return SNew(SDMXFixtureTypeEditor)
-		.DMXEditor(SharedThis(this));
+	return SNew(SDMXFixtureTypeEditor, SharedThis(this));
 }
 
 TSharedRef<SDMXFixturePatchEditor> FDMXEditor::CreateFixturePatchEditor()
diff --git a/Engine/Plugins/VirtualProduction/DMX/DMXEngine/Source/DMXEditor/Private/DMXEditorModule.cpp b/Engine/Plugins/VirtualProduction/DMX/DMXEngine/Source/DMXEditor/Private/DMXEditorModule.cpp
--- a/Engine/Plugins/VirtualProduction/DMX/DMXEngine/Source/DMXEditor/Private/DMXEditorModule.cpp
+++ b/Engine/Plugins/VirtualProduction/DMX/DMXEngine/Source/DMXEditor/Private/DMXEditorModule.cpp
@@ -10,9 +10,9 @@
 #include "DMXProtocolTypes.h"
 #include "AssetTools/AssetTypeActions_DMXEditorLibrary.h"
 #include "Commands/DMXEditorCommands.h"
+#include "Customizations/DMXEntityFixtureTypeDetails.h"
 #include "Customizations/DMXEntityReferenceCustomization.h"
-#include "Customizations/DMXFixtureFunctionCustomization.h"
-#include "Customizations/DMXFixtureModeCustomization.h"
+#include "Customizations/DMXFixtureSignalFormatCustomization.h"
 #include "Customizations/DMXLibraryPortReferencesCustomization.h"
 #include "Customizations/DMXNameListCustomization.h"
 #include "Customizations/DMXPixelMappingDistributionCustomization.h"
@@ -66,6 +66,7 @@ void FDMXEditorModule::StartupModule()
 	RegisterAssetTypeCategories();
 	RegisterAssetTypeActions();
 
+	RegisterClassCustomizations();
 	RegisterPropertyTypeCustomizations();
 	RegisterSequencerTypes();
 	RegisterNomadTabSpawners();
@@ -216,35 +217,58 @@ void FDMXEditorModule::RegisterAssetTypeActions()
 	RegisterAssetTypeAction(MakeShared<FAssetTypeActions_DMXEditorLibrary>());
 }
 
-void FDMXEditorModule::RegisterPropertyTypeCustomizations()
+void FDMXEditorModule::RegisterClassCustomizations()
 {
 	FPropertyEditorModule& PropertyModule = FModuleManager::GetModuleChecked<FPropertyEditorModule>("PropertyEditor");
 
-	// Property type customization for the DMXFixtureMode struct
-	FOnGetPropertyTypeCustomizationInstance FixtureModePropertyTypeCustomizationGetter = FOnGetPropertyTypeCustomizationInstance::CreateStatic(&FDMXFixtureModeCustomization::MakeInstance);
-	PropertyModule.RegisterCustomPropertyTypeLayout(FDMXFixtureMode::StaticStruct()->GetFName(), FixtureModePropertyTypeCustomizationGetter);
+	// Property type customization for the UDMXEntityFixtureType class
+	PropertyModule.RegisterCustomClassLayout(UDMXEntityFixtureType::StaticClass()->GetFName(),
+		FOnGetDetailCustomizationInstance::CreateStatic(&FDMXEntityFixtureTypeDetails::MakeInstance)
+	);
+}
 
-	// Property type customization for the DMXFixtureFunction struct
-	FOnGetPropertyTypeCustomizationInstance FixtureFunctionPropertyTypeCustomizationGetter = FOnGetPropertyTypeCustomizationInstance::CreateStatic(&FDMXFixtureFunctionCustomization::MakeInstance);
-	PropertyModule.RegisterCustomPropertyTypeLayout(FDMXFixtureFunction::StaticStruct()->GetFName(), FixtureFunctionPropertyTypeCustomizationGetter);
+void FDMXEditorModule::RegisterPropertyTypeCustomizations()
+{
+	FPropertyEditorModule& PropertyModule = FModuleManager::GetModuleChecked<FPropertyEditorModule>("PropertyEditor");
 
 	// Property type customization for the EDMXPixelMappingDistribution enum
-	RegisterCustomPropertyTypeLayout("EDMXPixelMappingDistribution", FOnGetPropertyTypeCustomizationInstance::CreateStatic(&FDMXPixelMappingDistributionCustomization::MakeInstance));
+	RegisterCustomPropertyTypeLayout("EDMXPixelMappingDistribution", 
+		FOnGetPropertyTypeCustomizationInstance::CreateStatic(&FDMXPixelMappingDistributionCustomization::MakeInstance)
+	);
+
+	// Property type customization for the EDMXFixtureSignalFormat enum
+	RegisterCustomPropertyTypeLayout("EDMXFixtureSignalFormat", 
+		FOnGetPropertyTypeCustomizationInstance::CreateStatic(&FDMXFixtureSignalFormatCustomization::MakeInstance)
+	);
 
 	// Property type customization for Name List structs
-	RegisterCustomPropertyTypeLayout(FDMXProtocolName::StaticStruct()->GetFName(), FOnGetPropertyTypeCustomizationInstance::CreateStatic(&FDMXNameListCustomization<FDMXProtocolName>::MakeInstance));
-	RegisterCustomPropertyTypeLayout(FDMXFixtureCategory::StaticStruct()->GetFName(), FOnGetPropertyTypeCustomizationInstance::CreateStatic(&FDMXNameListCustomization<FDMXFixtureCategory>::MakeInstance));
-	RegisterCustomPropertyTypeLayout(FDMXAttributeName::StaticStruct()->GetFName(), FOnGetPropertyTypeCustomizationInstance::CreateStatic(&FDMXNameListCustomization<FDMXAttributeName>::MakeInstance));
+	RegisterCustomPropertyTypeLayout(FDMXProtocolName::StaticStruct()->GetFName(), 
+		FOnGetPropertyTypeCustomizationInstance::CreateStatic(&FDMXNameListCustomization<FDMXProtocolName>::MakeInstance)
+	);
+	RegisterCustomPropertyTypeLayout(FDMXFixtureCategory::StaticStruct()->GetFName(), 
+		FOnGetPropertyTypeCustomizationInstance::CreateStatic(&FDMXNameListCustomization<FDMXFixtureCategory>::MakeInstance)
+	);
+	RegisterCustomPropertyTypeLayout(FDMXAttributeName::StaticStruct()->GetFName(), 
+		FOnGetPropertyTypeCustomizationInstance::CreateStatic(&FDMXNameListCustomization<FDMXAttributeName>::MakeInstance)
+	);
 
-	// Customizations for the Entity Reference structs
-	RegisterCustomPropertyTypeLayout(FDMXEntityFixtureTypeRef::StaticStruct()->GetFName(), FOnGetPropertyTypeCustomizationInstance::CreateStatic(&FDMXEntityReferenceCustomization::MakeInstance));
-	RegisterCustomPropertyTypeLayout(FDMXEntityFixturePatchRef::StaticStruct()->GetFName(), FOnGetPropertyTypeCustomizationInstance::CreateStatic(&FDMXEntityReferenceCustomization::MakeInstance));
+	// Customizations for the FDMXEntityReference structs
+	RegisterCustomPropertyTypeLayout(FDMXEntityFixtureTypeRef::StaticStruct()->GetFName(), 
+		FOnGetPropertyTypeCustomizationInstance::CreateStatic(&FDMXEntityReferenceCustomization::MakeInstance)
+	);
+	RegisterCustomPropertyTypeLayout(FDMXEntityFixturePatchRef::StaticStruct()->GetFName(), 
+		FOnGetPropertyTypeCustomizationInstance::CreateStatic(&FDMXEntityReferenceCustomization::MakeInstance)
+	);
 
-	// Customization for the DMXLibraryPortReferences struct type
-	RegisterCustomPropertyTypeLayout(FDMXLibraryPortReferences::StaticStruct()->GetFName(), FOnGetPropertyTypeCustomizationInstance::CreateStatic(&FDMXLibraryPortReferencesCustomization::MakeInstance));
+	// Customization for the FDMXLibraryPortReferences struct type
+	RegisterCustomPropertyTypeLayout(FDMXLibraryPortReferences::StaticStruct()->GetFName(), 
+		FOnGetPropertyTypeCustomizationInstance::CreateStatic(&FDMXLibraryPortReferencesCustomization::MakeInstance)
+	);
 
-	// DMXLibrary TakeRecorder AddAllPatchesButton customization
-	RegisterCustomPropertyTypeLayout(FAddAllPatchesButton::StaticStruct()->GetFName(), FOnGetPropertyTypeCustomizationInstance::CreateStatic(&FTakeRecorderDMXLibrarySourceEditorCustomization::MakeInstance));
+	// Customization for the DMXLibrary TakeRecorder AddAllPatchesButton
+	RegisterCustomPropertyTypeLayout(FAddAllPatchesButton::StaticStruct()->GetFName(), 
+		FOnGetPropertyTypeCustomizationInstance::CreateStatic(&FTakeRecorderDMXLibrarySourceEditorCustomization::MakeInstance)
+	);
 }
 
 void FDMXEditorModule::RegisterSequencerTypes()
diff --git a/Engine/Plugins/VirtualProduction/DMX/DMXEngine/Source/DMXEditor/Private/DMXEditorUtils.cpp b/Engine/Plugins/VirtualProduction/DMX/DMXEngine/Source/DMXEditor/Private/DMXEditorUtils.cpp
--- a/Engine/Plugins/VirtualProduction/DMX/DMXEngine/Source/DMXEditor/Private/DMXEditorUtils.cpp
+++ b/Engine/Plugins/VirtualProduction/DMX/DMXEngine/Source/DMXEditor/Private/DMXEditorUtils.cpp
@@ -89,6 +89,7 @@ protected:
 
 FString FDMXEditorUtils::GenerateUniqueNameFromExisting(const TSet<FString>& InExistingNames, const FString& InBaseName)
 {
+	// DEPRECATED UE5 (now in FDMXRuntimeUtils)
 	if (!InBaseName.IsEmpty() && !InExistingNames.Contains(InBaseName))
 	{
 		return InBaseName;
@@ -148,7 +149,7 @@ FString FDMXEditorUtils::FindUniqueEntityName(const UDMXLibrary* InLibrary, TSub
 		BaseName = InEntityClass->GetDisplayNameText().ToString();
 	}
 
-	return GenerateUniqueNameFromExisting(EntityNames, BaseName);
+	return FDMXRuntimeUtils::GenerateUniqueNameFromExisting(EntityNames, BaseName);
 }
 
 void FDMXEditorUtils::SetNewFixtureFunctionsNames(UDMXEntityFixtureType* InFixtureType)
@@ -179,7 +180,7 @@ void FDMXEditorUtils::SetNewFixtureFunctionsNames(UDMXEntityFixtureType* InFixtu
 				}
 			}
 
-			Mode.ModeName = GenerateUniqueNameFromExisting(ModesNames, TEXT("Mode"));
+			Mode.ModeName = FDMXRuntimeUtils::GenerateUniqueNameFromExisting(ModesNames, TEXT("Mode"));
 			ModesNames.Add(Mode.ModeName);
 		}
 
@@ -201,7 +202,7 @@ void FDMXEditorUtils::SetNewFixtureFunctionsNames(UDMXEntityFixtureType* InFixtu
 					}
 				}
 
-				Function.FunctionName = GenerateUniqueNameFromExisting(FunctionsNames, TEXT("Function"));
+				Function.FunctionName = FDMXRuntimeUtils::GenerateUniqueNameFromExisting(FunctionsNames, TEXT("Function"));
 				FunctionsNames.Add(Function.FunctionName);
 			}
 		}
diff --git a/Engine/Plugins/VirtualProduction/DMX/DMXEngine/Source/DMXEditor/Private/DMXFixtureTypeSharedData.cpp b/Engine/Plugins/VirtualProduction/DMX/DMXEngine/Source/DMXEditor/Private/DMXFixtureTypeSharedData.cpp
--- a/Engine/Plugins/VirtualProduction/DMX/DMXEngine/Source/DMXEditor/Private/DMXFixtureTypeSharedData.cpp
+++ b/Engine/Plugins/VirtualProduction/DMX/DMXEngine/Source/DMXEditor/Private/DMXFixtureTypeSharedData.cpp
@@ -3,21 +3,22 @@
 #include "DMXFixtureTypeSharedData.h"
 
 #include "DMXEditorUtils.h"
+#include "DMXRuntimeUtils.h"
 #include "Library/DMXEntityFixtureType.h"
+
 #include "ScopedTransaction.h"
-#include "PropertyHandle.h"
 
 
 #define LOCTEXT_NAMESPACE "DMXFixtureTypeSharedData"
 
-namespace 
+namespace
 {
 	namespace FDMXFixtureTypeSharedDataDetails
 	{
-		// Helper to raise a mode Pre/PostEditChangeChainPropertyEvent with array index
-		struct FDMXScopedModeEditChangeChainProperty
+		// DEPRECATED 5.0, only here in to keep support for deprecated members of DMXFixtureTypeSharedData
+		struct FDMXScopedModeEditChangeChainProperty_DEPRECATED
 		{
-			FDMXScopedModeEditChangeChainProperty(UDMXEntityFixtureType* InFixtureType)
+			FDMXScopedModeEditChangeChainProperty_DEPRECATED(UDMXEntityFixtureType* InFixtureType)
 			{
 				check(InFixtureType);
 
@@ -34,7 +35,7 @@ namespace
 				FixtureType->Modify();
 			}
 
-			~FDMXScopedModeEditChangeChainProperty()
+			~FDMXScopedModeEditChangeChainProperty_DEPRECATED()
 			{
 				if (FixtureType.IsValid())
 				{
@@ -69,513 +70,516 @@ namespace
 	}
 }
 
-FDMXFixtureModeItem::FDMXFixtureModeItem(const TWeakPtr<FDMXFixtureTypeSharedData>& InSharedDataPtr, const TSharedPtr<IPropertyHandle>& InModeNameHandle)
-	: SharedDataPtr(InSharedDataPtr)
-	, ModeNameHandle(InModeNameHandle)
+void FDMXFixtureTypeSharedData::SelectFixtureTypes(const TArray<TWeakObjectPtr<UDMXEntityFixtureType>>& InFixtureTypes)
 {
-	check(ModeNameHandle.IsValid() && ModeNameHandle->IsValidHandle());
-	check(ModeNameHandle->GetProperty() && ModeNameHandle->GetProperty()->GetFName() == GET_MEMBER_NAME_CHECKED(FDMXFixtureMode, ModeName));
-
-}
-
-FDMXFixtureModeItem::~FDMXFixtureModeItem()
-{}
+	if (SelectedFixtureTypes != InFixtureTypes)
+	{
+		SelectedFixtureTypes = InFixtureTypes;
 
-bool FDMXFixtureModeItem::EqualsMode(const TSharedPtr<FDMXFixtureModeItem>& Other) const
-{
-	check(Other->SharedDataPtr == SharedDataPtr);
+		// Selected Modes and Functions turn invalid
+		SelectedModeIndices.Reset();
+		SelectedFunctionIndices.Reset();
+		bFixtureMatrixSelected = false;
 
-	return Other->GetModeName() == GetModeName();
+		OnFixtureTypesSelected.Broadcast();
+	}
 }
 
-FString FDMXFixtureModeItem::GetModeName() const
+void FDMXFixtureTypeSharedData::SelectModes(const TArray<int32>& InModeIndices)
 {
-	FString ModeName;
-	if (ModeNameHandle.IsValid())
+	if (SelectedModeIndices != InModeIndices)
 	{
-		ModeNameHandle->GetValue(ModeName);
-	}
-	return ModeName;
-}
+		SelectedModeIndices = InModeIndices;
 
-void FDMXFixtureModeItem::SetModeName(const FString& InDesiredName, FString& OutNewName)
-{
-	SetName(InDesiredName, OutNewName);
-}
+		// Selected Functions turn invalid
+		SelectedFunctionIndices.Reset();
+		bFixtureMatrixSelected = false;
 
-bool FDMXFixtureModeItem::IsModeBeingEdited() const
-{
-	if (TSharedPtr<FDMXFixtureTypeSharedData> SharedData = SharedDataPtr.Pin())
-	{
-		const TSharedPtr<FDMXFixtureModeItem>* EditedModeItem = SharedData->GetModesBeingEdited().FindByPredicate([&](const TSharedPtr<FDMXFixtureModeItem>& Other) {
-			return EqualsMode(Other);
-			});
-		return EditedModeItem != nullptr;
+		OnModesSelected.Broadcast();
 	}
-	return false;
 }
 
-bool FDMXFixtureModeItem::IsModeSelected() const 
+void FDMXFixtureTypeSharedData::SetFunctionAndMatrixSelection(const TArray<int32>& InFunctionIndices, bool bMatrixSelected)
 {
-	if (TSharedPtr<FDMXFixtureTypeSharedData> SharedData = SharedDataPtr.Pin())
+	if (SelectedFunctionIndices != InFunctionIndices ||
+		bFixtureMatrixSelected != bMatrixSelected)
 	{
-		const TSharedPtr<FDMXFixtureModeItem>* SelectedModeItem = SharedData->GetSelectedModes().FindByPredicate([&](const TSharedPtr<FDMXFixtureModeItem>& Other) {
-			return EqualsMode(Other);
-			});
-		return SelectedModeItem != nullptr;
-	}
-	return false;
-}
+		SelectedFunctionIndices = InFunctionIndices;
+		bFixtureMatrixSelected = bMatrixSelected;
 
-void FDMXFixtureModeItem::SelectMode(bool bSelected)
-{
-	if (TSharedPtr<FDMXFixtureTypeSharedData> SharedData = SharedDataPtr.Pin())
-	{
-		SharedData->SelectModes(TArray<TSharedPtr<FDMXFixtureModeItem>>({ SharedThis(this) }));
+		OnFunctionsSelected.Broadcast();
+		OnMatrixSelectionChanged.Broadcast();
 	}
 }
 
-void FDMXFixtureModeItem::GetOuterFixtureTypes(TArray<TWeakObjectPtr<UDMXEntityFixtureType>>& OuterFixtureTypes) const
+bool FDMXFixtureTypeSharedData::CanCopyModesToClipboard() const
 {
-	OuterFixtureTypes.Reset();
-	if (ModeNameHandle.IsValid())
-	{
-		TArray<UObject*> OuterObjects;
-		ModeNameHandle->GetOuterObjects(OuterObjects);
-		for (UObject* Object : OuterObjects)
-		{
-			TWeakObjectPtr<UDMXEntityFixtureType> OuterFixtureType = CastChecked<UDMXEntityFixtureType>(Object);
-			OuterFixtureTypes.Add(OuterFixtureType);
-		}
-	}
-}
+	UE_LOG(LogTemp, Warning, TEXT("Test can copy"));
 
-void FDMXFixtureModeItem::GetName(FString& OutUniqueName) const
-{
-	if (ModeNameHandle.IsValid())
+	if (SelectedModeIndices.Num() == 0)
 	{
-		ModeNameHandle->GetValue(OutUniqueName);
+		return false;
 	}
-}
-
-bool FDMXFixtureModeItem::IsNameUnique(const FString& TestedName) const
-{
-	if (ModeNameHandle.IsValid())
+	else if (SelectedFixtureTypes.Num() == 1)
 	{
-		TArray<TWeakObjectPtr<UDMXEntityFixtureType>> OuterFixtureTypes;
-		GetOuterFixtureTypes(OuterFixtureTypes);
-
-		TSet<FString> ModeNames;
-		for (const TWeakObjectPtr<UDMXEntityFixtureType>& FixtureType : OuterFixtureTypes)
+		if (UDMXEntityFixtureType* FixtureType = SelectedFixtureTypes[0].Get())
 		{
-			for (const FDMXFixtureMode& Mode : FixtureType->Modes)
+			for (int32 ModeIndex : SelectedModeIndices)
 			{
-				ModeNames.Add(Mode.ModeName);
+				if (!FixtureType->Modes.IsValidIndex(SelectedModeIndices[0]))
+				{
+					return false;
+				}
 			}
 		}
-
-		return !ModeNames.Contains(TestedName);
 	}
+
 	return true;
 }
 
-void FDMXFixtureModeItem::SetName(const FString& InDesiredName, FString& OutUniqueName)
+void FDMXFixtureTypeSharedData::CopyModesToClipboard()
 {
-	if (TSharedPtr<FDMXFixtureTypeSharedData> SharedData = SharedDataPtr.Pin())
+	if (ensureMsgf(CanCopyModesToClipboard(), TEXT("Cannot copy Modes to clipboard. Please test first with FDMXFixtureTypeSharedData::CanCopyModesToClipboard.")))
 	{
-		TSet<FString> ModeNames;
-		for (const TWeakObjectPtr<UDMXEntityFixtureType>& FixtureType : SharedData->GetFixtureTypesBeingEdited())
+		ModesClipboard.Reset();
+
+		for (TWeakObjectPtr<UDMXEntityFixtureType> WeakFixtureType : SelectedFixtureTypes)
 		{
-			for (const FDMXFixtureMode& Mode : FixtureType->Modes)
+			if (UDMXEntityFixtureType* FixtureType = WeakFixtureType.Get())
 			{
-				ModeNames.Add(Mode.ModeName);
+				for (const int32 ModeIndex : SelectedModeIndices)
+				{
+					const FDMXFixtureMode& Mode = FixtureType->Modes[ModeIndex];
+					if (FixtureType->Modes.IsValidIndex(ModeIndex))
+					{
+						const TOptional<FString> CopyStr = FDMXRuntimeUtils::SerializeStructToString(FDMXFixtureMode::StaticStruct(), (void*)&Mode);
+						if (CopyStr.IsSet())
+						{
+							ModesClipboard.Add(CopyStr.GetValue());
+						}
+					}
+				}
 			}
 		}
-
-		OutUniqueName = FDMXEditorUtils::GenerateUniqueNameFromExisting(ModeNames, InDesiredName);
-		ModeNameHandle->SetValue(OutUniqueName);
 	}
 }
 
-FDMXFixtureFunctionItem::FDMXFixtureFunctionItem(TWeakPtr<FDMXFixtureTypeSharedData> InOwnerPtr, const TSharedPtr<IPropertyHandle> InModeNameHandle, const TSharedPtr<IPropertyHandle> InFunctionHandle)
-	: FDMXFixtureModeItem(InOwnerPtr, InModeNameHandle)
-	, FunctionHandle(InFunctionHandle)
-{
-	check(FunctionHandle.IsValid());
-
-	FunctionNameHandle = InFunctionHandle->GetChildHandle(GET_MEMBER_NAME_CHECKED(FDMXFixtureFunction, FunctionName));
-	check(FunctionNameHandle.IsValid() && FunctionNameHandle->IsValidHandle());
-
-	FunctionChannelHandle = InFunctionHandle->GetChildHandle(GET_MEMBER_NAME_CHECKED(FDMXFixtureFunction, Channel));
-	check(FunctionChannelHandle.IsValid() && FunctionChannelHandle->IsValidHandle());
-
-	check(FunctionNameHandle.IsValid() && FunctionNameHandle->IsValidHandle());
-	check(FunctionNameHandle->GetProperty() && FunctionNameHandle->GetProperty()->GetFName() == GET_MEMBER_NAME_CHECKED(FDMXFixtureFunction, FunctionName));
-}
-
-FDMXFixtureFunctionItem::~FDMXFixtureFunctionItem()
-{}
-
-bool FDMXFixtureFunctionItem::EqualsFunction(const TSharedPtr<FDMXFixtureFunctionItem>& Other) const
+bool FDMXFixtureTypeSharedData::CanPasteModesFromClipboard() const
 {
-	check(Other->SharedDataPtr == SharedDataPtr);
-	// Since we expect a common shared data instance, 
-	// Edited objects being equal is implied
-
-	if (!EqualsMode(Other))
+	if (SelectedFixtureTypes.Num() == 1 && ModesClipboard.Num() > 0)
 	{
-		return false;
+		if (UDMXEntityFixtureType* FixtureType = SelectedFixtureTypes[0].Get())
+		{
+			return true;
+		}
 	}
 
-	return GetFunctionName() == Other->GetFunctionName();
+	return true;
 }
 
-FString FDMXFixtureFunctionItem::GetFunctionName() const
+void FDMXFixtureTypeSharedData::PasteModesFromClipboard(TArray<int32>& OutNewlyAddedModeIndices)
 {
-	FString FunctionName;	
-	if (FunctionNameHandle.IsValid())
+	if (ensureMsgf(CanPasteModesFromClipboard(), TEXT("Cannot paste Modes from clipboard. Please test first with FDMXFixtureTypeSharedData::CanPasteModesFromClipboard.")))
 	{
-		FunctionNameHandle->GetValue(FunctionName);
-	}
-	return FunctionName;
-}
+		if (SelectedFixtureTypes.Num() == 1)
+		{
+			if (UDMXEntityFixtureType* FixtureType = SelectedFixtureTypes[0].Get())
+			{
+				const FText TransactionText = FText::Format(LOCTEXT("PasteModesTransaction", "Paste Fixture {0}|plural(one=Mode, other=Modes)"), GetSelectedModeIndices().Num());
+				const FScopedTransaction PasteModesTransaction(TransactionText);
 
-int32 FDMXFixtureFunctionItem::GetFunctionChannel() const
-{
-	int32 FunctionChannel = 0;
-	if (FunctionChannelHandle.IsValid())
-	{
-		FunctionChannelHandle->GetValue(FunctionChannel);
-	}
-	return FunctionChannel;
-}
+				FixtureType->PreEditChange(UDMXEntityFixtureType::StaticClass()->FindPropertyByName(GET_MEMBER_NAME_CHECKED(UDMXEntityFixtureType, Modes)));
 
-void FDMXFixtureFunctionItem::SetFunctionName(const FString& InDesiredName, FString& OutNewName)
-{
-	if (FunctionNameHandle.IsValid())
-	{
-		FunctionNameHandle->GetValue(OutNewName);
-	}
-}
+				int32 NumPastedElements = 0;
+				for (int32 ClipboardIndex = 0; ClipboardIndex < ModesClipboard.Num(); ClipboardIndex++)
+				{
+					const FString& ClipboardString = ModesClipboard[ClipboardIndex];
 
-bool FDMXFixtureFunctionItem::IsFunctionEdited() const
-{
-	if (TSharedPtr<FDMXFixtureTypeSharedData> SharedData = SharedDataPtr.Pin())
-	{
-		const TSharedPtr<FDMXFixtureFunctionItem>* EditedFunctionItem = SharedData->GetFunctionsBeingEdited().FindByPredicate([&](const TSharedPtr<FDMXFixtureFunctionItem>& Other) {
-			return EqualsFunction(Other);
-			});
-		return EditedFunctionItem != nullptr;
-	}
-	return false;
-}
+					const int32 PasteOntoIndex = [FixtureType, &ClipboardIndex, this]()
+					{
+						if (SelectedModeIndices.IsValidIndex(ClipboardIndex))
+						{
+							return SelectedModeIndices[ClipboardIndex];
+						}
+						else if (SelectedModeIndices.Num() > 0)
+						{
+							return SelectedModeIndices.Last();
+						}
+						else
+						{
+							return FixtureType->Modes.Num() - 1;
+						}
+					}();
+
+					TSharedPtr<FJsonObject> RootJsonObject;
+					TSharedRef<TJsonReader<>> Reader = TJsonReaderFactory<>::Create(ClipboardString);
+					FJsonSerializer::Deserialize(Reader, RootJsonObject);
+
+					if (RootJsonObject.IsValid())
+					{
+						FDMXFixtureMode ModeFromClipboard;
+						if (FJsonObjectConverter::JsonObjectToUStruct(RootJsonObject.ToSharedRef(), FDMXFixtureMode::StaticStruct(), &ModeFromClipboard, 0, 0))
+						{
+							int32 IndexOfNewlyAddedMode = PasteOntoIndex + 1 + NumPastedElements;
+							if (FixtureType->Modes.IsValidIndex(IndexOfNewlyAddedMode))
+							{
+								IndexOfNewlyAddedMode = FixtureType->Modes.Insert(ModeFromClipboard, IndexOfNewlyAddedMode);
+							}
+							else
+							{
+								IndexOfNewlyAddedMode = FixtureType->Modes.Add(ModeFromClipboard);
+							}
+							OutNewlyAddedModeIndices.Add(IndexOfNewlyAddedMode);
+
+							// Make a uinque mode name by setting it
+							FString UnusedString;
+							FixtureType->SetModeName(IndexOfNewlyAddedMode, ModeFromClipboard.ModeName, UnusedString);
+
+							NumPastedElements++;
+						}
+					}
+				}
 
-bool FDMXFixtureFunctionItem::IsFunctionSelected() const
-{
-	if (TSharedPtr<FDMXFixtureTypeSharedData> SharedData = SharedDataPtr.Pin())
-	{
-		const TSharedPtr<FDMXFixtureFunctionItem>* SelectedFunctionItem = SharedData->GetSelectedFunctions().FindByPredicate([&](const TSharedPtr<FDMXFixtureFunctionItem>& Other) {
-			return EqualsFunction(Other);
-			});
-		return SelectedFunctionItem != nullptr;
+				FixtureType->PostEditChange();
+			}
+		}
 	}
-	return false;
 }
 
-void FDMXFixtureFunctionItem::GetName(FString& OutUniqueName) const
+bool FDMXFixtureTypeSharedData::CanCopyFunctionsToClipboard() const
 {
-	if (FunctionNameHandle.IsValid())
+	if (SelectedFunctionIndices.Num() == 0 || bFixtureMatrixSelected)
 	{
-		FunctionNameHandle->GetValue(OutUniqueName);
+		return false;
 	}
-}
-
-bool FDMXFixtureFunctionItem::IsNameUnique(const FString& TestedName) const
-{
-	if (TSharedPtr<FDMXFixtureTypeSharedData> SharedData = SharedDataPtr.Pin())
+	else if (SelectedFixtureTypes.Num() == 1 && SelectedModeIndices.Num() == 1)
 	{
-		TSet<FString> FunctionNames;
-		for (const TWeakObjectPtr<UDMXEntityFixtureType>& FixtureType : SharedData->GetFixtureTypesBeingEdited())
+		if (UDMXEntityFixtureType* FixtureType = SelectedFixtureTypes[0].Get())
 		{
-			for (const FDMXFixtureMode& Mode : FixtureType->Modes)
+			if(FixtureType->Modes.IsValidIndex(SelectedModeIndices[0]))
 			{
-				for (const FDMXFixtureFunction& Function : Mode.Functions)
+				const FDMXFixtureMode& Mode = FixtureType->Modes[SelectedModeIndices[0]];
+
+				for (const int32 FunctionIndex : SelectedFunctionIndices)
 				{
-					FunctionNames.Add(Function.FunctionName);
+					if (!Mode.Functions.IsValidIndex(FunctionIndex))
+					{
+						return false;
+					}
 				}
 			}
 		}
-
-		return !FunctionNames.Contains(TestedName);
 	}
+
 	return true;
 }
 
-void FDMXFixtureFunctionItem::SetName(const FString& InDesiredName, FString& OutUniqueName)
+void FDMXFixtureTypeSharedData::CopyFunctionsToClipboard()
 {
-	if (TSharedPtr<FDMXFixtureTypeSharedData> SharedData = SharedDataPtr.Pin())
+	if (ensureMsgf(CanCopyFunctionsToClipboard(), TEXT("Cannot copy Functions to clipboard. Please test first with FDMXFixtureTypeSharedData::CopyFunctionsToClipboard.")))
 	{
-		TSet<FString> FunctionNames;
-		for (const TWeakObjectPtr<UDMXEntityFixtureType>& FixtureType : SharedData->GetFixtureTypesBeingEdited())
+		FunctionsClipboard.Reset();
+
+		if (SelectedFixtureTypes.Num() == 1 && SelectedModeIndices.Num() == 1)
 		{
-			for (const FDMXFixtureMode& Mode : FixtureType->Modes)
+			if (UDMXEntityFixtureType* FixtureType = SelectedFixtureTypes[0].Get())
 			{
-				for (const FDMXFixtureFunction& Function : Mode.Functions)
+				if (FixtureType->Modes.IsValidIndex(SelectedModeIndices[0]))
 				{
-					FunctionNames.Add(Function.FunctionName);
+					const FDMXFixtureMode& Mode = FixtureType->Modes[SelectedModeIndices[0]];
+
+					// Sort selected function indices by starting channel, so elements are copied in visual right order, not in order of how things were selected
+					SelectedFunctionIndices.Sort([Mode](const int32 FunctionIndexA, const int32 FunctionIndexB)
+						{
+							if (Mode.Functions.IsValidIndex(FunctionIndexA) && Mode.Functions.IsValidIndex(FunctionIndexB))
+							{
+								return Mode.Functions[FunctionIndexA].Channel <= Mode.Functions[FunctionIndexB].Channel;
+							}
+							return false;
+						});
+
+					for (const int32 FunctionIndex : SelectedFunctionIndices)
+					{
+						if (Mode.Functions.IsValidIndex(FunctionIndex))
+						{
+							const TOptional<FString> CopyStr = FDMXRuntimeUtils::SerializeStructToString(FDMXFixtureFunction::StaticStruct(), &Mode.Functions[FunctionIndex]);
+							if (CopyStr.IsSet())
+							{
+								FunctionsClipboard.Add(CopyStr.GetValue());
+							}
+						}
+					}
 				}
 			}
 		}
-
-		OutUniqueName = FDMXEditorUtils::GenerateUniqueNameFromExisting(FunctionNames, InDesiredName);
-		FunctionNameHandle->SetValue(OutUniqueName);
 	}
 }
 
-void FDMXFixtureTypeSharedData::SetFixtureTypesBeingEdited(const TArray<TWeakObjectPtr<UDMXEntityFixtureType>>& InFixtureTypesBeingEdited)
+bool FDMXFixtureTypeSharedData::CanPasteFunctionsFromClipboard() const
 {
-	FixtureTypesBeingEdited = InFixtureTypesBeingEdited;
+	if (SelectedFixtureTypes.Num() == 1 && SelectedModeIndices.Num() == 1 && ModesClipboard.Num() > 0)
+	{
+		if (UDMXEntityFixtureType* FixtureType = SelectedFixtureTypes[0].Get())
+		{
+			if (FixtureType->Modes.IsValidIndex(SelectedModeIndices[0]))
+			{
+				return true;
+			}
+		}
+	}
 
-	// Selections turn invalid
-	SelectedModes.Reset();
-	SelectedFunctions.Reset();
+	return true;
 }
 
-int32 FDMXFixtureTypeSharedData::GetNumSelectedModes() const
+void FDMXFixtureTypeSharedData::PasteFunctionsFromClipboard(TArray<int32>& OutNewlyAddedFunctionIndices)
 {
-	return SelectedModes.Num();
-}
+	if (ensureMsgf(CanPasteFunctionsFromClipboard(), TEXT("Cannot paste Functions from clipboard. Please test first with FDMXFixtureTypeSharedData::CanPasteFunctionsFromClipboard.")))
+	{
+		if (SelectedFixtureTypes.Num() == 1 && SelectedModeIndices.Num() == 1)
+		{
+			if (UDMXEntityFixtureType* FixtureType = SelectedFixtureTypes[0].Get())
+			{
+				const int32 SelectedModeIndex = SelectedModeIndices[0];
+				if (FixtureType->Modes.IsValidIndex(SelectedModeIndex))
+				{
+					FDMXFixtureMode& Mode = FixtureType->Modes[SelectedModeIndex];
+
+					// Sort selected function indices by starting channel, so elements are pasted in visual right order, not in order of how things were selected
+					SelectedFunctionIndices.Sort([Mode](const int32 FunctionIndexA, const int32 FunctionIndexB)
+						{
+							if (Mode.Functions.IsValidIndex(FunctionIndexA) && Mode.Functions.IsValidIndex(FunctionIndexB))
+							{
+								return Mode.Functions[FunctionIndexA].Channel <= Mode.Functions[FunctionIndexB].Channel;
+							}
+							return false;
+						});
 
-int32 FDMXFixtureTypeSharedData::GetNumSelectedFunctions() const
-{
-	return SelectedFunctions.Num();
-}
+					const FText TransactionText = FText::Format(LOCTEXT("PasteModesTransaction", "Paste Fixture {0}|plural(one=Mode, other=Modes)"), GetSelectedModeIndices().Num());
+					const FScopedTransaction PasteModesTransaction(TransactionText);
 
-void FDMXFixtureTypeSharedData::SelectModes(const TArray<TSharedPtr<FDMXFixtureModeItem>>& InModes)
-{
-	SelectedModes = InModes;
+					FixtureType->PreEditChange(UDMXEntityFixtureType::StaticClass()->FindPropertyByName(GET_MEMBER_NAME_CHECKED(FDMXFixtureMode, Functions)));
 
-	// Remove functions that are not contained in selected modes from selection
-	SelectedFunctions.RemoveAll([&](const TSharedPtr<FDMXFixtureFunctionItem>& FunctionItem) {
-		TSharedPtr<FDMXFixtureModeItem>* SelectedParentMode = SelectedModes.FindByPredicate([&](const TSharedPtr<FDMXFixtureModeItem>& ModeItem) {
-			return ModeItem->GetModeName() == FunctionItem->GetModeName();
-			});
-		return !SelectedParentMode;
-		});
+					int32 NumPastedElements = 0;
+					for (int32 ClipboardIndex = 0; ClipboardIndex < FunctionsClipboard.Num(); ClipboardIndex++)
+					{
+						TSharedPtr<FJsonObject> RootJsonObject;
+						TSharedRef<TJsonReader<>> Reader = TJsonReaderFactory<>::Create(FunctionsClipboard[ClipboardIndex]);
+						FJsonSerializer::Deserialize(Reader, RootJsonObject);
+
+						if (RootJsonObject.IsValid())
+						{
+							FDMXFixtureFunction FunctionFromClipboard;
+							if (FJsonObjectConverter::JsonObjectToUStruct(RootJsonObject.ToSharedRef(), FDMXFixtureFunction::StaticStruct(), &FunctionFromClipboard, 0, 0))
+							{
+								// Find an index to paste onto
+								const int32 PasteOntoIndex = [Mode, ClipboardIndex, this]()
+								{
+									if (SelectedFunctionIndices.IsValidIndex(ClipboardIndex))
+									{
+										return SelectedFunctionIndices[ClipboardIndex];
+									}
+									else if (SelectedFunctionIndices.Num() > 0)
+									{
+										return SelectedFunctionIndices.Last();
+									}
+									else
+									{
+										// Paste at the end of the Functions array
+										return Mode.Functions.Num() - 1;
+									}
+								}();
+
+								FDMXFixtureFunction NewFunction = FunctionFromClipboard;
+
+								// Insert the Function
+								const int32 DesiredIndex = PasteOntoIndex + 1 + NumPastedElements;
+								const int32 ResultingIndex = FixtureType->InsertFunction(SelectedModeIndex, DesiredIndex, NewFunction);
+								OutNewlyAddedFunctionIndices.Add(ResultingIndex);
+
+								NumPastedElements++;
+							}
+						}
+					}
+
+					if (NumPastedElements > 0)
+					{
+						FixtureType->UpdateChannelSpan(SelectedModeIndex);
+					}
 
-	OnModesSelected.Broadcast();
+					FixtureType->PostEditChange();
+				}
+			}
+		}
+	}
 }
 
-void FDMXFixtureTypeSharedData::SelectFunctions(const TArray<TSharedPtr<FDMXFixtureFunctionItem>>& InFunctions)
+void FDMXFixtureTypeSharedData::SelectFunctions(const TArray<int32>& InFunctionIndices)
 {
-	SelectedFunctions = InFunctions;
-	OnFunctionsSelected.Broadcast();
+	// DEPRECATED 5.0
+	if (SelectedFunctionIndices != InFunctionIndices)
+	{
+		SelectedFunctionIndices = InFunctionIndices;
+		OnFunctionsSelected.Broadcast();
+	}
 }
 
 bool FDMXFixtureTypeSharedData::CanAddMode() const
 {
-	return GetFixtureTypesBeingEdited().Num() == 1;
+	// DEPRECATED 5.0
+
+	return SelectedFixtureTypes.Num() == 1;
 }
 
-void FDMXFixtureTypeSharedData::AddMode(const FScopedTransaction& TransactionScope)
+void FDMXFixtureTypeSharedData::AddMode()
 {
-	check(CanAddMode());
-	if (!CanAddMode())
-	{
-		return;
-	}
+	// DEPRECATED 5.0
 
-	for (const TWeakObjectPtr<UDMXEntityFixtureType>& FixtureType : GetFixtureTypesBeingEdited())
+	PRAGMA_DISABLE_DEPRECATION_WARNINGS
+	const bool bCanAddMode = CanAddMode();
+	PRAGMA_ENABLE_DEPRECATION_WARNINGS
+
+	if (ensureMsgf(bCanAddMode, TEXT("Call to FDMXFixtureTypeSharedData::AddMode when no mode can be added")))
 	{
-		if (!FixtureType.IsValid())
+		for (const TWeakObjectPtr<UDMXEntityFixtureType>& FixtureType : SelectedFixtureTypes)
 		{
-			continue;
-		}
-
-		FDMXFixtureMode NewMode;
+			if (!FixtureType.IsValid())
+			{
+				continue;
+			}
 
-		// Make a unique name for the new mode
-		TSet<FString> ModeNames;
-		for (const FDMXFixtureMode& Mode : FixtureType->Modes)
-		{
-			ModeNames.Add(Mode.ModeName);
-		}
-		NewMode.ModeName = FDMXEditorUtils::GenerateUniqueNameFromExisting(ModeNames, LOCTEXT("DMXFixtureTypeSharedData.DefaultModeName", "Mode").ToString());
+			FDMXFixtureMode NewMode;
 
-		FProperty* ModesProperty = UDMXEntityFixtureType::StaticClass()->FindPropertyByName(GET_MEMBER_NAME_STRING_CHECKED(UDMXEntityFixtureType, Modes));
-		check(ModesProperty);
+			// Make a unique name for the new mode
+			TSet<FString> ModeNames;
+			for (const FDMXFixtureMode& Mode : FixtureType->Modes)
+			{
+				ModeNames.Add(Mode.ModeName);
+			}
+			NewMode.ModeName = FDMXRuntimeUtils::GenerateUniqueNameFromExisting(ModeNames, LOCTEXT("DefaultModeName", "Mode").ToString());
 
-		const FDMXFixtureTypeSharedDataDetails::FDMXScopedModeEditChangeChainProperty ScopedChange = FDMXFixtureTypeSharedDataDetails::FDMXScopedModeEditChangeChainProperty(FixtureType.Get());
+			const FDMXFixtureTypeSharedDataDetails::FDMXScopedModeEditChangeChainProperty_DEPRECATED ScopedChange(FixtureType.Get());
 
-		FixtureType->Modes.Add(NewMode);		
+			FixtureType->Modes.Add(NewMode);
+		}
 	}
 }
 
-/** Helper struct to keep references to modes in arrays while operating on it */
-struct FModeRef
+void FDMXFixtureTypeSharedData::DuplicateModes(const TArray<int32>& ModeÎndiciesToDuplicate)
 {
-	FModeRef(TWeakObjectPtr<UDMXEntityFixtureType> InFixtureType, const FString& InModeName)
-		: FixtureType(InFixtureType)
-		, ModeName(InModeName)
-	{}
+	// DEPRECATED 5.0
 
-	TWeakObjectPtr<UDMXEntityFixtureType> FixtureType;
-	FString ModeName;
-
-	static TArray<FModeRef> GetModeReferences(const TArray<TSharedPtr<FDMXFixtureModeItem>>& Modes)
-	{
-		TArray<FModeRef> Result;
-		for (const TSharedPtr<FDMXFixtureModeItem>& Mode : Modes)
+	for (const TWeakObjectPtr<UDMXEntityFixtureType>& WeakFixtureType : SelectedFixtureTypes)
+	{ 
+		if (UDMXEntityFixtureType* FixtureType = WeakFixtureType.Get())
 		{
-			TArray<TWeakObjectPtr<UDMXEntityFixtureType>> OuterFixtureTypes;
-			Mode->GetOuterFixtureTypes(OuterFixtureTypes);
-			for (const TWeakObjectPtr<UDMXEntityFixtureType>& FixtureType : OuterFixtureTypes)
+			for (int32 ModeIndex : ModeÎndiciesToDuplicate)
 			{
-				Result.Add(FModeRef(FixtureType, Mode->GetModeName()));
-			}
-		}
-		return Result;
-	}
-};
-
-void FDMXFixtureTypeSharedData::DuplicateModes(const TArray<TSharedPtr<FDMXFixtureModeItem>>& ModesToDuplicate)
-{
-	TArray<FModeRef> ModeReferences = FModeRef::GetModeReferences(ModesToDuplicate);
-
-	for (const FModeRef& ModeRef : ModeReferences)
-	{
-		if (!ModeRef.FixtureType.IsValid())
-		{
-			continue;
-		}
-
-		int32 ModeIndex = ModeRef.FixtureType->Modes.IndexOfByPredicate([&](const FDMXFixtureMode& Mode) {
-			return Mode.ModeName == ModeRef.ModeName;
-			});
-		check(ModeIndex != INDEX_NONE);
-
-		// Copy
-		FDMXFixtureMode NewMode = ModeRef.FixtureType->Modes[ModeIndex];
-
-		// Get a unique name
-		TSet<FString> ModeNames;
-		for (const FDMXFixtureMode& Mode : ModeRef.FixtureType->Modes)
-		{
-			ModeNames.Add(Mode.ModeName);
-		}
-		NewMode.ModeName = FDMXEditorUtils::GenerateUniqueNameFromExisting(ModeNames, ModeRef.ModeName);
+				if (FixtureType->Modes.IsValidIndex(ModeIndex))
+				{
+					// Copy
+					FDMXFixtureMode NewMode = FixtureType->Modes[ModeIndex];
 
-		// Transaction
-		const FScopedTransaction Transaction = FScopedTransaction(LOCTEXT("DMXFixtureTypeSharedData.ModeDuplicated", "DMX Editor: Duplicated Fixture Type Mode"));
+					// Get a unique name
+					TSet<FString> ModeNames;
+					for (const FDMXFixtureMode& Mode : FixtureType->Modes)
+					{
+						ModeNames.Add(Mode.ModeName);
+					}
+					NewMode.ModeName = FDMXRuntimeUtils::GenerateUniqueNameFromExisting(ModeNames, NewMode.ModeName);
 
-		const FDMXFixtureTypeSharedDataDetails::FDMXScopedModeEditChangeChainProperty ScopedChange = FDMXFixtureTypeSharedDataDetails::FDMXScopedModeEditChangeChainProperty(ModeRef.FixtureType.Get());
+					const FDMXFixtureTypeSharedDataDetails::FDMXScopedModeEditChangeChainProperty_DEPRECATED ScopedChange(FixtureType);
 
-		if (ModeRef.FixtureType->Modes.IsValidIndex(ModeIndex + 1))
-		{
-			ModeRef.FixtureType->Modes.Insert(NewMode, ModeIndex + 1);
-		}
-		else
-		{
-			ModeRef.FixtureType->Modes.Add(NewMode);
+					if (FixtureType->Modes.IsValidIndex(ModeIndex + 1))
+					{
+						FixtureType->Modes.Insert(NewMode, ModeIndex + 1);
+					}
+					else
+					{
+						FixtureType->Modes.Add(NewMode);
+					}
+				}
+			}
 		}
 	}
 }
 
-void FDMXFixtureTypeSharedData::DeleteModes(const TArray<TSharedPtr<FDMXFixtureModeItem>>& ModesToDelete)
+void FDMXFixtureTypeSharedData::DeleteModes(const TArray<int32>& ModeIndicesToDelete)
 {
-	TArray<FModeRef> ModeReferences = FModeRef::GetModeReferences(ModesToDelete);
+	// DEPRECATED 5.0
 
-	for (const FModeRef& ModeRef : ModeReferences)
+	for (const TWeakObjectPtr<UDMXEntityFixtureType>& WeakFixtureType : SelectedFixtureTypes)
 	{
-		if (!ModeRef.FixtureType.IsValid())
+		if (UDMXEntityFixtureType* FixtureType = WeakFixtureType.Get())
 		{
-			continue;
-		}
-
-		int32 ModeIndex = ModeRef.FixtureType->Modes.IndexOfByPredicate([&](const FDMXFixtureMode& Mode) {
-			return Mode.ModeName == ModeRef.ModeName;
-			});
-		check(ModeIndex != INDEX_NONE);
-
-		// Transaction
-		const FScopedTransaction Transaction = FScopedTransaction(LOCTEXT("DMXFixtureTypeSharedData.ModeDeleted", "DMX Editor: Deleted Fixture Type Mode"));
+			for (int32 ModeIndex : ModeIndicesToDelete)
+			{
+				if (FixtureType->Modes.IsValidIndex(ModeIndex))
+				{
+					// Transaction
+					const FScopedTransaction Transaction = FScopedTransaction(LOCTEXT("DeleteModeTransaction", "Delete Fixture Type Mode"));
 
-		const FDMXFixtureTypeSharedDataDetails::FDMXScopedModeEditChangeChainProperty ScopedChange = FDMXFixtureTypeSharedDataDetails::FDMXScopedModeEditChangeChainProperty(ModeRef.FixtureType.Get());
+					const FDMXFixtureTypeSharedDataDetails::FDMXScopedModeEditChangeChainProperty_DEPRECATED ScopedChange(FixtureType);
 
-		ModeRef.FixtureType->Modes.RemoveAt(ModeIndex);
+					FixtureType->Modes.RemoveAt(ModeIndex);
+				}
+			}
+		}
 	}
 }
 
+void FDMXFixtureTypeSharedData::PasteClipboardToModes(const TArray<int32>& ModeIndices)
+{	
+	// DEPRECATED 5.0
 
-void FDMXFixtureTypeSharedData::CopyModesToClipboard(const TArray<TSharedPtr<FDMXFixtureModeItem>>& ModeItems)
-{
-	ModesClipboard.Reset();
-	for (TSharedPtr<FDMXFixtureModeItem> ModeItem : ModeItems)
+	for (int32 PasteOntoElementIndex = 0; PasteOntoElementIndex < ModeIndices.Num(); PasteOntoElementIndex++)
 	{
-		for(const TWeakObjectPtr<UDMXEntityFixtureType>& FixtureType : GetFixtureTypesBeingEdited())
-		{ 
-			if (!FixtureType.IsValid())
+		const int32 PasteOntoIndex = ModeIndices[PasteOntoElementIndex];
+
+		const FString& ClipboardString = [PasteOntoElementIndex, this]()
+		{
+			if (ModesClipboard.IsValidIndex(PasteOntoElementIndex))
 			{
-				continue;
+				return ModesClipboard[PasteOntoElementIndex];
 			}
-			FDMXFixtureMode* ModeInAssetPtr = FixtureType->Modes.FindByPredicate([&](const FDMXFixtureMode& ModeInAsset) {
-				return ModeInAsset.ModeName == ModeItem->GetModeName();
-				});
-			check(ModeInAssetPtr);
-
-			FString CopyStr;
-			if (SerializeStruct<FDMXFixtureMode>(*ModeInAssetPtr, CopyStr))
+			else if (ModesClipboard.Num() > 0)
 			{
-				ModesClipboard.Add(CopyStr);
+				return ModesClipboard.Last();
 			}
-		}
-	}
-}
 
-void FDMXFixtureTypeSharedData::PasteClipboardToModes(const TArray<TSharedPtr<FDMXFixtureModeItem>>& ModeItems)
-{
-	for (int32 IndexOfMode = 0; IndexOfMode < ModeItems.Num(); IndexOfMode++)
-	{
-		const FString& ClipboardElement = ModesClipboard[IndexOfMode];
+			return FString();
+		}();
 
 		TSharedPtr<FJsonObject> RootJsonObject;
-		TSharedRef<TJsonReader<>> Reader = TJsonReaderFactory<>::Create(ClipboardElement);
+		TSharedRef<TJsonReader<>> Reader = TJsonReaderFactory<>::Create(ClipboardString);
 		FJsonSerializer::Deserialize(Reader, RootJsonObject);
 
 		if (RootJsonObject.IsValid())
 		{
 			FDMXFixtureMode ModeFromClipboard;
-
 			if (FJsonObjectConverter::JsonObjectToUStruct(RootJsonObject.ToSharedRef(), FDMXFixtureMode::StaticStruct(), &ModeFromClipboard, 0, 0))
 			{
-				for (const TWeakObjectPtr<UDMXEntityFixtureType>& FixtureType : GetFixtureTypesBeingEdited())
+				for (const TWeakObjectPtr<UDMXEntityFixtureType>& WeakFixtureType : SelectedFixtureTypes)
 				{
-					if (!FixtureType.IsValid())
+					if (UDMXEntityFixtureType* FixtureType = WeakFixtureType.Get())
 					{
-						continue;
-					}
+						if (FixtureType->Modes.IsValidIndex(PasteOntoIndex))
+						{
+							FDMXFixtureMode& PasteOntoMode = FixtureType->Modes[PasteOntoIndex];
 
-					FDMXFixtureMode* ModePtr = FixtureType->Modes.FindByPredicate([&](const FDMXFixtureMode& ModeInAsset) {
-						return ModeInAsset.ModeName == ModeItems[IndexOfMode]->GetModeName();
-						});
-					if (!ModePtr)
-					{
-						continue;
-					}
+							FString OldModeName = PasteOntoMode.ModeName;
 
-					FString OldModeName = ModePtr->ModeName;
-					FDMXFixtureMode& Mode = *ModePtr;
-					
-					FixtureType->PreEditChange(UDMXEntityFixtureType::StaticClass()->FindPropertyByName(GET_MEMBER_NAME_STRING_CHECKED(UDMXEntityFixtureType, Modes)));
-					const FScopedTransaction Transaction = FScopedTransaction(LOCTEXT("DMXFixtureTypeSharedData.FixtureTypeModePasted", "DMX Editor: Pasted Fixture Type Mode"));
-					FixtureType->Modify();
+							FixtureType->PreEditChange(UDMXEntityFixtureType::StaticClass()->FindPropertyByName(GET_MEMBER_NAME_STRING_CHECKED(UDMXEntityFixtureType, Modes)));
+							const FScopedTransaction Transaction = FScopedTransaction(LOCTEXT("PasteModeTransaction", "Paste Fixture Type Mode"));
+							FixtureType->Modify();
 
-					Mode = ModeFromClipboard;
+							PasteOntoMode = ModeFromClipboard;
 
-					// Keep the original mode name
-					Mode.ModeName = OldModeName;
+							// Keep the original mode name
+							PasteOntoMode.ModeName = OldModeName;
 
-					FixtureType->PostEditChange();
+							FixtureType->PostEditChange();
+						}
+					}
 				}
 			}
 		}
@@ -584,232 +588,166 @@ void FDMXFixtureTypeSharedData::PasteClipboardToModes(const TArray<TSharedPtr<FD
 
 bool FDMXFixtureTypeSharedData::CanAddFunction() const
 {
-	return GetNumSelectedModes() == 1;
+	// DEPRECATED 5.0
+
+	return SelectedModeIndices.Num() == 1;
 }
 
 void FDMXFixtureTypeSharedData::AddFunctionToSelectedMode()
 {
-	check(CanAddFunction());
-	if (!CanAddFunction())
-	{
-		return;
-	}
-	TArray<FModeRef> ModeReferences = FModeRef::GetModeReferences(TArray<TSharedPtr<FDMXFixtureModeItem>>({ SelectedModes[0] }));
+	// DEPRECATED 5.0
 
-	for (const FModeRef& ModeRef : ModeReferences)
-	{
-		if (!ModeRef.FixtureType.IsValid())
-		{
-			continue;
-		}
+	PRAGMA_DISABLE_DEPRECATION_WARNINGS
+	const bool bCanAddFunction = CanAddFunction();
+	PRAGMA_ENABLE_DEPRECATION_WARNINGS
 
-		FDMXFixtureMode* ModePtr = ModeRef.FixtureType->Modes.FindByPredicate([&](const FDMXFixtureMode& ModeInAsset) {
-			return ModeInAsset.ModeName == SelectedModes[0]->GetModeName();
-			});
-		if (!ModePtr)
+	if (ensureMsgf(bCanAddFunction, TEXT("Call to FDMXFixtureTypeSharedData::AddFunctionToSelectedMode, when no function can be added")))
+	{
+		for (const TWeakObjectPtr<UDMXEntityFixtureType>& WeakFixtureType : SelectedFixtureTypes)
 		{
-			continue;
-		}
-		FDMXFixtureMode& Mode = *ModePtr;
+			if (UDMXEntityFixtureType* FixtureType = WeakFixtureType.Get())
+			{
+				const FScopedTransaction AddFunctionTransaction = FScopedTransaction(LOCTEXT("AddFunctionTransaction", "Add Fixture Type Function"));
 
-		// Get a unique name
-		TSet<FString> FunctionNames;
-		for (const FDMXFixtureFunction& Function : Mode.Functions)
-		{
-			FunctionNames.Add(Function.FunctionName);
-		}
-		FDMXFixtureFunction NewFunction;
-		NewFunction.FunctionName = FDMXEditorUtils::GenerateUniqueNameFromExisting(FunctionNames, LOCTEXT("DMXFixtureTypeSharedData.DefaultFunctionName", "Function").ToString());
+				FixtureType->PreEditChange(UDMXEntityFixtureType::StaticClass()->FindPropertyByName(GET_MEMBER_NAME_STRING_CHECKED(UDMXEntityFixtureType, Modes)));
+				FixtureType->Modify();
 
-		const FScopedTransaction ModeNameTransaction = FScopedTransaction(LOCTEXT("DMXFixtureTypeSharedData.AddedFunction", "DMX Editor: Added Fixture Type Function"));
+				int32 ModeIndex = SelectedModeIndices[0]; // Implicitly valid via CanAddFunctions tested before
+				if (FixtureType->Modes.IsValidIndex(ModeIndex))
+				{
+					FDMXFixtureMode& Mode = FixtureType->Modes[ModeIndex];
 
-		ModeRef.FixtureType->PreEditChange(UDMXEntityFixtureType::StaticClass()->FindPropertyByName(GET_MEMBER_NAME_STRING_CHECKED(UDMXEntityFixtureType, Modes)));
-		ModeRef.FixtureType->Modify();
+					// Get a unique name
+					TSet<FString> FunctionNames;
+					for (const FDMXFixtureFunction& Function : Mode.Functions)
+					{
+						FunctionNames.Add(Function.FunctionName);
+					}
+					FDMXFixtureFunction NewFunction;
+					NewFunction.FunctionName = FDMXRuntimeUtils::GenerateUniqueNameFromExisting(FunctionNames, LOCTEXT("DMXFixtureTypeSharedData.DefaultFunctionName", "Function").ToString());
 
-		Mode.Functions.Add(NewFunction);
-		ModeRef.FixtureType->UpdateChannelSpan(Mode);
+					Mode.Functions.Add(NewFunction);
+					FixtureType->UpdateChannelSpan(ModeIndex);
+				}
 
-		ModeRef.FixtureType->PostEditChange();
+				FDMXEditorUtils::AutoAssignedAddresses(FixtureType);
 
-		FDMXEditorUtils::AutoAssignedAddresses(ModeRef.FixtureType.Get());
+				FixtureType->PostEditChange();
+			}
+		}
 	}
 }
 
-/** Helper struct to keep references to functions in arrays while operating on it */
-struct FFunctionRef
+void FDMXFixtureTypeSharedData::DuplicateFunctions(const TArray<int32>& FunctionIndicesToDuplicate)
 {
-	FFunctionRef(TWeakObjectPtr<UDMXEntityFixtureType> InFixtureType, const FString& InModeName, const FString& InFunctionName)
-		: FixtureType(InFixtureType)
-		, ModeName(InModeName)
-		, FunctionName(InFunctionName)	
-	{}
-
-	TWeakObjectPtr<UDMXEntityFixtureType> FixtureType;
-	FString ModeName;
-	FString FunctionName;
+	// DEPRECATED 5.0
 
-	static TArray<FFunctionRef> GetFunctionReferences(const TArray<TSharedPtr<FDMXFixtureFunctionItem>>& Functions)
+	for (const TWeakObjectPtr<UDMXEntityFixtureType>& WeakFixtureType : SelectedFixtureTypes)
 	{
-		TArray<FFunctionRef> Result;
-		for (const TSharedPtr<FDMXFixtureFunctionItem>& Function : Functions)
+		if (UDMXEntityFixtureType* FixtureType = WeakFixtureType.Get())
 		{
-			TArray<TWeakObjectPtr<UDMXEntityFixtureType>> OuterFixtureTypes;
-			Function->GetOuterFixtureTypes(OuterFixtureTypes);
-			for (const TWeakObjectPtr<UDMXEntityFixtureType>& FixtureType : OuterFixtureTypes)
+			for (int32 ModeIndex : SelectedModeIndices)
 			{
-				Result.Add(FFunctionRef(FixtureType, Function->GetModeName(), Function->GetFunctionName()));
-			}
-		}
-		return Result;
-	}
-};
+				if (FixtureType->Modes.IsValidIndex(ModeIndex))
+				{
+					FDMXFixtureMode& Mode = FixtureType->Modes[ModeIndex];
 
-void FDMXFixtureTypeSharedData::DuplicateFunctions(const TArray<TSharedPtr<FDMXFixtureFunctionItem>>& FunctionsToDuplicate)
-{
-	TArray<FFunctionRef> FunctionReferences = FFunctionRef::GetFunctionReferences(FunctionsToDuplicate);
+					// Transaction
+					const FScopedTransaction Transaction = FScopedTransaction(LOCTEXT("DuplicatFunctionTransaction", "Duplicate Fixture Type Function"));
 
-	for (const FFunctionRef& FuncRef : FunctionReferences)
-	{
-		if (!FuncRef.FixtureType.IsValid())
-		{
-			continue;
-		}
+					FixtureType->PreEditChange(UDMXEntityFixtureType::StaticClass()->FindPropertyByName(GET_MEMBER_NAME_STRING_CHECKED(UDMXEntityFixtureType, Modes)));
+					FixtureType->Modify();
 
-		FDMXFixtureMode* ModePtr = FuncRef.FixtureType->Modes.FindByPredicate([&](const FDMXFixtureMode& Mode) {
-			return Mode.ModeName == FuncRef.ModeName;
-			});
-		check(ModePtr);
+					int32 NumDuplicates = 0;
+					for (int32 FunctionIndex : FunctionIndicesToDuplicate)
+					{
+						const TArray<int32> DuplicateAtIndices({ FunctionIndex + NumDuplicates });
+						TArray<int32> NewlyAddedFunctionIndices;
+						FixtureType->DuplicateFunctions(ModeIndex, DuplicateAtIndices, NewlyAddedFunctionIndices);
 
-		int32 IndexOfFunction = ModePtr->Functions.IndexOfByPredicate([&](const FDMXFixtureFunction& Function) {
-			return Function.FunctionName == FuncRef.FunctionName;
-			});
-		check(IndexOfFunction != INDEX_NONE);
+						NumDuplicates++;
+					}
 
-		// Copy
-		FDMXFixtureFunction NewFunction = ModePtr->Functions[IndexOfFunction];
+					FixtureType->UpdateChannelSpan(ModeIndex);
 
-		// Get a unique name
-		TSet<FString> FunctionNames;
-		for (const FDMXFixtureFunction& Function : ModePtr->Functions)
-		{
-			FunctionNames.Add(Function.FunctionName);
+					FixtureType->PostEditChange();
+				}
+			}
 		}
-		NewFunction.FunctionName = FDMXEditorUtils::GenerateUniqueNameFromExisting(FunctionNames, FuncRef.FunctionName);
-
-		// Transaction
-		const FScopedTransaction Transaction = FScopedTransaction(LOCTEXT("DMXFixtureTypeSharedData.DuplicatedFunction", "DMX Editor: Duplicated Fixture Type Function"));
-
-		FuncRef.FixtureType->PreEditChange(UDMXEntityFixtureType::StaticClass()->FindPropertyByName(GET_MEMBER_NAME_STRING_CHECKED(UDMXEntityFixtureType, Modes)));
-		FuncRef.FixtureType->Modify();
-
-		ModePtr->AddOrInsertFunction(IndexOfFunction, NewFunction);
-
-		FuncRef.FixtureType->PostEditChange();
 	}
 }
 
-void FDMXFixtureTypeSharedData::DeleteFunctions(const TArray<TSharedPtr<FDMXFixtureFunctionItem>>& FunctionsToDelete)
+void FDMXFixtureTypeSharedData::DeleteFunctions(const TArray<int32>& FunctionIndicesToDelete)
 {
-	FArrayProperty* ModesProperty = FindFProperty<FArrayProperty>(UDMXEntityFixtureType::StaticClass(), GET_MEMBER_NAME_CHECKED(UDMXEntityFixtureType, Modes));
-	FArrayProperty* FunctionsProperty = FindFProperty<FArrayProperty>(FDMXFixtureMode::StaticStruct(), GET_MEMBER_NAME_CHECKED(FDMXFixtureMode, Functions));
-	FEditPropertyChain PropertyChain;
-	PropertyChain.AddHead(ModesProperty);
-	PropertyChain.AddTail(FunctionsProperty);
-	PropertyChain.SetActiveMemberPropertyNode(ModesProperty);
-	PropertyChain.SetActivePropertyNode(FunctionsProperty);
-
-	TArray<FFunctionRef> FunctionReferences = FFunctionRef::GetFunctionReferences(FunctionsToDelete);
-	if (FunctionReferences.Num() == 0)
-	{
-		return;
-	}
+	// DEPRECATED 5.0
 
-	UDMXEntityFixtureType* FixtureType = FunctionReferences[0].FixtureType.Get();
-	if (!FixtureType)
+	for (const TWeakObjectPtr<UDMXEntityFixtureType>& WeakFixtureType : SelectedFixtureTypes)
 	{
-		return;
-	}
+		if (UDMXEntityFixtureType* FixtureType = WeakFixtureType.Get())
+		{
+			for (int32 ModeIndex : SelectedModeIndices)
+			{
+				if (FixtureType->Modes.IsValidIndex(ModeIndex))
+				{
+					FDMXFixtureMode& Mode = FixtureType->Modes[ModeIndex];
 
-	const FScopedTransaction Transaction = FScopedTransaction(LOCTEXT("DMXFixtureTypeSharedData.DeletedFunction", "DMX Editor: Deleted Fixture Type Function"));
+					const FScopedTransaction Transaction = FScopedTransaction(LOCTEXT("DeleteFunctionsTransaction", "Deleted Fixture Type Functions"));
 
-	FixtureType->PreEditChange(PropertyChain);
-	FixtureType->Modify();
+					// Required this way for the 5.0 Deprecated UDMXEntityFixtureType::OnDataTypeChanged event
+					FArrayProperty* ModesProperty = FindFProperty<FArrayProperty>(UDMXEntityFixtureType::StaticClass(), GET_MEMBER_NAME_CHECKED(UDMXEntityFixtureType, Modes));
+					FArrayProperty* FunctionsProperty = FindFProperty<FArrayProperty>(FDMXFixtureMode::StaticStruct(), GET_MEMBER_NAME_CHECKED(FDMXFixtureMode, Functions));
+					FEditPropertyChain PropertyChain;
+					PropertyChain.AddHead(ModesProperty);
+					PropertyChain.AddTail(FunctionsProperty);
+					PropertyChain.SetActiveMemberPropertyNode(ModesProperty);
+					PropertyChain.SetActivePropertyNode(FunctionsProperty);
 
-	for (const FFunctionRef& FuncRef : FunctionReferences)
-	{
-		check(FuncRef.FixtureType == FixtureType);
+					FixtureType->PreEditChange(FunctionsProperty);
+					FixtureType->Modify();
 
-		FDMXFixtureMode* ModePtr = FuncRef.FixtureType->Modes.FindByPredicate([&](const FDMXFixtureMode& Mode) {
-			return Mode.ModeName == FuncRef.ModeName;
-			});
-		check(ModePtr);
+					for (const int32 FunctionIndex : FunctionIndicesToDelete)
+					{
+						if (Mode.Functions.IsValidIndex(FunctionIndex))
+						{
+							const TArray<int32> RemoveAtIndices({ FunctionIndex });
+							FixtureType->RemoveFunctions(ModeIndex, RemoveAtIndices);
+						}
+					}
 
-		int32 IndexOfFunction = ModePtr->Functions.IndexOfByPredicate([&](const FDMXFixtureFunction& Function) {
-			return Function.FunctionName == FuncRef.FunctionName;
-			});
-		check(IndexOfFunction != INDEX_NONE);
+					FPropertyChangedEvent PropertyChangedEvent(FunctionsProperty, EPropertyChangeType::ArrayRemove);
+					FPropertyChangedChainEvent PropertyChangedChainEvent(PropertyChain, PropertyChangedEvent);
 
-		ModePtr->Functions.RemoveAt(IndexOfFunction);
+					FixtureType->PostEditChangeChainProperty(PropertyChangedChainEvent);
+				}
+			}
+		}
 	}
-
-	FPropertyChangedEvent PropertyChangedEvent(FunctionsProperty, EPropertyChangeType::ArrayRemove);
-	FPropertyChangedChainEvent PropertyChangedChainEvent(PropertyChain, PropertyChangedEvent);
-
-	FixtureType->PostEditChangeChainProperty(PropertyChangedChainEvent);
-
-	FDMXEditorUtils::AutoAssignedAddresses(FixtureType);
 }
 
-void FDMXFixtureTypeSharedData::CopyFunctionsToClipboard(const TArray<TSharedPtr<FDMXFixtureFunctionItem>>& FunctionItems)
+void FDMXFixtureTypeSharedData::PasteClipboardToFunctions(const TArray<int32>& FunctionIndices)
 {
-	FunctionsClipboard.Reset();
-	for (TSharedPtr<FDMXFixtureFunctionItem> Item : FunctionItems)
+	// DEPRECATED 5.0
+
+	for (int32 PasteOntoElementIndex = 0; PasteOntoElementIndex < FunctionIndices.Num(); PasteOntoElementIndex++)
 	{
-		for (const TWeakObjectPtr<UDMXEntityFixtureType>& FixtureType : GetFixtureTypesBeingEdited())
-		{
-			if (!FixtureType.IsValid())
-			{
-				continue;
-			}
+		const int32 PasteOntoIndex = FunctionIndices[PasteOntoElementIndex];
 
-			const FDMXFixtureMode* ModePtr = FixtureType->Modes.FindByPredicate([&](const FDMXFixtureMode& ModeInAsset) {
-				return ModeInAsset.ModeName == Item->GetModeName();
-				});
-			if (!ModePtr)
+		const FString& ClipboardString = [PasteOntoElementIndex, this]()
+		{
+			if (FunctionsClipboard.IsValidIndex(PasteOntoElementIndex))
 			{
-				continue;
+				return FunctionsClipboard[PasteOntoElementIndex];
 			}
-
-			const FDMXFixtureFunction* FunctionPtr = ModePtr->Functions.FindByPredicate([&](const FDMXFixtureFunction& FunctionInAsset) {
-				return FunctionInAsset.FunctionName == Item->GetFunctionName();
-				});
-			if (!FunctionPtr)
+			else if (FunctionsClipboard.Num() > 0)
 			{
-				continue;
+				return FunctionsClipboard.Last();
 			}
 
-			FString CopyStr;
-			if (SerializeStruct<FDMXFixtureFunction>(*FunctionPtr, CopyStr))
-			{
-				FunctionsClipboard.Add(CopyStr);
-			}
-		}
-	}
-}
-
-void FDMXFixtureTypeSharedData::PasteClipboardToFunctions(const TArray<TSharedPtr<FDMXFixtureFunctionItem>>& FunctionItems)
-{
-	for (int32 IndexOfFunction = 0; IndexOfFunction < FunctionItems.Num(); IndexOfFunction++)
-	{
-		if (!FunctionsClipboard.IsValidIndex(IndexOfFunction))
-		{
-			return;
-		}
-
-		const FString& ClipboardElement = FunctionsClipboard[IndexOfFunction];
+			return FString();
+		}();
 
 		TSharedPtr<FJsonObject> RootJsonObject;
-		TSharedRef<TJsonReader<>> Reader = TJsonReaderFactory<>::Create(ClipboardElement);
+		TSharedRef<TJsonReader<>> Reader = TJsonReaderFactory<>::Create(ClipboardString);
 		FJsonSerializer::Deserialize(Reader, RootJsonObject);
 
 		if (RootJsonObject.IsValid())
@@ -817,44 +755,82 @@ void FDMXFixtureTypeSharedData::PasteClipboardToFunctions(const TArray<TSharedPt
 			FDMXFixtureFunction FunctionFromClipboard;
 			if (FJsonObjectConverter::JsonObjectToUStruct(RootJsonObject.ToSharedRef(), FDMXFixtureFunction::StaticStruct(), &FunctionFromClipboard, 0, 0))
 			{
-				for (const TWeakObjectPtr<UDMXEntityFixtureType>& FixtureType : GetFixtureTypesBeingEdited())
+				for (const TWeakObjectPtr<UDMXEntityFixtureType>& WeakFixtureType : SelectedFixtureTypes)
 				{
-					if (!FixtureType.IsValid())
+					if (UDMXEntityFixtureType* FixtureType = WeakFixtureType.Get())
 					{
-						continue;
-					}
-					FDMXFixtureMode* ModeInAssetPtr = FixtureType->Modes.FindByPredicate([&](const FDMXFixtureMode& ModeInAsset) {
-						return ModeInAsset.ModeName == FunctionItems[IndexOfFunction]->GetModeName();
-						});
-					if (!ModeInAssetPtr)
-					{
-						continue;
-					}
-					FDMXFixtureFunction* FunctionPtr = ModeInAssetPtr->Functions.FindByPredicate([&](const FDMXFixtureFunction& FunctionInAsset) {
-						return FunctionInAsset.FunctionName == FunctionItems[IndexOfFunction]->GetFunctionName();
-						});
-					if (!FunctionPtr)
-					{
-						continue;
-					}
+						const FScopedTransaction ModeNameTransaction = FScopedTransaction(LOCTEXT("PasteModeTransaction", "Pasted Fixture Type Mode"));
 
-					FString OldFunctionName = FunctionPtr->FunctionName;
-					FDMXFixtureFunction& Function = *FunctionPtr;					
+						FixtureType->PreEditChange(UDMXEntityFixtureType::StaticClass()->FindPropertyByName(GET_MEMBER_NAME_STRING_CHECKED(UDMXEntityFixtureType, Modes)));
+						FixtureType->Modify();
 
-					FixtureType->PreEditChange(UDMXEntityFixtureType::StaticClass()->FindPropertyByName(GET_MEMBER_NAME_STRING_CHECKED(UDMXEntityFixtureType, Modes)));
-					const FScopedTransaction ModeNameTransaction = FScopedTransaction(LOCTEXT("DMXFixtureTypeSharedData.FixtureTypeModePasted", "DMX Editor: Pasted Fixture Type Mode"));
-					FixtureType->Modify();
+						for (int32 ModeIndex : SelectedModeIndices)
+						{
+							if (FixtureType->Modes.IsValidIndex(ModeIndex))
+							{
+								FDMXFixtureMode& Mode = FixtureType->Modes[ModeIndex];
 
-					Function = FunctionFromClipboard;
+								if (Mode.Functions.IsValidIndex(PasteOntoIndex))
+								{
+									FDMXFixtureFunction& PasteOntoFunction = Mode.Functions[PasteOntoIndex];
 
-					// Keep the original function name
-					Function.FunctionName = OldFunctionName;
+									FString OldFunctionName = PasteOntoFunction.FunctionName;
 
-					FixtureType->PostEditChange();
+									PasteOntoFunction = FunctionFromClipboard;
+
+									// Keep the original function name
+									PasteOntoFunction.FunctionName = OldFunctionName;
+								}
+							}
+						}
+
+						FixtureType->PostEditChange();
+					}
 				}
 			}
 		}
 	}
 }
 
+bool FDMXFixtureTypeSharedData::CanAddCellAttribute() const
+{
+	// DEPRECATED 5.0
+
+	return 
+		SelectedFixtureTypes.Num() == 1 && 
+		SelectedModeIndices.Num() == 1;
+}
+
+void FDMXFixtureTypeSharedData::AddCellAttributeToSelectedMode()
+{
+	// DEPRECATED 5.0
+
+	PRAGMA_DISABLE_DEPRECATION_WARNINGS
+	const bool bCanAddCellAttribute = CanAddCellAttribute();
+	PRAGMA_ENABLE_DEPRECATION_WARNINGS
+
+	if (ensureMsgf(bCanAddCellAttribute, TEXT("Call to FDMXFixtureTypeSharedData::AddCellAttributeToSelectedMode, when no Cell Attribute can be added.")))
+	{
+		if (UDMXEntityFixtureType* FixtureType = SelectedFixtureTypes[0].Get())
+		{
+			const int32 SelectedModeIndex = SelectedModeIndices[0];
+			if (FixtureType->Modes.IsValidIndex(SelectedModeIndex) && FixtureType->Modes[SelectedModeIndex].bFixtureMatrixEnabled)
+			{
+				FDMXFixtureMode& Mode = FixtureType->Modes[SelectedModeIndices[0]];
+
+				FDMXFixtureCellAttribute NewAttribute;
+				TArray<FName> AvailableAttributes = FDMXAttributeName::GetPossibleValues();
+				NewAttribute.Attribute = AvailableAttributes.Num() > 0 ? FDMXAttributeName(AvailableAttributes[0]) : FDMXAttributeName();
+
+				const FDMXFixtureTypeSharedDataDetails::FDMXScopedModeEditChangeChainProperty_DEPRECATED ScopedChange(FixtureType);
+
+				Mode.FixtureMatrixConfig.CellAttributes.Add(NewAttribute);
+
+				// Update affected properties
+				FixtureType->UpdateChannelSpan(SelectedModeIndex);
+			}
+		}
+	}
+}
+
 #undef LOCTEXT_NAMESPACE
diff --git a/Engine/Plugins/VirtualProduction/DMX/DMXEngine/Source/DMXEditor/Private/Sequencer/MovieSceneDMXLibraryTrackRecorder.cpp b/Engine/Plugins/VirtualProduction/DMX/DMXEngine/Source/DMXEditor/Private/Sequencer/MovieSceneDMXLibraryTrackRecorder.cpp
--- a/Engine/Plugins/VirtualProduction/DMX/DMXEngine/Source/DMXEditor/Private/Sequencer/MovieSceneDMXLibraryTrackRecorder.cpp
+++ b/Engine/Plugins/VirtualProduction/DMX/DMXEngine/Source/DMXEditor/Private/Sequencer/MovieSceneDMXLibraryTrackRecorder.cpp
@@ -7,7 +7,6 @@
 #include "DMXStats.h"
 #include "DMXSubsystem.h"
 #include "DMXTypes.h"
-#include "DMXUtils.h"
 #include "IO/DMXInputPort.h"
 #include "IO/DMXOutputPort.h"
 #include "IO/DMXRawListener.h"
diff --git a/Engine/Plugins/VirtualProduction/DMX/DMXEngine/Source/DMXEditor/Private/Style/DMXEditorStyle.cpp b/Engine/Plugins/VirtualProduction/DMX/DMXEngine/Source/DMXEditor/Private/Style/DMXEditorStyle.cpp
--- a/Engine/Plugins/VirtualProduction/DMX/DMXEngine/Source/DMXEditor/Private/Style/DMXEditorStyle.cpp
+++ b/Engine/Plugins/VirtualProduction/DMX/DMXEngine/Source/DMXEditor/Private/Style/DMXEditorStyle.cpp
@@ -46,54 +46,65 @@ TSharedRef<FSlateStyleSet> FDMXEditorStyle::Create()
 	TSharedRef<FSlateStyleSet> Style = MakeShared<FSlateStyleSet>(GetStyleSetName());
 	Style->SetContentRoot(IPluginManager::Get().FindPlugin(TEXT("DMXEngine"))->GetBaseDir() / TEXT("Resources"));
 
-	// Solid color brushes
-	Style->Set("DMXEditor.WhiteBrush", new FSlateColorBrush(FLinearColor(1.0f, 1.0f, 1.0f, 1.0f)));
-	Style->Set("DMXEditor.BlackBrush", new FSlateColorBrush(FLinearColor(0.0f, 0.0f, 0.0f, 1.0f)));
-
-	// Fonts
-	const TCHAR* FontPathRoboto = TEXT("Font'/Engine/EngineFonts/Roboto.Roboto'");
-	const UFont* FontRoboto = Cast<UFont>(StaticLoadObject(UFont::StaticClass(), nullptr, FontPathRoboto));
-	check(FontRoboto != nullptr);
-
-	Style->Set("DMXEditor.Font.InputChannelID", FSlateFontInfo(FontRoboto, 8, FName(TEXT("Light"))));
-	Style->Set("DMXEditor.Font.InputChannelValue", FSlateFontInfo(FontRoboto, 10, FName(TEXT("Regular"))));
-
-	Style->Set("DMXEditor.Font.InputUniverseHeader", FSlateFontInfo(FontRoboto, 10, FName(TEXT("Bold"))));
-	Style->Set("DMXEditor.Font.InputUniverseID", FSlateFontInfo(FontRoboto, 10, FName(TEXT("Regular"))));
-	Style->Set("DMXEditor.Font.InputUniverseChannelID", FSlateFontInfo(FontRoboto, 10, FName(TEXT("Regular"))));
-	Style->Set("DMXEditor.Font.InputUniverseChannelValue", FSlateFontInfo(FontRoboto, 10, FName(TEXT("Light"))));
-	
+	// Default color brushes
+	{
+		Style->Set("DMXEditor.WhiteBrush", new FSlateColorBrush(FLinearColor(1.0f, 1.0f, 1.0f, 1.0f)));
+		Style->Set("DMXEditor.BlackBrush", new FSlateColorBrush(FLinearColor(0.0f, 0.0f, 0.0f, 1.0f)));
+	}
+
+	// Default Fonts
+	{
+		const TCHAR* FontPathRoboto = TEXT("Font'/Engine/EngineFonts/Roboto.Roboto'");
+		const UFont* FontRoboto = Cast<UFont>(StaticLoadObject(UFont::StaticClass(), nullptr, FontPathRoboto));
+		check(FontRoboto != nullptr);
+
+		Style->Set("DMXEditor.Font.InputChannelID", FSlateFontInfo(FontRoboto, 8, FName(TEXT("Light"))));
+		Style->Set("DMXEditor.Font.InputChannelValue", FSlateFontInfo(FontRoboto, 10, FName(TEXT("Regular"))));
+
+		Style->Set("DMXEditor.Font.InputUniverseHeader", FSlateFontInfo(FontRoboto, 10, FName(TEXT("Bold"))));
+		Style->Set("DMXEditor.Font.InputUniverseID", FSlateFontInfo(FontRoboto, 10, FName(TEXT("Regular"))));
+		Style->Set("DMXEditor.Font.InputUniverseChannelID", FSlateFontInfo(FontRoboto, 10, FName(TEXT("Regular"))));
+		Style->Set("DMXEditor.Font.InputUniverseChannelValue", FSlateFontInfo(FontRoboto, 10, FName(TEXT("Light"))));
+	}
+
 	// Level editor toolbar icons
-	Style->Set("DMXEditor.LevelEditor.MenuIcon_none", new IMAGE_BRUSH("Icons/LevelEditor/icon_DMXTools_none_50x40", Icon50x40));
-	Style->Set("DMXEditor.LevelEditor.MenuIcon_snd", new IMAGE_BRUSH("Icons/LevelEditor/icon_DMXTools_snd_50x40", Icon50x40));
-	Style->Set("DMXEditor.LevelEditor.MenuIcon_rcv", new IMAGE_BRUSH("Icons/LevelEditor/icon_DMXTools_rcv_50x40", Icon50x40));
-	Style->Set("DMXEditor.LevelEditor.MenuIcon_snd-rcv", new IMAGE_BRUSH("Icons/LevelEditor/icon_DMXTools_snd-rcv_50x40", Icon50x40));
+	{
+		Style->Set("DMXEditor.LevelEditor.MenuIcon_none", new IMAGE_BRUSH("Icons/LevelEditor/icon_DMXTools_none_50x40", Icon50x40));
+		Style->Set("DMXEditor.LevelEditor.MenuIcon_snd", new IMAGE_BRUSH("Icons/LevelEditor/icon_DMXTools_snd_50x40", Icon50x40));
+		Style->Set("DMXEditor.LevelEditor.MenuIcon_rcv", new IMAGE_BRUSH("Icons/LevelEditor/icon_DMXTools_rcv_50x40", Icon50x40));
+		Style->Set("DMXEditor.LevelEditor.MenuIcon_snd-rcv", new IMAGE_BRUSH("Icons/LevelEditor/icon_DMXTools_snd-rcv_50x40", Icon50x40));
+	}
 
 	// Distribution Grid buttons
-	Style->Set("DMXEditor.PixelMapping.DistributionGrid.0.0", new IMAGE_BRUSH("Icons/DistributionGrid/icon_PixelDirection_0_0_161x160", Icon34x29));
-	Style->Set("DMXEditor.PixelMapping.DistributionGrid.0.1", new IMAGE_BRUSH("Icons/DistributionGrid/icon_PixelDirection_0_1_161x160", Icon34x29));
-	Style->Set("DMXEditor.PixelMapping.DistributionGrid.0.2", new IMAGE_BRUSH("Icons/DistributionGrid/icon_PixelDirection_0_2_161x160", Icon34x29));
-	Style->Set("DMXEditor.PixelMapping.DistributionGrid.0.3", new IMAGE_BRUSH("Icons/DistributionGrid/icon_PixelDirection_0_3_161x160", Icon34x29));
-
-	Style->Set("DMXEditor.PixelMapping.DistributionGrid.1.0", new IMAGE_BRUSH("Icons/DistributionGrid/icon_PixelDirection_1_0_161x160", Icon34x29));
-	Style->Set("DMXEditor.PixelMapping.DistributionGrid.1.1", new IMAGE_BRUSH("Icons/DistributionGrid/icon_PixelDirection_1_1_161x160", Icon34x29));
-	Style->Set("DMXEditor.PixelMapping.DistributionGrid.1.2", new IMAGE_BRUSH("Icons/DistributionGrid/icon_PixelDirection_1_2_161x160", Icon34x29));
-	Style->Set("DMXEditor.PixelMapping.DistributionGrid.1.3", new IMAGE_BRUSH("Icons/DistributionGrid/icon_PixelDirection_1_3_161x160", Icon34x29));
-
-	Style->Set("DMXEditor.PixelMapping.DistributionGrid.2.0", new IMAGE_BRUSH("Icons/DistributionGrid/icon_PixelDirection_2_0_161x160", Icon34x29));
-	Style->Set("DMXEditor.PixelMapping.DistributionGrid.2.1", new IMAGE_BRUSH("Icons/DistributionGrid/icon_PixelDirection_2_1_161x160", Icon34x29));
-	Style->Set("DMXEditor.PixelMapping.DistributionGrid.2.2", new IMAGE_BRUSH("Icons/DistributionGrid/icon_PixelDirection_2_2_161x160", Icon34x29));
-	Style->Set("DMXEditor.PixelMapping.DistributionGrid.2.3", new IMAGE_BRUSH("Icons/DistributionGrid/icon_PixelDirection_2_3_161x160", Icon34x29));
-
-	Style->Set("DMXEditor.PixelMapping.DistributionGrid.3.0", new IMAGE_BRUSH("Icons/DistributionGrid/icon_PixelDirection_3_0_161x160", Icon34x29));
-	Style->Set("DMXEditor.PixelMapping.DistributionGrid.3.1", new IMAGE_BRUSH("Icons/DistributionGrid/icon_PixelDirection_3_1_161x160", Icon34x29));
-	Style->Set("DMXEditor.PixelMapping.DistributionGrid.3.2", new IMAGE_BRUSH("Icons/DistributionGrid/icon_PixelDirection_3_2_161x160", Icon34x29));
-	Style->Set("DMXEditor.PixelMapping.DistributionGrid.3.3", new IMAGE_BRUSH("Icons/DistributionGrid/icon_PixelDirection_3_3_161x160", Icon34x29));
+	{
+		Style->Set("DMXEditor.PixelMapping.DistributionGrid.0.0", new IMAGE_BRUSH("Icons/DistributionGrid/icon_PixelDirection_0_0_161x160", Icon34x29));
+		Style->Set("DMXEditor.PixelMapping.DistributionGrid.0.1", new IMAGE_BRUSH("Icons/DistributionGrid/icon_PixelDirection_0_1_161x160", Icon34x29));
+		Style->Set("DMXEditor.PixelMapping.DistributionGrid.0.2", new IMAGE_BRUSH("Icons/DistributionGrid/icon_PixelDirection_0_2_161x160", Icon34x29));
+		Style->Set("DMXEditor.PixelMapping.DistributionGrid.0.3", new IMAGE_BRUSH("Icons/DistributionGrid/icon_PixelDirection_0_3_161x160", Icon34x29));
+
+		Style->Set("DMXEditor.PixelMapping.DistributionGrid.1.0", new IMAGE_BRUSH("Icons/DistributionGrid/icon_PixelDirection_1_0_161x160", Icon34x29));
+		Style->Set("DMXEditor.PixelMapping.DistributionGrid.1.1", new IMAGE_BRUSH("Icons/DistributionGrid/icon_PixelDirection_1_1_161x160", Icon34x29));
+		Style->Set("DMXEditor.PixelMapping.DistributionGrid.1.2", new IMAGE_BRUSH("Icons/DistributionGrid/icon_PixelDirection_1_2_161x160", Icon34x29));
+		Style->Set("DMXEditor.PixelMapping.DistributionGrid.1.3", new IMAGE_BRUSH("Icons/DistributionGrid/icon_PixelDirection_1_3_161x160", Icon34x29));
+
+		Style->Set("DMXEditor.PixelMapping.DistributionGrid.2.0", new IMAGE_BRUSH("Icons/DistributionGrid/icon_PixelDirection_2_0_161x160", Icon34x29));
+		Style->Set("DMXEditor.PixelMapping.DistributionGrid.2.1", new IMAGE_BRUSH("Icons/DistributionGrid/icon_PixelDirection_2_1_161x160", Icon34x29));
+		Style->Set("DMXEditor.PixelMapping.DistributionGrid.2.2", new IMAGE_BRUSH("Icons/DistributionGrid/icon_PixelDirection_2_2_161x160", Icon34x29));
+		Style->Set("DMXEditor.PixelMapping.DistributionGrid.2.3", new IMAGE_BRUSH("Icons/DistributionGrid/icon_PixelDirection_2_3_161x160", Icon34x29));
+
+		Style->Set("DMXEditor.PixelMapping.DistributionGrid.3.0", new IMAGE_BRUSH("Icons/DistributionGrid/icon_PixelDirection_3_0_161x160", Icon34x29));
+		Style->Set("DMXEditor.PixelMapping.DistributionGrid.3.1", new IMAGE_BRUSH("Icons/DistributionGrid/icon_PixelDirection_3_1_161x160", Icon34x29));
+		Style->Set("DMXEditor.PixelMapping.DistributionGrid.3.2", new IMAGE_BRUSH("Icons/DistributionGrid/icon_PixelDirection_3_2_161x160", Icon34x29));
+		Style->Set("DMXEditor.PixelMapping.DistributionGrid.3.3", new IMAGE_BRUSH("Icons/DistributionGrid/icon_PixelDirection_3_3_161x160", Icon34x29));
+	}
 
 	// OutputConsole icons
-	Style->Set("DMXEditor.OutputConsole.MacroSineWave", new IMAGE_BRUSH("Icons/icon_MacroSineWave_51x31", Icon51x31));
-	Style->Set("DMXEditor.OutputConsole.MacroMin", new IMAGE_BRUSH("Icons/icon_MacroMin_51x30", Icon51x30));
-	Style->Set("DMXEditor.OutputConsole.MacroMax", new IMAGE_BRUSH("Icons/icon_MacroMax_51x30", Icon51x30));
+	{
+		Style->Set("DMXEditor.OutputConsole.MacroSineWave", new IMAGE_BRUSH("Icons/icon_MacroSineWave_51x31", Icon51x31));
+		Style->Set("DMXEditor.OutputConsole.MacroMin", new IMAGE_BRUSH("Icons/icon_MacroMin_51x30", Icon51x30));
+		Style->Set("DMXEditor.OutputConsole.MacroMax", new IMAGE_BRUSH("Icons/icon_MacroMax_51x30", Icon51x30));
+	}
+
 	return Style;
 }
 #undef IMAGE_BRUSH
diff --git a/Engine/Plugins/VirtualProduction/DMX/DMXEngine/Source/DMXEditor/Private/Widgets/FixtureType/SDMXFixtureTypeEditor.cpp b/Engine/Plugins/VirtualProduction/DMX/DMXEngine/Source/DMXEditor/Private/Widgets/FixtureType/SDMXFixtureTypeEditor.cpp
--- a/Engine/Plugins/VirtualProduction/DMX/DMXEngine/Source/DMXEditor/Private/Widgets/FixtureType/SDMXFixtureTypeEditor.cpp
+++ b/Engine/Plugins/VirtualProduction/DMX/DMXEngine/Source/DMXEditor/Private/Widgets/FixtureType/SDMXFixtureTypeEditor.cpp
@@ -1,53 +1,34 @@
 // Copyright Epic Games, Inc. All Rights Reserved.
 
-#include "Widgets/FixtureType/SDMXFixtureTypeEditor.h"
+#include "SDMXFixtureTypeEditor.h"
 
 #include "DMXEditor.h"
 #include "DMXEditorUtils.h"
 #include "DMXFixtureTypeSharedData.h"
+#include "SDMXFixtureFunctionEditor.h"
+#include "SDMXFixtureModeEditor.h"
+#include "SDMXFixtureTypeFunctionsEditor.h"
+#include "SDMXFixtureTypeModesEditor.h"
 #include "SDMXFixtureTypeTree.h"
 #include "Library/DMXEntityFixtureType.h"
 #include "Library/DMXImport.h"
 #include "Library/DMXLibrary.h"
 #include "Library/DMXEntityFixtureType.h"
-#include "Widgets/SDMXEntityInspector.h"
 
 #include "Widgets/Layout/SSplitter.h"
 
 
-void SDMXFixtureTypeEditor::Construct(const FArguments& InArgs)
+void SDMXFixtureTypeEditor::Construct(const FArguments& InArgs, const TSharedRef<FDMXEditor>& InDMXEditor)
 {
 	SDMXEntityEditor::Construct(SDMXEntityEditor::FArguments());
 
-	DMXEditor = InArgs._DMXEditor;
+	WeakDMXEditor = InDMXEditor;
+	FixtureTypeSharedData = InDMXEditor->GetFixtureTypeSharedData();
 
 	SetCanTick(false);
 	bCanSupportFocus = false;
 
-	FixtureSettingsInspector =
-		SNew(SDMXEntityInspectorFixtureTypes, EDMXFixtureTypeLayout::FixtureSettings)		
-		.DMXEditor(DMXEditor)														
-		.OnFinishedChangingProperties(this, &SDMXFixtureTypeEditor::OnFinishedChangingProperties);
-
-	ModesInspector =
-		SNew(SDMXEntityInspectorFixtureTypes, EDMXFixtureTypeLayout::Modes)
-		.DMXEditor(DMXEditor)
-		.OnFinishedChangingProperties(this, &SDMXFixtureTypeEditor::OnFinishedChangingProperties);
-
-	ModePropertiesInspector =
-		SNew(SDMXEntityInspectorFixtureTypes, EDMXFixtureTypeLayout::ModeProperties)
-		.DMXEditor(DMXEditor)
-		.OnFinishedChangingProperties(this, &SDMXFixtureTypeEditor::OnFinishedChangingProperties);
-
-	FunctionsInspector =
-		SNew(SDMXEntityInspectorFixtureTypes, EDMXFixtureTypeLayout::Functions)
-		.DMXEditor(DMXEditor)
-		.OnFinishedChangingProperties(this, &SDMXFixtureTypeEditor::OnFinishedChangingProperties);
-
-	FunctionPropertiesInspector =
-		SNew(SDMXEntityInspectorFixtureTypes, EDMXFixtureTypeLayout::FunctionProperties)
-		.DMXEditor(DMXEditor)
-		.OnFinishedChangingProperties(this, &SDMXFixtureTypeEditor::OnFinishedChangingProperties);
+	FixtureTypeDetailsView = GenerateFixtureTypeDetailsView();
 
 	ChildSlot
 	.VAlign(VAlign_Fill)
@@ -67,15 +48,14 @@ void SDMXFixtureTypeEditor::Construct(const FArguments& InArgs)
 			// Top Row
 			+ SSplitter::Slot()
 			[
-				SAssignNew(FixtureTypeTree, SDMXFixtureTypeTree, UDMXEntityFixtureType::StaticClass())
-				.DMXEditor(InArgs._DMXEditor)
-				.OnSelectionChanged(this, &SDMXFixtureTypeEditor::OnSelectionChanged)
+				SAssignNew(FixtureTypeTree, SDMXFixtureTypeTree)
+				.DMXEditor(InDMXEditor)
 			]
 
 			// Bottom Row
 			+ SSplitter::Slot()
 			[
-				FixtureSettingsInspector.ToSharedRef()
+				FixtureTypeDetailsView.ToSharedRef()
 			]
 		]
 			
@@ -89,13 +69,13 @@ void SDMXFixtureTypeEditor::Construct(const FArguments& InArgs)
 			// Top Row
 			+ SSplitter::Slot()			
 			[
-				ModesInspector.ToSharedRef()
+				SNew(SDMXFixtureTypeModesEditor, InDMXEditor)
 			]
 
 			// Bottom Row
 			+ SSplitter::Slot()
 			[
-				ModePropertiesInspector.ToSharedRef()
+				SNew(SDMXFixtureModeEditor, InDMXEditor)
 			]
 		]
 
@@ -109,16 +89,27 @@ void SDMXFixtureTypeEditor::Construct(const FArguments& InArgs)
 			// Top Row
 			+ SSplitter::Slot()
 			[
-				FunctionsInspector.ToSharedRef()
+				SNew(SDMXFixtureTypeFunctionsEditor, InDMXEditor)
 			]
 
 			// Bottom Row
 			+ SSplitter::Slot()
 			[
-				FunctionPropertiesInspector.ToSharedRef()
+				SNew(SDMXFixtureFunctionEditor, InDMXEditor)
 			]
 		]
 	];	
+
+	// Set the initially selected fixture type
+	TArray<TWeakObjectPtr<UDMXEntityFixtureType>> SelectedFixtureTypes = FixtureTypeSharedData->GetSelectedFixtureTypes();
+	if (SelectedFixtureTypes.Num() == 1 && SelectedFixtureTypes[0].IsValid())
+	{
+		FixtureTypeDetailsView->SetObjects(TArray<UObject*>({ SelectedFixtureTypes[0].Get() }));
+	}
+
+	// Initially set keyboard focus to the Fixture Type Tree
+	FSlateApplication::Get().SetKeyboardFocus(FixtureTypeTree.ToSharedRef());
+
 }
 
 void SDMXFixtureTypeEditor::RequestRenameOnNewEntity(const UDMXEntity* InEntity, ESelectInfo::Type SelectionType)
@@ -150,80 +141,28 @@ TArray<UDMXEntity*> SDMXFixtureTypeEditor::GetSelectedEntities() const
 	return FixtureTypeTree->GetSelectedEntities();
 }
 
-void SDMXFixtureTypeEditor::OnSelectionChanged(const TArray<UDMXEntity*>& InSelectedEntities)
+void SDMXFixtureTypeEditor::OnFixtureTypesSelected()
 {
-	if (TSharedPtr<FDMXEditor> DMXEditorPtr = DMXEditor.Pin())
+	TArray<TWeakObjectPtr<UDMXEntityFixtureType>> SelectedFixtureTypes = FixtureTypeSharedData->GetSelectedFixtureTypes();
+	TArray<UObject*> SelectedObjects;
+	for (TWeakObjectPtr<UDMXEntityFixtureType> WeakSelectedFixtureType : SelectedFixtureTypes)
 	{
-		// Set selected fixture types in shared data
-		if (TSharedPtr<FDMXFixtureTypeSharedData> SharedData = DMXEditorPtr->GetFixtureTypeSharedData())
+		if (UDMXEntity* SelectedObject = WeakSelectedFixtureType.Get())
 		{
-			TArray<TWeakObjectPtr<UDMXEntityFixtureType>> SelectedFixtureTypes;
-			for (UDMXEntity* Entity : InSelectedEntities)
-			{
-				if (UDMXEntityFixtureType* FixtureType = Cast<UDMXEntityFixtureType>(Entity))
-				{
-					SelectedFixtureTypes.Add(FixtureType);
-				}
-			}
-			SharedData->SetFixtureTypesBeingEdited(SelectedFixtureTypes);
+			SelectedObjects.Add(SelectedObject);
 		}
-
-		check(FixtureSettingsInspector.IsValid());
-		check(ModesInspector.IsValid());
-		check(ModePropertiesInspector.IsValid());
-		check(FunctionsInspector.IsValid());
-		check(FunctionPropertiesInspector.IsValid());
-
-		FixtureSettingsInspector->ShowDetailsForEntities(InSelectedEntities);
-		ModesInspector->ShowDetailsForEntities(InSelectedEntities);
-		ModePropertiesInspector->ShowDetailsForEntities(InSelectedEntities);
-		FunctionsInspector->ShowDetailsForEntities(InSelectedEntities);
-		FunctionPropertiesInspector->ShowDetailsForEntities(InSelectedEntities);
 	}
+	FixtureTypeDetailsView->SetObjects(SelectedObjects);
 }
 
-void SDMXFixtureTypeEditor::OnFinishedChangingProperties(const FPropertyChangedEvent& PropertyChangedEvent)
+TSharedRef<IDetailsView> SDMXFixtureTypeEditor::GenerateFixtureTypeDetailsView() const
 {
-	if (PropertyChangedEvent.ChangeType == EPropertyChangeType::ArrayAdd)
-	{
-		if (PropertyChangedEvent.GetPropertyName() == GET_MEMBER_NAME_CHECKED(UDMXEntityFixtureType, Modes)
-			|| PropertyChangedEvent.GetPropertyName() == GET_MEMBER_NAME_CHECKED(FDMXFixtureMode, Functions))
-		{
-			// When the user adds a Mode or Function, their names can't be empty
-
-			const TArray<UDMXEntity*>&& SelectedEntities = FixtureTypeTree->GetSelectedEntities();
+	FPropertyEditorModule& PropertyEditorModule = FModuleManager::GetModuleChecked<FPropertyEditorModule>("PropertyEditor");
 
-			for (UDMXEntity* Entity : SelectedEntities)
-			{
-				if (UDMXEntityFixtureType* FixtureType = Cast<UDMXEntityFixtureType>(Entity))
-				{
-					FDMXEditorUtils::SetNewFixtureFunctionsNames(FixtureType);
-				}
-			}
-		}
-	}
-	else if (PropertyChangedEvent.ChangeType == EPropertyChangeType::ValueSet)
-	{
-		if (PropertyChangedEvent.GetPropertyName() == GET_MEMBER_NAME_CHECKED(UDMXEntityFixtureType, DMXImport))
-		{
-			const TArray<UDMXEntity*>&& SelectedEntities = FixtureTypeTree->GetSelectedEntities();
-
-			for (UDMXEntity* Entity : SelectedEntities)
-			{
-				if (UDMXEntityFixtureType* FixtureType = Cast<UDMXEntityFixtureType>(Entity))
-				{
-					if (FixtureType->DMXImport != nullptr && FixtureType->DMXImport->IsValidLowLevelFast())
-					{
-						FixtureType->Modify();
-						FixtureType->SetModesFromDMXImport(FixtureType->DMXImport);
-					}
-				}
-			}
-		}
-	}
+	FDetailsViewArgs DetailsViewArgs;
+	DetailsViewArgs.bAllowSearch = true;
+	DetailsViewArgs.NameAreaSettings = FDetailsViewArgs::HideNameArea;
+	DetailsViewArgs.bHideSelectionTip = true;
 
-	if (FixtureTypeTree.IsValid())
-	{
-		FixtureTypeTree->UpdateTree();
-	}
+	return PropertyEditorModule.CreateDetailView(DetailsViewArgs);
 }
diff --git a/Engine/Plugins/VirtualProduction/DMX/DMXEngine/Source/DMXEditor/Private/Widgets/FixtureType/SDMXFixtureTypeEditor.h b/Engine/Plugins/VirtualProduction/DMX/DMXEngine/Source/DMXEditor/Private/Widgets/FixtureType/SDMXFixtureTypeEditor.h
--- a/Engine/Plugins/VirtualProduction/DMX/DMXEngine/Source/DMXEditor/Private/Widgets/FixtureType/SDMXFixtureTypeEditor.h
+++ b/Engine/Plugins/VirtualProduction/DMX/DMXEngine/Source/DMXEditor/Private/Widgets/FixtureType/SDMXFixtureTypeEditor.h
@@ -7,28 +7,26 @@
 #include "CoreMinimal.h"
 
 class FDMXEditor;
-class SDMXEntityInspector;
+class FDMXFixtureTypeSharedData;
+class SDMXFixtureModeEditor;
 class SDMXFixtureTypeTree;
 
 struct FPropertyChangedEvent;
+class IDetailsView;
 
-
-/** Editor for Fixture Types */
+/** The whole editor for Fixture Types */
 class SDMXFixtureTypeEditor
 	: public SDMXEntityEditor
 {
 public:
 	SLATE_BEGIN_ARGS(SDMXFixtureTypeEditor)
-		: _DMXEditor(nullptr)
 	{}
 
-		SLATE_ARGUMENT(TWeakPtr<FDMXEditor>, DMXEditor)
-
 	SLATE_END_ARGS()
 
 public:
 	/** Constructs this widget */
-	void Construct(const FArguments& InArgs);
+	void Construct(const FArguments& InArgs, const TSharedRef<FDMXEditor>& InDMXEditor);
 
 	/** Begin SDMXEntityEditorTab interface */
 	void RequestRenameOnNewEntity(const UDMXEntity* InEntity, ESelectInfo::Type SelectionType);
@@ -37,32 +35,22 @@ public:
 	TArray<UDMXEntity*> GetSelectedEntities() const;
 	/** ~End SDMXEntityEditorTab interface */
 
-protected:
-	/** Callback for when entities list selection changes */
-	virtual void OnSelectionChanged(const TArray<UDMXEntity*>& InSelectedEntities);
+private:
+	/** Called when fixture types were selected */
+	void OnFixtureTypesSelected();
 
-	/** Callback for when some property has changed in the inspector */
-	virtual void OnFinishedChangingProperties(const FPropertyChangedEvent& PropertyChangedEvent);
-
-protected:
-	/** Pointer back to the DMXEditor tool that owns us */
-	TWeakPtr<FDMXEditor> DMXEditor;
+	/** Generates a details view for the Fixture Type Details */
+	TSharedRef<IDetailsView> GenerateFixtureTypeDetailsView() const;
 
 	/** Tree of Fixture Types */
 	TSharedPtr<SDMXFixtureTypeTree> FixtureTypeTree;
 
-	/** Inspector for the selected Fixture(s) Settings */
-	TSharedPtr<SDMXEntityInspector> FixtureSettingsInspector;
+	/** Details View for the selected Fixture Type */
+	TSharedPtr<IDetailsView> FixtureTypeDetailsView;
 
-	/** Inspector for Modes */
-	TSharedPtr<SDMXEntityInspector> ModesInspector;
+	/** Shared Data for Fixture Types */
+	TSharedPtr<FDMXFixtureTypeSharedData> FixtureTypeSharedData;
 
-	/** Inspector for the selected mode's properties */
-	TSharedPtr<SDMXEntityInspector> ModePropertiesInspector;
-
-	/** Inspector for the selected Fixture(s) functions */
-	TSharedPtr<SDMXEntityInspector> FunctionsInspector;
-
-	/** Inspector for the selected function's properties */
-	TSharedPtr<SDMXEntityInspector> FunctionPropertiesInspector;
+	/** Pointer back to the DMXEditor tool that owns us */
+	TWeakPtr<FDMXEditor> WeakDMXEditor;
 }; 
diff --git a/Engine/Plugins/VirtualProduction/DMX/DMXEngine/Source/DMXEditor/Private/Widgets/FixtureType/SDMXFixtureTypeTree.cpp b/Engine/Plugins/VirtualProduction/DMX/DMXEngine/Source/DMXEditor/Private/Widgets/FixtureType/SDMXFixtureTypeTree.cpp
--- a/Engine/Plugins/VirtualProduction/DMX/DMXEngine/Source/DMXEditor/Private/Widgets/FixtureType/SDMXFixtureTypeTree.cpp
+++ b/Engine/Plugins/VirtualProduction/DMX/DMXEngine/Source/DMXEditor/Private/Widgets/FixtureType/SDMXFixtureTypeTree.cpp
@@ -4,6 +4,7 @@
 
 #include "DMXEditor.h"
 #include "DMXEditorUtils.h"
+#include "DMXFixtureTypeSharedData.h"
 #include "SDMXFixtureTypeTreeCategoryRow.h"
 #include "SDMXFixtureTypeTreeFixtureTypeRow.h"
 #include "Commands/DMXEditorCommands.h"
@@ -24,12 +25,12 @@
 
 #define LOCTEXT_NAMESPACE "SDMXFixtureTypeTree"
 
-void SDMXFixtureTypeTree::Construct(const FArguments& InArgs, const TSubclassOf<UDMXEntity> InListType)
+void SDMXFixtureTypeTree::Construct(const FArguments& InArgs)
 {
 	SDMXEntityTreeViewBase::FArguments BaseArguments =
 		SDMXEntityTreeViewBase::FArguments()
 		.DMXEditor(InArgs._DMXEditor)
-		.OnSelectionChanged(InArgs._OnSelectionChanged)
+		.OnSelectionChanged(this, &SDMXFixtureTypeTree::OnEntitySelected)
 		.OnEntitiesAdded(InArgs._OnEntitiesAdded)
 		.OnEntityOrderChanged(InArgs._OnEntityOrderChanged)
 		.OnEntitiesRemoved(InArgs._OnEntitiesRemoved);
@@ -359,13 +360,55 @@ TSharedPtr<SDMXFixtureTypeTreeFixtureTypeRow> SDMXFixtureTypeTree::FindEntityRow
 	return nullptr;
 }
 
+void SDMXFixtureTypeTree::OnEntitySelected(const TArray<UDMXEntity*>& NewSelection)
+{
+	if (TSharedPtr<FDMXEditor> PinnedEditor = DMXEditor.Pin())
+	{
+		if (TSharedPtr<FDMXFixtureTypeSharedData> SharedData = PinnedEditor->GetFixtureTypeSharedData())
+		{
+			// Never clear the selection
+			if (GetSelectedEntities().Num() == 0)
+			{
+				TArray<TWeakObjectPtr<UDMXEntityFixtureType>> OldSelection = SharedData->GetSelectedFixtureTypes();
+				TArray<UDMXEntity*> OldSelectionAsEntities;
+				for (TWeakObjectPtr<UDMXEntityFixtureType> SelectedFixtureType : OldSelection)
+				{
+					if (UDMXEntityFixtureType* SelectedEntity = SelectedFixtureType.Get())
+					{
+						OldSelectionAsEntities.Add(SelectedEntity);
+					}
+				}
+
+				SelectItemsByEntities(OldSelectionAsEntities);
+			}
+			else
+			{
+				// Select selected Fixture Types in Fixture Type Shared Data
+				TArray<TWeakObjectPtr<UDMXEntityFixtureType>> SelectedFixtureTypes;
+
+				for (UDMXEntity* Entity : NewSelection)
+				{
+					if (UDMXEntityFixtureType* FixtureType = Cast<UDMXEntityFixtureType>(Entity))
+					{
+						SelectedFixtureTypes.Add(FixtureType);
+					}
+				}
+
+				SharedData->SelectFixtureTypes(SelectedFixtureTypes);
+			}
+		}
+	}
+}
+
 FReply SDMXFixtureTypeTree::OnAddNewFixtureTypeClicked()
 {
-	check(DMXEditor.IsValid());
-	TSharedPtr<FDMXEditor> PinnedEditor = DMXEditor.Pin();
+	if (TSharedPtr<FDMXEditor> PinnedEditor = DMXEditor.Pin())
+	{		
+		PinnedEditor->GetToolkitCommands()->ExecuteAction(FDMXEditorCommands::Get().AddNewEntityFixtureType.ToSharedRef());
+		return FReply::Handled();
+	}
 
-	PinnedEditor->GetToolkitCommands()->ExecuteAction(FDMXEditorCommands::Get().AddNewEntityFixtureType.ToSharedRef());
-	return FReply::Handled();
+	return FReply::Unhandled();
 }
 
 #undef LOCTEXT_NAMESPACE
diff --git a/Engine/Plugins/VirtualProduction/DMX/DMXEngine/Source/DMXEditor/Private/Widgets/FixtureType/SDMXFixtureTypeTree.h b/Engine/Plugins/VirtualProduction/DMX/DMXEngine/Source/DMXEditor/Private/Widgets/FixtureType/SDMXFixtureTypeTree.h
--- a/Engine/Plugins/VirtualProduction/DMX/DMXEngine/Source/DMXEditor/Private/Widgets/FixtureType/SDMXFixtureTypeTree.h
+++ b/Engine/Plugins/VirtualProduction/DMX/DMXEngine/Source/DMXEditor/Private/Widgets/FixtureType/SDMXFixtureTypeTree.h
@@ -20,9 +20,6 @@ public:
 		/** The DMX Editor that owns this widget */
 		SLATE_ARGUMENT(TWeakPtr<FDMXEditor>, DMXEditor)
 
-		/** Exectued when the list changed its selection */
-		SLATE_EVENT(FDMXOnSelectionChanged, OnSelectionChanged)
-
 		/** Exectued when entites were added to the DMXEditor's library */
 		SLATE_EVENT(FSimpleDelegate, OnEntitiesAdded)
 
@@ -35,7 +32,7 @@ public:
 	SLATE_END_ARGS()
 
 	/** Constructs the widget */
-	void Construct(const FArguments& InArgs, const TSubclassOf<UDMXEntity> InListType);
+	void Construct(const FArguments& InArgs);
 
 	/** Destructor */
 	virtual ~SDMXFixtureTypeTree() {}
@@ -64,6 +61,9 @@ protected:
 	//~ End SDMXEntityTreeViewBase interface
 
 private:
+	/** Called when the selection changed */
+	void OnEntitySelected(const TArray<UDMXEntity*>& NewSelection);
+
 	/** Called when the Add button was clicked */
 	FReply OnAddNewFixtureTypeClicked();
 
