DMX - Rework DMX Fixture Type editor for better GDTF support, better UX and to ease maintaining and extending the editor

diff --git a/Engine/Plugins/VirtualProduction/DMX/DMXEngine/Source/DMXEditor/Private/Widgets/SDMXEntityInspector.cpp b/Engine/Plugins/VirtualProduction/DMX/DMXEngine/Source/DMXEditor/Private/Widgets/SDMXEntityInspector.cpp
--- a/Engine/Plugins/VirtualProduction/DMX/DMXEngine/Source/DMXEditor/Private/Widgets/SDMXEntityInspector.cpp
+++ b/Engine/Plugins/VirtualProduction/DMX/DMXEngine/Source/DMXEditor/Private/Widgets/SDMXEntityInspector.cpp
@@ -6,12 +6,6 @@
 #include "DMXEditorLog.h"
 #include "DMXProtocolConstants.h"
 
-#include "Customizations/DMXEntityFixtureTypeFixtureSettingsDetails.h"
-#include "Customizations/DMXEntityFixtureTypeModesDetails.h"
-#include "Customizations/DMXEntityFixtureTypeModePropertiesDetails.h"
-#include "Customizations/DMXEntityFixtureTypeFunctionsDetails.h"
-#include "Customizations/DMXEntityFixtureTypeFunctionPropertiesDetails.h"
-
 #include "Library/DMXEntityFixtureType.h"
 #include "Library/DMXEntityFixturePatch.h"
 
@@ -127,37 +121,4 @@ TSharedRef<SWidget> SDMXEntityInspector::MakeEditingWidget(const TArray<UObject*
 	return InnerEditingWidget;
 }
 
-void SDMXEntityInspectorFixtureTypes::Construct(const FArguments& InArgs, EDMXFixtureTypeLayout Layout)
-{
-	SDMXEntityInspector::Construct(SDMXEntityInspector::FArguments()
-		.DMXEditor(InArgs._DMXEditor)
-		.OnFinishedChangingProperties(InArgs._OnFinishedChangingProperties)
-	);
-
-	// Register details customizations Fixture Type UOBJECT
-	FOnGetDetailCustomizationInstance FixtureTypeDetails;
-	switch (Layout)
-	{
-	case EDMXFixtureTypeLayout::FixtureSettings:
-		FixtureTypeDetails = FOnGetDetailCustomizationInstance::CreateStatic(&FDMXEntityFixtureTypeFixtureSettingsDetails::MakeInstance, InArgs._DMXEditor);
-		break;
-	case EDMXFixtureTypeLayout::Modes:
-		FixtureTypeDetails = FOnGetDetailCustomizationInstance::CreateStatic(&FDMXEntityFixtureTypeModesDetails::MakeInstance, InArgs._DMXEditor);
-		break;
-	case EDMXFixtureTypeLayout::ModeProperties:
-		FixtureTypeDetails = FOnGetDetailCustomizationInstance::CreateStatic(&FDMXEntityFixtureTypeModePropertiesDetails::MakeInstance, InArgs._DMXEditor);
-		break;
-	case EDMXFixtureTypeLayout::Functions:
-		FixtureTypeDetails = FOnGetDetailCustomizationInstance::CreateStatic(&FDMXEntityFixtureTypeFunctionsDetails::MakeInstance, InArgs._DMXEditor);
-		break;
-	case EDMXFixtureTypeLayout::FunctionProperties:
-		FixtureTypeDetails = FOnGetDetailCustomizationInstance::CreateStatic(&FDMXEntityFixtureTypeFunctionPropertiesDetails::MakeInstance, InArgs._DMXEditor);
-		break;
-	default:
-		check(0); // Unhandled layout
-		return;
-	}
-	GetPropertyView()->RegisterInstancedCustomPropertyLayout(UDMXEntityFixtureType::StaticClass(), FixtureTypeDetails);
-}
-
 #undef LOCTEXT_NAMESPACE
diff --git a/Engine/Plugins/VirtualProduction/DMX/DMXEngine/Source/DMXEditor/Private/Widgets/SDMXEntityInspector.h b/Engine/Plugins/VirtualProduction/DMX/DMXEngine/Source/DMXEditor/Private/Widgets/SDMXEntityInspector.h
--- a/Engine/Plugins/VirtualProduction/DMX/DMXEngine/Source/DMXEditor/Private/Widgets/SDMXEntityInspector.h
+++ b/Engine/Plugins/VirtualProduction/DMX/DMXEngine/Source/DMXEditor/Private/Widgets/SDMXEntityInspector.h
@@ -106,23 +106,6 @@ enum class EDMXFixtureTypeLayout : uint8
 	FunctionProperties
 };
 
-/** DMX fixture types inspector */
-class SDMXEntityInspectorFixtureTypes
-	: public SDMXEntityInspector
-{
-public:
-	SLATE_BEGIN_ARGS(SDMXEntityInspectorFixtureTypes)
-	{}
-
-	SLATE_ARGUMENT(TWeakPtr<FDMXEditor>, DMXEditor)
-	SLATE_EVENT(FOnFinishedChangingProperties::FDelegate, OnFinishedChangingProperties)
-
-	SLATE_END_ARGS()
-
-	/** Constructs the widget */
-	void Construct(const FArguments& InArgs, EDMXFixtureTypeLayout DetailsView);
-};
-
 /** DMX Component inspector */
 class SDMXComponentInspector
 	: public SDMXEntityInspector
diff --git a/Engine/Plugins/VirtualProduction/DMX/DMXEngine/Source/DMXEditor/Private/Widgets/SDMXEntityTreeViewBase.cpp b/Engine/Plugins/VirtualProduction/DMX/DMXEngine/Source/DMXEditor/Private/Widgets/SDMXEntityTreeViewBase.cpp
--- a/Engine/Plugins/VirtualProduction/DMX/DMXEngine/Source/DMXEditor/Private/Widgets/SDMXEntityTreeViewBase.cpp
+++ b/Engine/Plugins/VirtualProduction/DMX/DMXEngine/Source/DMXEditor/Private/Widgets/SDMXEntityTreeViewBase.cpp
@@ -96,7 +96,6 @@ void SDMXEntityTreeViewBase::Construct(const FArguments& InArgs)
 		.OnExpansionChanged(this, &SDMXEntityTreeViewBase::OnExpansionChanged)
 		.OnSelectionChanged(this, &SDMXEntityTreeViewBase::OnSelectionChanged)
 		.OnContextMenuOpening(this, &SDMXEntityTreeViewBase::OnContextMenuOpen)
-		// .OnItemScrolledIntoView(this, &SDMXEntityTreeViewBase::OnItemScrolledIntoView)
 		.HighlightParentNodesForSelection(false);
 
 	ChildSlot
diff --git a/Engine/Plugins/VirtualProduction/DMX/DMXEngine/Source/DMXEditor/Public/DMXEditor.h b/Engine/Plugins/VirtualProduction/DMX/DMXEngine/Source/DMXEditor/Public/DMXEditor.h
--- a/Engine/Plugins/VirtualProduction/DMX/DMXEngine/Source/DMXEditor/Public/DMXEditor.h
+++ b/Engine/Plugins/VirtualProduction/DMX/DMXEngine/Source/DMXEditor/Public/DMXEditor.h
@@ -25,6 +25,9 @@ DECLARE_MULTICAST_DELEGATE_TwoParams(FOnGetBaseNameForNewEntity, TSubclassOf<UDM
 // Used to enable Entity creator code to set values in a newly created entity
 DECLARE_MULTICAST_DELEGATE_OneParam(FOnSetupNewEntity, UDMXEntity*);
 
+/** 
+ * *** The DMX Library Editor, displayed when opening a DMX Library Asset ***
+ */
 class DMXEDITOR_API FDMXEditor
 	: public FWorkflowCentricApplication	// Allow Add ApplicationModes
 	, public FNotifyHook
diff --git a/Engine/Plugins/VirtualProduction/DMX/DMXEngine/Source/DMXEditor/Public/DMXEditorModule.h b/Engine/Plugins/VirtualProduction/DMX/DMXEngine/Source/DMXEditor/Public/DMXEditorModule.h
--- a/Engine/Plugins/VirtualProduction/DMX/DMXEngine/Source/DMXEditor/Public/DMXEditorModule.h
+++ b/Engine/Plugins/VirtualProduction/DMX/DMXEngine/Source/DMXEditor/Public/DMXEditorModule.h
@@ -87,6 +87,9 @@ protected:
 	/** Registers asset types actions */
 	void RegisterAssetTypeActions();
 
+	/** Registers global class customizations */
+	void RegisterClassCustomizations();
+
 	/** Registers global property type customizations */
 	void RegisterPropertyTypeCustomizations();
 
diff --git a/Engine/Plugins/VirtualProduction/DMX/DMXEngine/Source/DMXEditor/Public/DMXEditorSettings.h b/Engine/Plugins/VirtualProduction/DMX/DMXEngine/Source/DMXEditor/Public/DMXEditorSettings.h
--- a/Engine/Plugins/VirtualProduction/DMX/DMXEngine/Source/DMXEditor/Public/DMXEditorSettings.h
+++ b/Engine/Plugins/VirtualProduction/DMX/DMXEngine/Source/DMXEditor/Public/DMXEditorSettings.h
@@ -8,6 +8,22 @@
 
 #include "DMXEditorSettings.generated.h"
 
+
+/** Settings for the Fixture Type Functions Editor */
+USTRUCT()
+struct FDMXFixtureTypeFunctionsEditorSettings
+{
+	GENERATED_BODY()
+
+	/** Width of the function name column */
+	UPROPERTY()
+	float NameColumnWidth = 0.f;
+
+	/** With of the attribute column */
+	UPROPERTY()
+	float AttributeColumnWidth = 0.f;
+};
+
 /** Struct to describe a single fader, so it can be stored in the config */
 USTRUCT()
 struct FDMXOutputConsoleFaderDescriptor
@@ -82,11 +98,20 @@ class UDMXEditorSettings : public UObject
 {
 	GENERATED_BODY()
 
+	// Fixture Type Functions Editor
+public:
+	UPROPERTY(Config)
+	FDMXFixtureTypeFunctionsEditorSettings FixtureTypeFunctionsEditorSettings;
+
+
+	// Output Console
 public:
 	/** Stores the faders specified in Output Console */
 	UPROPERTY(Config)
 	TArray<FDMXOutputConsoleFaderDescriptor> OutputConsoleFaders;
 	
+
+	// Channels Monitor
 public:
 	/** The Universe ID to be monitored in the Channels Monitor  */
 	UPROPERTY(Config)
@@ -96,6 +121,8 @@ public:
 	UPROPERTY(Config)
 	FDMXMonitorSourceDescriptor ChannelsMonitorSource;
 
+
+	// Activity Monitor
 public:
 	/** Source for the DMX Activity Monitor */
 	UPROPERTY(Config)
diff --git a/Engine/Plugins/VirtualProduction/DMX/DMXEngine/Source/DMXEditor/Public/DMXEditorUtils.h b/Engine/Plugins/VirtualProduction/DMX/DMXEngine/Source/DMXEditor/Public/DMXEditorUtils.h
--- a/Engine/Plugins/VirtualProduction/DMX/DMXEngine/Source/DMXEditor/Public/DMXEditorUtils.h
+++ b/Engine/Plugins/VirtualProduction/DMX/DMXEngine/Source/DMXEditor/Public/DMXEditorUtils.h
@@ -11,16 +11,20 @@ class UDMXEntityFixtureType;
 class UDMXEntityFader;
 class UDMXEntityFixturePatch;
 
-
+/** 
+ * Generic Editor Utilities.
+ * for Fixture Type, refer to DMXFixtureTypeSharedData instead.
+ */
 class DMXEDITOR_API FDMXEditorUtils
 {
 public:
 	typedef TArray<UDMXEntityFixturePatch*> FUnassignedPatchesArray;
-		
+
 	/**
 	 * Generates a unique name given a base one and a list of existing ones, by appending an index to
 	 * existing names. If InBaseName is an empty String, it returns "Default name".
 	 */
+	UE_DEPRECATED(5.0, "Moved to DMXRuntimeUtils to allow using it at runtime.")
 	static FString GenerateUniqueNameFromExisting(const TSet<FString>& InExistingNames, const FString& InBaseName);
 
 	/**
@@ -90,8 +94,7 @@ public:
 	/**  Returns the Entity class type name (e.g: Fixture Type for UDMXEntityFixtureType) in singular or plural */
 	static FText GetEntityTypeNameText(TSubclassOf<UDMXEntity> EntityClass, bool bPlural = false);
 
-	
-	
+
 // Auto Assign:
 	
 	/**
diff --git a/Engine/Plugins/VirtualProduction/DMX/DMXEngine/Source/DMXEditor/Public/DMXFixtureTypeSharedData.h b/Engine/Plugins/VirtualProduction/DMX/DMXEngine/Source/DMXEditor/Public/DMXFixtureTypeSharedData.h
--- a/Engine/Plugins/VirtualProduction/DMX/DMXEngine/Source/DMXEditor/Public/DMXFixtureTypeSharedData.h
+++ b/Engine/Plugins/VirtualProduction/DMX/DMXEngine/Source/DMXEditor/Public/DMXFixtureTypeSharedData.h
@@ -10,262 +10,133 @@
 
 #include "JsonObjectConverter.h"
 
+enum class EDMXFixtureSignalFormat : uint8;
 class FDMXEditor;
 class FDMXFixtureTypeSharedData;
 class UDMXEntityFixtureType;
-class IPropertyHandleArray;
-class IPropertyHandle;
+
 class FScopedTransaction;
 
 
-/**
- * Item to identify and access a DMXFixtureMode.
- *
- * Note, this will no longer point to the live property when:
- * - The fixture types being edited change
- * - The details view is destroyed (be it via refresh or otherwise)
- * - The arrays that contain the property change num elements.
- */
-class FDMXFixtureModeItem
-	: public IDMXNamedType
+/** Shared data and utilities for objects in a DMX Editor */
+class FDMXFixtureTypeSharedData
+	: public TSharedFromThis<FDMXFixtureTypeSharedData>
 {
 public:
-	/** Constructor */
-	FDMXFixtureModeItem(const TWeakPtr<FDMXFixtureTypeSharedData>& InSharedDataPtr, const TSharedPtr<IPropertyHandle>& InModeNameHandle);
-	
-	/** Destructor */
-	virtual ~FDMXFixtureModeItem();
+	FDMXFixtureTypeSharedData(TWeakPtr<FDMXEditor> InDMXEditorPtr)
+		: DMXEditorPtr(InDMXEditorPtr)
+	{}	
 
 public:
-	/** Returns true if the mode names equal */
-	bool EqualsMode(const TSharedPtr<FDMXFixtureModeItem>& Other) const;
-
-	/** Gets the name of the mode */
-	FString GetModeName() const;
-
-	/** Returns the current mode name */
-	void SetModeName(const FString& InDesiredName, FString& OutNewName);
+	/** Selects specified Fixture Types */
+	void SelectFixtureTypes(const TArray<TWeakObjectPtr<UDMXEntityFixtureType>>& InFixtureTypes);
 
-	/** True if the mode is being edited by shared data */
-	bool IsModeBeingEdited() const;
-
-	/** True if the mode is selected by shared data */
-	bool IsModeSelected() const;
-
-	/** Selects the mode in shared data */
-	void SelectMode(bool bSelected);
-
-	/** Returns the outer fixture types of the mode name handle */
-	void GetOuterFixtureTypes(TArray<TWeakObjectPtr<UDMXEntityFixtureType>>& OuterFixtureTypes) const;
-
-protected:
-	// IDMXNamedType Interface
-	virtual void GetName(FString& OutUniqueName) const;
-	virtual bool IsNameUnique(const FString& TestedName) const;
-	virtual void SetName(const FString& InDesiredName, FString& OutUniqueName);
-	// ~IDMXNamedType Interface
-
-	/** Shared data that owns this */
-	TWeakPtr<FDMXFixtureTypeSharedData> SharedDataPtr;
-
-private:
-	/** The mode name handle */
-	TSharedPtr<IPropertyHandle> ModeNameHandle;
-};
+	/** Selects specified Modes */
+	void SelectModes(const TArray<int32>& InModeIndices);
 
+	/** Selects specified Functions */
+	void SetFunctionAndMatrixSelection(const TArray<int32>& InFunctionIndices, bool bMatrixSelected);
 
-/**
- * Item to identify and access a DMXFixtureFunction.
- *
- * Note, this will no longer point to the live property when:
- * - The fixture types being edited change
- * - The details view is destroyed (be it via refresh or otherwise)
- * - The arrays that contain the property change num elements.
- */
-class FDMXFixtureFunctionItem
-	: public FDMXFixtureModeItem
-{
-public:
-	FDMXFixtureFunctionItem(TWeakPtr<FDMXFixtureTypeSharedData> InOwnerPtr, const TSharedPtr<IPropertyHandle> InModeNameHandle, const TSharedPtr<IPropertyHandle> InFunctionHandle);
+	/** Selects specified Functions */
+	UE_DEPRECATED(5.0, "Deprecated in favor of SetFunctionAndMatrixSelection to avoid the unclear state where both need change, but one contains the old state while the other changed.")
+	void SelectFunctions(const TArray<int32>& InFunctionIndices);
 
-	virtual ~FDMXFixtureFunctionItem();
+	const TArray<TWeakObjectPtr<UDMXEntityFixtureType>>& GetSelectedFixtureTypes() const { return SelectedFixtureTypes; }
+	const TArray<int32>& GetSelectedModeIndices() const { return SelectedModeIndices; }
+	const TArray<int32>& GetSelectedFunctionIndices() const { return SelectedFunctionIndices; }
+	bool IsFixtureMatrixSelected() const { return bFixtureMatrixSelected; }
 
-	/** Returns true if the function names equal */
-	bool EqualsFunction(const TSharedPtr<FDMXFixtureFunctionItem>& Other) const;
+	/** Returns true if selected Modes can be copied */
+	bool CanCopyModesToClipboard() const;
 
-	/** Returns the current function name */
-	FString GetFunctionName() const;
+	/** Copies selected Modes. Only should be called when CanCopyModesToClipboard */
+	void CopyModesToClipboard();
 
-	/** Returns the current function channel */
-	int32 GetFunctionChannel() const;
+	/** Returns true if Modes can be pasted */
+	bool CanPasteModesFromClipboard() const;
 
-	/** Sets the function name */
-	void SetFunctionName(const FString& InDesiredName, FString& OutNewName);
+	/** Pastes the Modes clipboard to selected Modes, returns the newly added Mode indices. Only should be called when CanPasteModesFromClipboard */
+	void PasteModesFromClipboard(TArray<int32>& OutNewlyAddedModeIndices);
 
-	/** True if the function is edited in shared data */
-	bool IsFunctionEdited() const;
+	/** Returns true if selected Functions can be copied */
+	bool CanCopyFunctionsToClipboard() const;
 
-	/** True if the function is selected in shared data */
-	bool IsFunctionSelected() const;
+	/** Copies selected Functions. Only should be called when CanCopyFunctionsToClipboard */
+	void CopyFunctionsToClipboard();
 
-	/** Get property handle for the function */
-	const TSharedPtr<IPropertyHandle>& GetFunctionHandle() { return FunctionHandle; }
+	/** Returns true if Functions can be pasted */
+	bool CanPasteFunctionsFromClipboard() const;
 
-	/** Get property handle for the function name */
-	const TSharedPtr<IPropertyHandle>& GetFunctionNameHandle() { return FunctionNameHandle; }
-protected:
-	// IDMXNamedType Interface
-	virtual void GetName(FString& OutUniqueName) const override;
-	virtual bool IsNameUnique(const FString& TestedName) const override;
-	virtual void SetName(const FString& InDesiredName, FString& OutUniqueName) override;
-	// ~IDMXNamedType Interface
+	/** Pastes the Modes clipboard to selected Functions, returns the newly added Function indices. Only should be called when CanPasteFunctionsFromClipboard */
+	void PasteFunctionsFromClipboard(TArray<int32>& OutNewlyAddedFunctionIndices);
 
-private:
-	/** Property handle for the function name */
-	TSharedPtr<IPropertyHandle> FunctionNameHandle;
+	UE_DEPRECATED(5.0, "Removed for a cleaner API. Please make up logic in place where this was used (e.g. instead use GetSelectedModes().Num() == 1 from this class).")
+	bool CanAddMode() const;
 
-	/** Property handle for the function channel */
-	TSharedPtr<IPropertyHandle> FunctionChannelHandle;
+	UE_DEPRECATED(5.0, "Replaced with UDMXEntityFixtureType::AddMode.")
+	void AddMode();
 
-	/** Property handle for the function */
-	TSharedPtr<IPropertyHandle> FunctionHandle;
-};
+	UE_DEPRECATED(5.0, "Replaced with UDMXEntityFixtureType::DuplicateModes.")
+	void DuplicateModes(const TArray<int32>& ModeIndicesToDuplicate);
 
+	UE_DEPRECATED(5.0, "Replaced with UDMXEntityFixtureType::DeleteModes.")
+	void DeleteModes(const TArray<int32>& ModeIndicesToDelete);
 
-/** Shared data for fixture type editors. */
-class FDMXFixtureTypeSharedData
-	: public TSharedFromThis<FDMXFixtureTypeSharedData>
-{
-	friend class FDMXFixtureModeItem;
-	friend class FDMXFixtureFunctionItem;
+	UE_DEPRECATED(5.0, "Replaced with FDMXFixtureTypeSharedData::PasteModesFromClipboard.")
+	void PasteClipboardToModes(const TArray<int32>& ModeIndices);
 
-public:
-	FDMXFixtureTypeSharedData(TWeakPtr<FDMXEditor> InDMXEditorPtr)
-		: DMXEditorPtr(InDMXEditorPtr)
-	{}	
+	UE_DEPRECATED(5.0, "Removed for a cleaner API. Please make up logic in place where this was used (e.g. instead use GetSelectedFunctions().Num() == 1 from this class).")
+	bool CanAddFunction() const;
 
-public:
-	/** Sets which fixture types are being edited in shared data */
-	void SetFixtureTypesBeingEdited(const TArray<TWeakObjectPtr<UDMXEntityFixtureType>>& InFixtureTypesBeingEdited);
+	UE_DEPRECATED(5.0, "Replaced with UDMXEntityFixtureType::AddFunction.")
+	void AddFunctionToSelectedMode();
 
-	/** Sets which modes are being edited */
-	void SetModesBeingEdited(const TArray<TSharedPtr<FDMXFixtureModeItem>>& InModesBeingEdited) { ModesBeingEdited = InModesBeingEdited; }
+	UE_DEPRECATED(5.0, "Replaced with UDMXEntityFixtureType::DuplicateFunctions.")
+	void DuplicateFunctions(const TArray<int32>& FunctionIndicesToDuplicate);
 
-	/** Sets which functions are being edited */
-	void SetFunctionsBeingEdited(const TArray<TSharedPtr<FDMXFixtureFunctionItem>>& InFunctionsBeingEdited) { FunctionsBeingEdited = InFunctionsBeingEdited; }
+	UE_DEPRECATED(5.0, "Replaced with UDMXEntityFixtureType::DeleteFunctions.")
+	void DeleteFunctions(const TArray<int32>& FunctionIndicesToDelete);
 
-	/** Selects specified modes */
-	void SelectModes(const TArray<TSharedPtr<FDMXFixtureModeItem>>& InModeItems);
+	UE_DEPRECATED(5.0, "Replaced with FDMXFixtureTypeSharedData::PasteFunctionsFromClipboard.")
+	void PasteClipboardToFunctions(const TArray<int32>& FunctionIndices);
 
-	/** Selects specified functions */
-	void SelectFunctions(const TArray<TSharedPtr<FDMXFixtureFunctionItem>>& InFunctions);
+	UE_DEPRECATED(5.0, "Removed for a cleaner API. Please make up logic in place where this was used (e.g. instead use IsMatrixSelected() from this class).")
+	bool CanAddCellAttribute() const;
 
-	/** Returns the number of modes being edited */
-	int32 GetNumSelectedModes() const;
+	UE_DEPRECATED(5.0, "Replaced with UDMXEntityFixtureType::AddCellAttribute.")
+	void AddCellAttributeToSelectedMode();
 
-	/** Returns the number of modes being edited */
-	int32 GetNumSelectedFunctions() const;
+	/** Broadcasts when selected Fixture Types changed */
+	FSimpleMulticastDelegate OnFixtureTypesSelected;
 
-	/** Broadcasts when selected modes changed */
+	/** Broadcasts when selected Modes changed */
 	FSimpleMulticastDelegate OnModesSelected;
 
-	/** Broadcasts when selected functions changed */
+	/** Broadcasts when selected Functions changed */
 	FSimpleMulticastDelegate OnFunctionsSelected;
 
-public:	
-	/** True when a mode can be added (true, if one fixture type is edited) */
-	bool CanAddMode() const;
-
-	/** Adds a mode to edited fixture type */
-	void AddMode(const FScopedTransaction& TransactionScope);
-
-	/** Duplicates specified modes in the modes arrays they reside in */
-	void DuplicateModes(const TArray<TSharedPtr<FDMXFixtureModeItem>>& ModesToDuplicate);
+	/** Broadcasts when the Matrix was selected or unselected */
+	FSimpleMulticastDelegate OnMatrixSelectionChanged;
 
-	/** Deletes specified modes from the modes arrays they reside in */
-	void DeleteModes(const TArray<TSharedPtr<FDMXFixtureModeItem>>& ModesToDelete);
-
-	/** Copies functions to clipboard */
-	void CopyModesToClipboard(const TArray<TSharedPtr<FDMXFixtureModeItem>>& ModeItems);
-
-	/** Pastes clipboard to specified functions*/
-	void PasteClipboardToModes(const TArray<TSharedPtr<FDMXFixtureModeItem>>& ModeItems);
-
-public:
-	/** True when a function can be added (true, if one mode is selected) */
-	bool CanAddFunction() const;
-
-	/** Duplicates specified modes in the modes arrays they reside in */
-	void DuplicateFunctions(const TArray<TSharedPtr<FDMXFixtureFunctionItem>>& FunctionsToDuplicate);
-
-	/** Deletes specified modes from the modes arrays they reside in */
-	void DeleteFunctions(const TArray<TSharedPtr<FDMXFixtureFunctionItem>>& FunctionsToDelete);
-
-	/** Adds a function to selected mode, only if CanAddFunction */
-	void AddFunctionToSelectedMode();
+private:
+	/** The Fixture types being edited */
+	TArray<TWeakObjectPtr<UDMXEntityFixtureType>> SelectedFixtureTypes;
 
-	/** Copies functions to clipboard */
-	void CopyFunctionsToClipboard(const TArray<TSharedPtr<FDMXFixtureFunctionItem>>& FunctionItems);
+	/** The Mode indices in the Selected Fixture Types currently being selected */
+	TArray<int32> SelectedModeIndices;
 
-	/** Pastes clipboard to specified functions*/
-	void PasteClipboardToFunctions(const TArray<TSharedPtr<FDMXFixtureFunctionItem>>& FunctionItems);
+	/** The Function indices in the Selected Fixture Types currently being selected  */
+	TArray<int32> SelectedFunctionIndices;
 
-protected:
-	// Protected to prevent from copying name properties across detail customizations
-	const TArray<TWeakObjectPtr<UDMXEntityFixtureType>>& GetFixtureTypesBeingEdited() const { return FixtureTypesBeingEdited; }
-	const TArray<TSharedPtr<FDMXFixtureModeItem>>& GetModesBeingEdited() const { return ModesBeingEdited; }
-	const TArray<TSharedPtr<FDMXFixtureFunctionItem>>& GetFunctionsBeingEdited() const { return FunctionsBeingEdited; }
-	const TArray<TSharedPtr<FDMXFixtureModeItem>>& GetSelectedModes() const { return SelectedModes; }
-	const TArray<TSharedPtr<FDMXFixtureFunctionItem>>& GetSelectedFunctions() const { return SelectedFunctions; }
+	/** If true the Fixture Matrices in the currently seleccted Modes are selected */
+	bool bFixtureMatrixSelected = false;
 
-private:
-	/** Internal cache for multi mode copy/paste*/
+	/** Cache for multi mode copy/paste*/
 	TArray<FString> ModesClipboard;
 
-	/** Internal cache for multi function copy/paste*/
+	/** Cache for multi function copy/paste*/
 	TArray<FString> FunctionsClipboard;
 
-private:
-	/** The Fixture types being edited */
-	TArray<TWeakObjectPtr<UDMXEntityFixtureType>> FixtureTypesBeingEdited;
-
-	/** Modes being edited */
-	TArray<TSharedPtr<FDMXFixtureModeItem>> ModesBeingEdited;
-
-	/** Gunctions being edited */
-	TArray<TSharedPtr<FDMXFixtureFunctionItem>> FunctionsBeingEdited;
-
-	/** Selected modes in fixture types */
-	TArray<TSharedPtr<FDMXFixtureModeItem>> SelectedModes;
-
-	/** Selected functions in modes */
-	TArray<TSharedPtr<FDMXFixtureFunctionItem>> SelectedFunctions;
-
 	/** Weak reference to the DMX editor */
 	TWeakPtr<FDMXEditor> DMXEditorPtr;	
-
-protected:
-	/** Serializes a struct */
-	template <typename StructType>
-	bool SerializeStruct(const StructType& Struct, FString& OutCopyStr) const
-	{
-		TSharedRef<FJsonObject> RootJsonObject = MakeShareable(new FJsonObject());
-
-		FJsonObjectConverter::UStructToJsonObject(StructType::StaticStruct(), &Struct, RootJsonObject, 0, 0);
-
-		typedef TJsonWriter<TCHAR, TPrettyJsonPrintPolicy<TCHAR>> FStringWriter;
-		typedef TJsonWriterFactory<TCHAR, TPrettyJsonPrintPolicy<TCHAR>> FStringWriterFactory;
-
-		FString CopyStr;
-		TSharedRef<FStringWriter> Writer = FStringWriterFactory::Create(&CopyStr);
-		FJsonSerializer::Serialize(RootJsonObject, Writer);
-
-		if (!CopyStr.IsEmpty())
-		{
-			OutCopyStr = CopyStr;
-			return true;
-		}
-
-		return false;
-	}
 };
diff --git a/Engine/Plugins/VirtualProduction/DMX/DMXEngine/Source/DMXRuntime/DMXRuntime.Build.cs b/Engine/Plugins/VirtualProduction/DMX/DMXEngine/Source/DMXRuntime/DMXRuntime.Build.cs
--- a/Engine/Plugins/VirtualProduction/DMX/DMXEngine/Source/DMXRuntime/DMXRuntime.Build.cs
+++ b/Engine/Plugins/VirtualProduction/DMX/DMXEngine/Source/DMXRuntime/DMXRuntime.Build.cs
@@ -15,6 +15,8 @@ public class DMXRuntime : ModuleRules
                 "CoreUObject",
                 "Engine",
 				"DMXProtocol",
+				"JsonUtilities",
+
 			}
 		);
 
diff --git a/Engine/Plugins/VirtualProduction/DMX/DMXEngine/Source/DMXRuntime/Private/DMXRuntimeObjectVersion.h b/Engine/Plugins/VirtualProduction/DMX/DMXEngine/Source/DMXRuntime/Private/DMXRuntimeObjectVersion.h
--- a/Engine/Plugins/VirtualProduction/DMX/DMXEngine/Source/DMXRuntime/Private/DMXRuntimeObjectVersion.h
+++ b/Engine/Plugins/VirtualProduction/DMX/DMXEngine/Source/DMXRuntime/Private/DMXRuntimeObjectVersion.h
@@ -18,6 +18,12 @@ struct FDMXRuntimeObjectVersion
 		// Update to DMX Library Section using normalized values by default
 		ReplaceWeakWithStrongFixturePatchReferncesInLibrarySection,
 
+		// Update DMX Entity Fixture Type to allow the Matrix to be enabled/disabled per Mode
+		DMXFixtureTypeAllowMatrixInEachFixtureMode,
+
+		// Update DMX Entity Fixture Type to allow for Input Modulators per Mode
+		DMXFixtureTypeAllowInputModulatorsInEachFixtureMode,
+
 		// -----<new versions can be added above this line>-------------------------------------------------
 		VersionPlusOne,
 		LatestVersion = VersionPlusOne - 1
diff --git a/Engine/Plugins/VirtualProduction/DMX/DMXEngine/Source/DMXRuntime/Private/DMXRuntimeUtils.cpp b/Engine/Plugins/VirtualProduction/DMX/DMXEngine/Source/DMXRuntime/Private/DMXRuntimeUtils.cpp
--- a/Engine/Plugins/VirtualProduction/DMX/DMXEngine/Source/DMXRuntime/Private/DMXRuntimeUtils.cpp
+++ b/Engine/Plugins/VirtualProduction/DMX/DMXEngine/Source/DMXRuntime/Private/DMXRuntimeUtils.cpp
@@ -6,6 +6,49 @@
 
 #define LOCTEXT_NAMESPACE "FDMXRuntimeUtils"
 
+
+FString FDMXRuntimeUtils::GenerateUniqueNameFromExisting(const TSet<FString>& InExistingNames, const FString& InBaseName)
+{
+	if (!InBaseName.IsEmpty() && !InExistingNames.Contains(InBaseName))
+	{
+		return InBaseName;
+	}
+
+	FString FinalName;
+	FString BaseName;
+
+	int32 Index = 0;
+	if (InBaseName.IsEmpty())
+	{
+		BaseName = TEXT("Default name");
+	}
+	else
+	{
+		// If there's an index at the end of the name, start from there
+		FDMXRuntimeUtils::GetNameAndIndexFromString(InBaseName, BaseName, Index);
+	}
+
+	int32 Count = (Index == 0) ? 1 : Index;
+	FinalName = BaseName;
+	// Add Count to the BaseName, increasing Count, until it's a non-existent name
+	do
+	{
+		// Calculate the number of digits in the number, adding 2 (1 extra to correctly count digits, another to account for the '_' that will be added to the name
+		int32 CountLength = Count > 0 ? (int32)FGenericPlatformMath::LogX(10.0f, Count) + 2 : 2;
+
+		// If the length of the final string will be too long, cut off the end so we can fit the number
+		if (CountLength + BaseName.Len() >= NAME_SIZE)
+		{
+			BaseName = BaseName.Left(NAME_SIZE - CountLength - 1);
+		}
+
+		FinalName = FString::Printf(TEXT("%s_%d"), *BaseName, Count);
+		++Count;
+	} while (InExistingNames.Contains(FinalName));
+
+	return FinalName;
+}
+
 bool FDMXRuntimeUtils::GetNameAndIndexFromString(const FString& InString, FString& OutName, int32& OutIndex)
 {
 	OutName = InString.TrimEnd();
@@ -109,4 +152,25 @@ FString FDMXRuntimeUtils::GenerateUniqueNameForImportFunction(TMap<FString, uint
 	return FinalName;
 }
 
+TOptional<FString> FDMXRuntimeUtils::SerializeStructToString(const UStruct* StructClass, const void* StructData)
+{
+	TSharedRef<FJsonObject> RootJsonObject = MakeShareable(new FJsonObject());
+
+	FJsonObjectConverter::UStructToJsonObject(StructClass, StructData, RootJsonObject, 0, 0);
+
+	typedef TJsonWriter<TCHAR, TPrettyJsonPrintPolicy<TCHAR>> FStringWriter;
+	typedef TJsonWriterFactory<TCHAR, TPrettyJsonPrintPolicy<TCHAR>> FStringWriterFactory;
+
+	FString CopyStr;
+	TSharedRef<FStringWriter> Writer = FStringWriterFactory::Create(&CopyStr);
+	FJsonSerializer::Serialize(RootJsonObject, Writer);
+
+	if (!CopyStr.IsEmpty())
+	{
+		return CopyStr;
+	}
+
+	return TOptional<FString>();
+}
+
 #undef LOCTEXT_NAMESPACE
diff --git a/Engine/Plugins/VirtualProduction/DMX/DMXEngine/Source/DMXRuntime/Private/DMXSubsystem.cpp b/Engine/Plugins/VirtualProduction/DMX/DMXEngine/Source/DMXRuntime/Private/DMXSubsystem.cpp
--- a/Engine/Plugins/VirtualProduction/DMX/DMXEngine/Source/DMXRuntime/Private/DMXSubsystem.cpp
+++ b/Engine/Plugins/VirtualProduction/DMX/DMXEngine/Source/DMXRuntime/Private/DMXSubsystem.cpp
@@ -3,6 +3,7 @@
 #include "DMXSubsystem.h"
 
 #include "DMXAttribute.h"
+#include "DMXConversions.h"
 #include "DMXProtocolSettings.h"
 #include "DMXProtocolTypes.h"
 #include "DMXRuntimeUtils.h"
@@ -400,7 +401,7 @@ void UDMXSubsystem::GetFixtureAttributes(const UDMXEntityFixturePatch* InFixture
 			{
 				for (const FDMXFixtureFunction& Function : ActiveModePtr->Functions)
 				{
-					if (!UDMXEntityFixtureType::IsFunctionInModeRange(Function, *ActiveModePtr, StartingAddress))
+					if (Function.GetLastChannel() <= InFixturePatch->GetChannelSpan())
 					{
 						// This function and the following ones are outside the Universe's range.
 						break;
@@ -464,7 +465,7 @@ bool UDMXSubsystem::GetFunctionsMap(UDMXEntityFixturePatch* InFixturePatch, TMap
 		for (const FDMXFixtureFunction& Function : ModePtr->Functions)
 		{
 			const int32 FunctionStartIndex = Function.Channel - 1 + PatchStartingIndex;
-			const int32 FunctionLastIndex = FunctionStartIndex + UDMXEntityFixtureType::NumChannelsToOccupy(Function.DataType) - 1;
+			const int32 FunctionLastIndex = FunctionStartIndex + FDMXConversions::GetSizeOfSignalFormat(Function.DataType) - 1;
 			if (FunctionLastIndex >= ChannelData.Num())
 			{
 				break;
@@ -705,7 +706,7 @@ float UDMXSubsystem::BytesToNormalizedValue(const TArray<uint8>& Bytes, bool bUs
 
 void UDMXSubsystem::NormalizedValueToBytes(float InValue, EDMXFixtureSignalFormat InSignalFormat, TArray<uint8>& Bytes, bool bUseLSB /*= false*/) const
 {
-	const uint8 NumBytes = UDMXEntityFixtureType::NumChannelsToOccupy(InSignalFormat);
+	const uint8 NumBytes = FDMXConversions::GetSizeOfSignalFormat(InSignalFormat);
 	// Make sure the array will fit the correct number of bytes
 	Bytes.Reset(NumBytes);
 	Bytes.AddZeroed(NumBytes);
@@ -715,7 +716,7 @@ void UDMXSubsystem::NormalizedValueToBytes(float InValue, EDMXFixtureSignalForma
 
 void UDMXSubsystem::IntValueToBytes(int32 InValue, EDMXFixtureSignalFormat InSignalFormat, TArray<uint8>& Bytes, bool bUseLSB /*= false*/)
 {
-	const uint8 NumBytes = UDMXEntityFixtureType::NumChannelsToOccupy(InSignalFormat);
+	const uint8 NumBytes = FDMXConversions::GetSizeOfSignalFormat(InSignalFormat);
 	// Make sure the array will fit the correct number of bytes
 	Bytes.Reset(NumBytes);
 	Bytes.AddZeroed(NumBytes);
@@ -725,7 +726,7 @@ void UDMXSubsystem::IntValueToBytes(int32 InValue, EDMXFixtureSignalFormat InSig
 
 float UDMXSubsystem::IntToNormalizedValue(int32 InValue, EDMXFixtureSignalFormat InSignalFormat) const
 {
-	return (float)(uint32)(InValue) / UDMXEntityFixtureType::GetDataTypeMaxValue(InSignalFormat);
+	return (float)(uint32)(InValue) / FDMXConversions::GetSignalFormatMaxValue(InSignalFormat);
 }
 
 float UDMXSubsystem::GetNormalizedAttributeValue(UDMXEntityFixturePatch* InFixturePatch, FDMXAttributeName InFunctionAttribute, int32 InValue) const
diff --git a/Engine/Plugins/VirtualProduction/DMX/DMXEngine/Source/DMXRuntime/Private/Library/DMXEntityFixturePatch.cpp b/Engine/Plugins/VirtualProduction/DMX/DMXEngine/Source/DMXRuntime/Private/Library/DMXEntityFixturePatch.cpp
--- a/Engine/Plugins/VirtualProduction/DMX/DMXEngine/Source/DMXRuntime/Private/Library/DMXEntityFixturePatch.cpp
+++ b/Engine/Plugins/VirtualProduction/DMX/DMXEngine/Source/DMXRuntime/Private/Library/DMXEntityFixturePatch.cpp
@@ -7,7 +7,6 @@
 #include "DMXRuntimeLog.h"
 #include "DMXStats.h"
 #include "DMXTypes.h"
-#include "DMXUtils.h"
 #include "Interfaces/IDMXProtocol.h"
 #include "IO/DMXInputPort.h"
 #include "IO/DMXOutputPort.h"
@@ -36,11 +35,27 @@ UDMXEntityFixturePatch::UDMXEntityFixturePatch()
 #endif // WITH_EDITORONLY_DATA
 {}
 
+#if WITH_EDITOR
+bool UDMXEntityFixturePatch::Modify(bool bAlwaysMarkDirty)
+{
+	if (UDMXLibrary* DMXLibrary = ParentLibrary.Get())
+	{
+		return DMXLibrary->Modify(bAlwaysMarkDirty) && Super::Modify(bAlwaysMarkDirty);
+	}
+
+	return Super::Modify(bAlwaysMarkDirty);
+}
+#endif // WITH_EDITOR
+
 void UDMXEntityFixturePatch::PostInitProperties()
 {
 	Super::PostInitProperties();
 
-	RebuildCache();
+	if (!HasAnyFlags(RF_ClassDefaultObject))
+	{
+		UDMXEntityFixtureType::GetOnFixtureTypeChanged().AddUObject(this, &UDMXEntityFixturePatch::OnFixtureTypeChanged);
+		RebuildCache();
+	}
 }
 
 void UDMXEntityFixturePatch::PostLoad()
@@ -66,7 +81,7 @@ void UDMXEntityFixturePatch::PostEditChangeProperty(FPropertyChangedEvent& Prope
 
 void UDMXEntityFixturePatch::Tick(float DeltaTime)
 {
-	bool bDataChanged = UpdateCache();
+	const bool bDataChanged = UpdateCache();
 
 	// Broadcast data changes
 	if (bDataChanged)
@@ -117,24 +132,26 @@ void UDMXEntityFixturePatch::SendDMX(TMap<FDMXAttributeName, int32> AttributeMap
 	{
 		for (const FDMXFixtureFunction& Function : ModePtr->Functions)
 		{
-			const FDMXAttributeName FunctionAttr = Function.Attribute;
+			const FDMXAttributeName& FunctionAttr = Function.Attribute;
 			if (FunctionAttr == Elem.Key)
 			{
-				if (!UDMXEntityFixtureType::IsFunctionInModeRange(Function, *ModePtr, GetStartingChannel() - 1))
-				{
-					continue;
-				}
+				const int32 LastFunctionChannel = Function.Channel + Function.GetNumChannels() - 1;
+				const bool bLastChannelExceedsChannelSpan = LastFunctionChannel > ModePtr->ChannelSpan;
+				const bool bLastChannelExceedsUniverseSize = LastFunctionChannel + GetStartingChannel() - 1 > DMX_MAX_ADDRESS;
 
-				const int32 Channel = Function.Channel + GetStartingChannel() - 1;
+				if (!bLastChannelExceedsChannelSpan && !bLastChannelExceedsUniverseSize)
+				{
+					const int32 Channel = Function.Channel + GetStartingChannel() - 1;
 
-				uint32 ChannelValue = 0;
-				uint8* ChannelValueBytes = reinterpret_cast<uint8*>(&ChannelValue);
-				UDMXEntityFixtureType::FunctionValueToBytes(Function, Elem.Value, ChannelValueBytes);
+					uint32 ChannelValue = 0;
+					uint8* ChannelValueBytes = reinterpret_cast<uint8*>(&ChannelValue);
+					UDMXEntityFixtureType::FunctionValueToBytes(Function, Elem.Value, ChannelValueBytes);
 
-				const uint8 NumBytesInSignalFormat = UDMXEntityFixtureType::NumChannelsToOccupy(Function.DataType);
-				for (uint8 ChannelIt = 0; ChannelIt < NumBytesInSignalFormat; ++ChannelIt)
-				{
-					DMXChannelToValueMap.Add(Channel + ChannelIt, ChannelValueBytes[ChannelIt]);
+					const uint8 NumBytesInSignalFormat = FDMXConversions::GetSizeOfSignalFormat(Function.DataType);
+					for (uint8 ChannelIt = 0; ChannelIt < NumBytesInSignalFormat; ++ChannelIt)
+					{
+						DMXChannelToValueMap.Add(Channel + ChannelIt, ChannelValueBytes[ChannelIt]);
+					}
 				}
 			}
 		}
@@ -302,6 +319,8 @@ const FDMXFixtureMode* UDMXEntityFixturePatch::GetActiveMode() const
 
 void UDMXEntityFixturePatch::SetFixtureType(UDMXEntityFixtureType* NewFixtureType)
 {
+	Modify();
+
 	if (NewFixtureType && NewFixtureType != ParentFixtureTypeTemplate)
 	{
 		ParentFixtureTypeTemplate = NewFixtureType;
@@ -318,6 +337,8 @@ void UDMXEntityFixturePatch::SetFixtureType(UDMXEntityFixtureType* NewFixtureTyp
 
 void UDMXEntityFixturePatch::SetUniverseID(int32 NewUniverseID)
 {
+	Modify();
+
 	UniverseID = NewUniverseID;
 
 	RebuildCache();
@@ -325,6 +346,8 @@ void UDMXEntityFixturePatch::SetUniverseID(int32 NewUniverseID)
 
 void UDMXEntityFixturePatch::SetAutoStartingAddress(int32 NewAutoStartingAddress)
 {
+	Modify();
+
 	AutoStartingAddress = NewAutoStartingAddress;
 	ManualStartingAddress = NewAutoStartingAddress;
 
@@ -333,11 +356,29 @@ void UDMXEntityFixturePatch::SetAutoStartingAddress(int32 NewAutoStartingAddress
 
 void UDMXEntityFixturePatch::SetManualStartingAddress(int32 NewManualStartingAddress)
 {
+	Modify();
+
 	ManualStartingAddress = NewManualStartingAddress;
 
 	RebuildCache();
 }
 
+void UDMXEntityFixturePatch::SetStartingChannel(int32 NewStartingChannel)
+{
+	Modify();
+
+	if (NewStartingChannel == AutoStartingAddress && ManualStartingAddress == NewStartingChannel)
+	{
+		return;
+	}
+
+	bAutoAssignAddress = false;
+	AutoStartingAddress = NewStartingChannel;
+	ManualStartingAddress = NewStartingChannel;
+
+	RebuildCache();
+}
+
 int32 UDMXEntityFixturePatch::GetStartingChannel() const
 {
 	if (bAutoAssignAddress)
@@ -362,6 +403,8 @@ int32 UDMXEntityFixturePatch::GetEndingChannel() const
 
 bool UDMXEntityFixturePatch::SetActiveModeIndex(int32 NewActiveModeIndex)
 {
+	Modify();
+
 	if (IsValid(ParentFixtureTypeTemplate))
 	{
 		if (ParentFixtureTypeTemplate->Modes.IsValidIndex(NewActiveModeIndex))
@@ -386,7 +429,7 @@ TArray<FDMXAttributeName> UDMXEntityFixturePatch::GetAllAttributesInActiveMode()
 	TArray<FDMXAttributeName> AttributeNames;
 	for (const FDMXFixtureFunction& Function : Cache.GetFunctions())
 	{
-		if (UDMXEntityFixtureType::GetFunctionLastChannel(Function) <= Cache.GetChannelSpan())
+		if (Function.GetLastChannel() <= Cache.GetChannelSpan())
 		{
 			AttributeNames.Add(Function.Attribute.GetName());
 		}
@@ -402,7 +445,7 @@ TMap<FDMXAttributeName, FDMXFixtureFunction> UDMXEntityFixturePatch::GetAttribut
 
 	for (const FDMXFixtureFunction& Function : Cache.GetFunctions())
 	{
-		if (UDMXEntityFixtureType::GetFunctionLastChannel(Function) <= Cache.GetChannelSpan())
+		if (Function.GetLastChannel() <= Cache.GetChannelSpan())
 		{
 			FunctionMap.Add(Function.Attribute.GetName(), Function);
 		}
@@ -418,7 +461,7 @@ TMap<FDMXAttributeName, int32> UDMXEntityFixturePatch::GetAttributeDefaultMap()
 
 	for (const FDMXFixtureFunction& Function : Cache.GetFunctions())
 	{
-		if (UDMXEntityFixtureType::GetFunctionLastChannel(Function) <= Cache.GetChannelSpan())
+		if (Function.GetLastChannel() <= Cache.GetChannelSpan())
 		{
 			DefaultValueMap.Add(Function.Attribute, Function.DefaultValue);
 		}
@@ -434,7 +477,7 @@ TMap<FDMXAttributeName, int32> UDMXEntityFixturePatch::GetAttributeChannelAssign
 
 	for (const FDMXFixtureFunction& Function : Cache.GetFunctions())
 	{
-		if (UDMXEntityFixtureType::GetFunctionLastChannel(Function) <= Cache.GetChannelSpan())
+		if (Function.GetLastChannel() <= Cache.GetChannelSpan())
 		{
 			ChannelMap.Add(Function.Attribute, Function.Channel + GetStartingChannel() - 1);
 		}
@@ -450,7 +493,7 @@ TMap<FDMXAttributeName, EDMXFixtureSignalFormat> UDMXEntityFixturePatch::GetAttr
 
 	for (const FDMXFixtureFunction& Function : Cache.GetFunctions())
 	{
-		if (UDMXEntityFixtureType::GetFunctionLastChannel(Function) <= Cache.GetChannelSpan())
+		if (Function.GetLastChannel() <= Cache.GetChannelSpan())
 		{
 			FormatMap.Add(Function.Attribute, Function.DataType);
 		}
@@ -466,7 +509,7 @@ TMap<FDMXAttributeName, int32> UDMXEntityFixturePatch::ConvertRawMapToAttributeM
 
 	for (const FDMXFixtureFunction& Function : Cache.GetFunctions())
 	{
-		if (UDMXEntityFixtureType::GetFunctionLastChannel(Function) <= Cache.GetChannelSpan())
+		if (Function.GetLastChannel() <= Cache.GetChannelSpan())
 		{
 			const int32 FunctionStartingChannel = Function.Channel + (GetStartingChannel() - 1);
 
@@ -474,7 +517,7 @@ TMap<FDMXAttributeName, int32> UDMXEntityFixturePatch::ConvertRawMapToAttributeM
 			if (RawMap.Contains(FunctionStartingChannel))
 			{
 				const uint8& RawValue(RawMap.FindRef(Function.Channel));
-				const uint8 ChannelsToAdd = UDMXEntityFixtureType::NumChannelsToOccupy(Function.DataType);
+				const uint8 ChannelsToAdd = FDMXConversions::GetSizeOfSignalFormat(Function.DataType);
 
 				TArray<uint8, TFixedAllocator<4>> Bytes;
 				for (uint8 ChannelIt = 0; ChannelIt < ChannelsToAdd; ChannelIt++)
@@ -499,19 +542,21 @@ TMap<int32, uint8> UDMXEntityFixturePatch::ConvertAttributeMapToRawMap(const TMa
 	TMap<int32, uint8> RawMap;
 	RawMap.Reserve(FunctionMap.Num());
 
-	for (const TPair<FDMXAttributeName, int32>& Elem : FunctionMap)
+	for (const TPair<FDMXAttributeName, int32>& AttributeNameValuePair : FunctionMap)
 	{
 		// Search for a function with Attribute == Elem.Key.Attribute
-		const FDMXFixtureFunction* FunctionPtr = Cache.GetFunctions().FindByPredicate([&Elem](const FDMXFixtureFunction& Function) {
-			return Elem.Key == Function.Attribute;
+		const FDMXFixtureFunction* FunctionPtr = Cache.GetFunctions().FindByPredicate([&AttributeNameValuePair](const FDMXFixtureFunction& Function) {
+			return AttributeNameValuePair.Key == Function.Attribute;
 		});
 
 		// Also check for the Function being in the valid range for the Active Mode's channel span
-		if (FunctionPtr != nullptr && UDMXEntityFixtureType::GetFunctionLastChannel(*FunctionPtr) <= Cache.GetChannelSpan())
+		if (FunctionPtr && FunctionPtr->GetLastChannel() <= Cache.GetChannelSpan())
 		{
 			const int32 FunctionStartingChannel = FunctionPtr->Channel + (GetStartingChannel() - 1);
-			const uint8 NumChannels = UDMXEntityFixtureType::NumChannelsToOccupy(FunctionPtr->DataType);
-			const uint32 Value = UDMXEntityFixtureType::ClampValueToDataType(FunctionPtr->DataType, Elem.Value);
+			const uint8 NumChannels = FDMXConversions::GetSizeOfSignalFormat(FunctionPtr->DataType);
+
+			const uint32 SafeValue = FMath::Max(0, AttributeNameValuePair.Value);
+			const uint32 ClampedValue = FDMXConversions::ClampValueBySignalFormat(SafeValue, FunctionPtr->DataType);
 
 			// To avoid branching in the loop, we'll decide before it on which byte to start
 			// and which direction to go, depending on the Function's bit endianness.
@@ -520,7 +565,7 @@ TMap<int32, uint8> UDMXEntityFixturePatch::ConvertAttributeMapToRawMap(const TMa
 
 			for (uint8 ByteOffset = 0; ByteOffset < NumChannels; ++ByteOffset)
 			{
-				const uint8 ChannelVal = (Value >> (8 * ByteOffset)) & 0xFF;
+				const uint8 ChannelVal = (ClampedValue >> (8 * ByteOffset)) & 0xFF;
 
 				const int32 FinalChannel = FunctionStartingChannel + ByteIndex;
 				RawMap.Add(FinalChannel, ChannelVal);
@@ -572,7 +617,7 @@ TMap<FDMXAttributeName, int32> UDMXEntityFixturePatch::ConvertToValidMap(const T
 		});
 
 		// Also check for the Function being in the valid range for the Active Mode's channel span
-		if (FunctionPtr && UDMXEntityFixtureType::GetFunctionLastChannel(*FunctionPtr) <= Cache.GetChannelSpan())
+		if (FunctionPtr && FunctionPtr->GetLastChannel() <= Cache.GetChannelSpan())
 		{
 			ValidMap.Add(Elem.Key, Elem.Value);
 		}
@@ -681,9 +726,10 @@ bool UDMXEntityFixturePatch::SendMatrixCellValue(const FIntPoint& CellCoordinate
 			int32 AttributeOffset = 0;
 			for (const FDMXFixtureCellAttribute& CellAttribute : Cache.GetCellAttributes())
 			{
-				if (CellAttribute.Attribute.Name == Attribute && Value >= 0)
+				if (CellAttribute.Attribute.Name == Attribute)
 				{
-					TArray<uint8> ByteArray = FDMXConversions::UnsignedInt32ToByteArray(Value, CellAttribute.DataType, CellAttribute.bUseLSBMode);
+					const uint32 SafeValue = FMath::Max(0, Value);
+					TArray<uint8> ByteArray = FDMXConversions::UnsignedInt32ToByteArray(SafeValue, CellAttribute.DataType, CellAttribute.bUseLSBMode);
 
 					TMap<int32, uint8> DMXChannelToValueMap;
 					DMXChannelToValueMap.Reserve(ByteArray.Num());
@@ -740,7 +786,7 @@ bool UDMXEntityFixturePatch::SendMatrixCellValueWithAttributeMap(const FIntPoint
 		}
 
 		const int32 FirstChannel = InAttributeNameChannelMap[Attribute];
-		const int32 LastChannel = FirstChannel + UDMXEntityFixtureType::NumChannelsToOccupy(CellAttribute.DataType) - 1;
+		const int32 LastChannel = FirstChannel + FDMXConversions::GetSizeOfSignalFormat(CellAttribute.DataType) - 1;
 
 		TArray<uint8> ByteArr;
 		ByteArr.AddZeroed(4);
@@ -787,7 +833,7 @@ bool UDMXEntityFixturePatch::SendNormalizedMatrixCellValue(const FIntPoint& Cell
 
 		Value = FMath::Clamp(Value, 0.0f, 1.0f);
 
-		const uint32 UnsignedIntValue = UDMXEntityFixtureType::GetDataTypeMaxValue(ExistingAttribute.DataType) * Value;
+		const uint32 UnsignedIntValue = FDMXConversions::GetSignalFormatMaxValue(ExistingAttribute.DataType) * Value;
 		const int32  IntValue = FMath::Clamp(UnsignedIntValue, static_cast<uint32>(0), static_cast<uint32>(TNumericLimits<int32>::Max()));
 	
 		SendMatrixCellValue(CellCoordinate, Attribute, IntValue);
@@ -897,16 +943,16 @@ bool UDMXEntityFixturePatch::GetMatrixProperties(FDMXFixtureMatrix& MatrixProper
 		return false;
 	}
 
-	if (!ParentFixtureTypeTemplate->bFixtureMatrixEnabled)
+	const FDMXFixtureMode* ModePtr = GetActiveMode();
+	if (!ModePtr)
 	{
-		UE_LOG(DMXEntityFixturePatchLog, Error, TEXT("Fixture Patch %s is not a Matrix Fixture"), *Name);
+		UE_LOG(DMXEntityFixturePatchLog, Error, TEXT("Invalid active Mode in Fixture Patch %s"), *Name);
 		return false;
 	}
 
-	const FDMXFixtureMode* ModePtr = GetActiveMode();
-	if (!ModePtr)
+	if (!ModePtr->bFixtureMatrixEnabled)
 	{
-		UE_LOG(DMXEntityFixturePatchLog, Error, TEXT("Invalid active Mode in Fixture Patch %s"), *Name);
+		UE_LOG(DMXEntityFixturePatchLog, Error, TEXT("Fixture Patch %s is not a Matrix Fixture"), *Name);
 		return false;
 	}
 
@@ -982,13 +1028,22 @@ bool UDMXEntityFixturePatch::GetAllMatrixCells(TArray<FDMXCell>& Cells)
 	return false;
 }
 
+void UDMXEntityFixturePatch::OnFixtureTypeChanged(const UDMXEntityFixtureType* FixtureType)
+{	
+	if (ParentFixtureTypeTemplate == FixtureType)
+	{
+		ValidateActiveMode();
+		RebuildCache();
+	}
+}
+
 const FDMXFixtureMatrix* UDMXEntityFixturePatch::GetFixtureMatrix() const
 {
 	if (ParentFixtureTypeTemplate)
 	{
-		if (ParentFixtureTypeTemplate->bFixtureMatrixEnabled)
+		if (const FDMXFixtureMode* ModePtr = GetActiveMode())
 		{
-			if (const FDMXFixtureMode* ModePtr = GetActiveMode())
+			if (ModePtr->bFixtureMatrixEnabled)
 			{
 				return &ModePtr->FixtureMatrixConfig;
 			}
diff --git a/Engine/Plugins/VirtualProduction/DMX/DMXEngine/Source/DMXRuntime/Private/Library/DMXEntityFixturePatchCache.cpp b/Engine/Plugins/VirtualProduction/DMX/DMXEngine/Source/DMXRuntime/Private/Library/DMXEntityFixturePatchCache.cpp
--- a/Engine/Plugins/VirtualProduction/DMX/DMXEngine/Source/DMXRuntime/Private/Library/DMXEntityFixturePatchCache.cpp
+++ b/Engine/Plugins/VirtualProduction/DMX/DMXEngine/Source/DMXRuntime/Private/Library/DMXEntityFixturePatchCache.cpp
@@ -3,6 +3,7 @@
 #include "Library/DMXEntityFixturePatchCache.h"
 #include "Modulators/DMXModulator.h"
 
+#include "DMXConversions.h"
 #include "DMXRuntimeUtils.h"
 
 
@@ -144,7 +145,7 @@ bool FDMXEntityFixturePatchCache::InputDMXSignal(const FDMXSignalSharedPtr& DMXS
 				const uint32 IntValue = UDMXEntityFixtureType::BytesToFunctionValue(Function, CachedDMXValues.GetData() + FunctionStartIndex);
 				CachedRawValuesPerAttribute.Add(Function.Attribute, IntValue);
 
-				const float NormalizedValue = (float)IntValue / (float)UDMXEntityFixtureType::GetDataTypeMaxValue(Function.DataType);
+				const float NormalizedValue = (float)IntValue / (float)FDMXConversions::GetSignalFormatMaxValue(Function.DataType);
 				CachedNormalizedValuesPerAttribute.Map.Add(Function.Attribute, NormalizedValue);
 			}
 
@@ -175,7 +176,7 @@ bool FDMXEntityFixturePatchCache::InputDMXSignal(const FDMXSignalSharedPtr& DMXS
 						const uint32 IntValue = UDMXEntityFixtureType::BytesToInt(CellAttribute.DataType, CellAttribute.bUseLSBMode, CachedDMXValues.GetData() + FirstAttributeDataIndex);
 						CellAttributeToRawValueMap.Add(CellAttribute.Attribute.GetName(), IntValue);
 
-						const float NormalizedValue = (float)IntValue / (float)UDMXEntityFixtureType::GetDataTypeMaxValue(CellAttribute.DataType);
+						const float NormalizedValue = (float)IntValue / (float)FDMXConversions::GetSignalFormatMaxValue(CellAttribute.DataType);
 						CellAttributeTNormalizedValueMap.Map.Add(CellAttribute.Attribute.GetName(), NormalizedValue);
 
 						// Increment attribute offset for the next attribute
