DMX - Rework DMX Fixture Type editor for better GDTF support, better UX and to ease maintaining and extending the editor

diff --git a/Engine/Plugins/VirtualProduction/DMX/DMXEngine/Source/DMXRuntime/Private/Library/DMXEntityFixtureType.cpp b/Engine/Plugins/VirtualProduction/DMX/DMXEngine/Source/DMXRuntime/Private/Library/DMXEntityFixtureType.cpp
--- a/Engine/Plugins/VirtualProduction/DMX/DMXEngine/Source/DMXRuntime/Private/Library/DMXEntityFixtureType.cpp
+++ b/Engine/Plugins/VirtualProduction/DMX/DMXEngine/Source/DMXRuntime/Private/Library/DMXEntityFixtureType.cpp
@@ -2,48 +2,84 @@
 
 #include "Library/DMXEntityFixtureType.h"
 
+#include "DMXConversions.h"
+#include "DMXProtocolSettings.h"
+#include "DMXRuntimeObjectVersion.h"
+#include "DMXRuntimeUtils.h"
 #include "Library/DMXEntityFixturePatch.h"
-#include "Library/DMXLibrary.h"
 #include "Library/DMXImport.h"
 #include "Library/DMXImportGDTF.h"
-#include "DMXProtocolSettings.h"
-#include "DMXRuntimeUtils.h"
-#include "DMXUtils.h"
+#include "Library/DMXLibrary.h"
 
 
+#define LOCTEXT_NAMESPACE "DMXEntityFixtureType"
 
-#if WITH_EDITOR
-int32 FDMXFixtureMode::AddOrInsertFunction(int32 IndexOfFunction, const FDMXFixtureFunction& InFunction)
+uint8 FDMXFixtureFunction::GetLastChannel() const
 {
-	int32 Index = 0;
-	FDMXFixtureFunction FunctionToAdd = InFunction;
-
-	// Shift the insert function channel
-	uint8 DataTypeBytes = UDMXEntityFixtureType::NumChannelsToOccupy(InFunction.DataType);
-	FunctionToAdd.Channel = InFunction.Channel + DataTypeBytes;
+	return Channel + GetNumChannels() - 1;
+}
 
-	if (Functions.IsValidIndex(IndexOfFunction + 1))
+FDMXFixtureMatrix::FDMXFixtureMatrix()
+{
+	// Add a default cell attribute	
+	if (!GIsInitialLoad)
 	{
-		Index = Functions.Insert(InFunction, IndexOfFunction + 1);
-
-		// Shift all function after this by size of insert function
-		for (int32 FunctionIndex = Index + 1; FunctionIndex < Functions.Num(); FunctionIndex++)
+		if (const UDMXProtocolSettings* ProtocolSettings = GetDefault<UDMXProtocolSettings>())
 		{
-			FDMXFixtureFunction& Function = Functions[FunctionIndex];
-			Function.Channel = Function.Channel + DataTypeBytes;
+			FDMXAttributeName RedAttribute("Red");
+
+			if (ProtocolSettings->Attributes.Contains(RedAttribute.GetAttribute()))
+			{
+				FDMXFixtureCellAttribute RedCellAttribute;
+				RedCellAttribute.Attribute = RedAttribute;
+				CellAttributes = { RedCellAttribute };
+			}
+			else
+			{
+				FDMXFixtureCellAttribute VoidAttribute;
+				CellAttributes = { VoidAttribute };
+			}
 		}
 	}
-	else
+}
+
+int32 FDMXFixtureMatrix::GetNumChannels() const
+{
+	int32 NumChannelsPerCell = 0;
+	for (const FDMXFixtureCellAttribute& CellAttribute : CellAttributes)
 	{
-		Index = Functions.Add(InFunction);
+		NumChannelsPerCell += CellAttribute.GetNumChannels();
 	}
 
-	return Index;
+	return XCells * YCells * NumChannelsPerCell;
+}
+
+int32 FDMXFixtureMatrix::GetLastChannel() const
+{
+	const int32 NumCells = XCells * YCells;
+	if (NumCells == 0)
+	{
+		return FirstCellChannel;
+	}
+
+	int32 NumChannelsPerCell = 0;
+	for (const FDMXFixtureCellAttribute& Attribute : CellAttributes)
+	{
+		NumChannelsPerCell += Attribute.GetNumChannels();
+	}
+
+	return FirstCellChannel + NumCells * NumChannelsPerCell - 1;
+}
+
+int32 FDMXFixtureMatrix::GetFixtureMatrixLastChannel() const
+{
+	// DEPRECATED 5.0
+	return GetLastChannel();
 }
-#endif // WITH_EDITOR
 
 bool FDMXFixtureMatrix::GetChannelsFromCell(FIntPoint CellCoordinate, FDMXAttributeName Attribute, TArray<int32>& Channels) const
 {
+	// DEPRECATED 5.0
 	Channels.Reset();
 
 	TArray<int32> AllChannels;
@@ -76,7 +112,7 @@ bool FDMXFixtureMatrix::GetChannelsFromCell(FIntPoint CellCoordinate, FDMXAttrib
 	int32 CellAttributeSize = 0;
 	for (const FDMXFixtureCellAttribute& CellAttribute : CellAttributes)
 	{
-		int32 CurrentFunctionSize = UDMXEntityFixtureType::NumChannelsToOccupy(CellAttribute.DataType);
+		int32 CurrentFunctionSize = FDMXConversions::GetSizeOfSignalFormat(CellAttribute.DataType);
 		if (CellAttribute.Attribute.GetAttribute() == Attribute.GetAttribute())
 		{
 			CellAttributeIndex = CellSize;
@@ -101,29 +137,130 @@ bool FDMXFixtureMatrix::GetChannelsFromCell(FIntPoint CellCoordinate, FDMXAttrib
 	return true;
 }
 
-int32 FDMXFixtureMatrix::GetFixtureMatrixLastChannel() const
+#if WITH_EDITOR
+int32 FDMXFixtureMode::AddOrInsertFunction(int32 IndexOfFunction, FDMXFixtureFunction InFunction)
 {
-	int32 CellAttributeSize = 0;
-	for (const FDMXFixtureCellAttribute& Attribute : CellAttributes)
+	// DEPRECATED 5.0
+	int32 Index = 0;
+	FDMXFixtureFunction FunctionToAdd = InFunction;
+
+	// Shift the insert function channel
+	uint8 DataTypeBytes = FDMXConversions::GetSizeOfSignalFormat(InFunction.DataType);
+	FunctionToAdd.Channel = InFunction.Channel + DataTypeBytes;
+
+	if (Functions.IsValidIndex(IndexOfFunction + 1))
+	{
+		Index = Functions.Insert(InFunction, IndexOfFunction + 1);
+
+		// Shift all function after this by size of insert function
+		for (int32 FunctionIndex = Index + 1; FunctionIndex < Functions.Num(); FunctionIndex++)
+		{
+			FDMXFixtureFunction& Function = Functions[FunctionIndex];
+			Function.Channel = Function.Channel + DataTypeBytes;
+		}
+	}
+	else
 	{
-		int32 CurrentFunctionSize = UDMXEntityFixtureType::NumChannelsToOccupy(Attribute.DataType);
-		CellAttributeSize += CurrentFunctionSize;
+		InFunction.Channel = Functions.Num() > 0 ? Functions.Last().Channel + Functions.Last().GetNumChannels() : 1;
+		Index = Functions.Add(InFunction);
 	}
-	int32 AllCells = XCells * YCells * CellAttributeSize;
-	if (AllCells == 0)
+
+	return Index;
+}
+#endif // WITH_EDITOR
+
+
+FDMXOnFixtureTypeChangedDelegate UDMXEntityFixtureType::OnFixtureTypeChangedDelegate;
+
+void UDMXEntityFixtureType::Serialize(FArchive& Ar)
+{
+	Super::Serialize(Ar);
+
+	Ar.UsingCustomVersion(FDMXRuntimeObjectVersion::GUID);
+	if (Ar.IsLoading())
 	{
-		return FirstCellChannel;
+		if (Ar.CustomVer(FDMXRuntimeObjectVersion::GUID) < FDMXRuntimeObjectVersion::DMXFixtureTypeAllowMatrixInEachFixtureMode)
+		{
+			// For assets that were created before each mode could enable or disable the matrix, copy the deprecated bFixtureMatrixEnabled property to each mode
+			for (FDMXFixtureMode& Mode : Modes)
+			{
+				PRAGMA_DISABLE_DEPRECATION_WARNINGS
+				Mode.bFixtureMatrixEnabled = bFixtureMatrixEnabled;
+				PRAGMA_ENABLE_DEPRECATION_WARNINGS
+			}
+		}
+
+		if (Ar.CustomVer(FDMXRuntimeObjectVersion::GUID) < FDMXRuntimeObjectVersion::DMXFixtureTypeAllowInputModulatorsInEachFixtureMode)
+		{
+			// For assets that were created before each mode had its own Input Modulators, copy the deprecated InputModulator property to each mode
+			for (FDMXFixtureMode& Mode : Modes)
+			{
+				PRAGMA_DISABLE_DEPRECATION_WARNINGS
+				Mode.InputModulators = InputModulators;
+				PRAGMA_ENABLE_DEPRECATION_WARNINGS
+			}
+		}
+	}
+}
+
+#if WITH_EDITOR
+bool UDMXEntityFixtureType::Modify(bool bAlwaysMarkDirty)
+{
+	if (UDMXLibrary* DMXLibrary = ParentLibrary.Get())
+	{
+		return DMXLibrary->Modify(bAlwaysMarkDirty) && Super::Modify(bAlwaysMarkDirty);
 	}
 
-	int32 LastChannel = FirstCellChannel + AllCells - 1;
+	return Super::Modify(bAlwaysMarkDirty);
+}
+#endif // WITH_EDITOR
+
+#if WITH_EDITOR
+void UDMXEntityFixtureType::PostEditChangeProperty(FPropertyChangedEvent& PropertyChangedEvent)
+{
+	Super::PostEditChangeProperty(PropertyChangedEvent);
 
-	return LastChannel;
+	if (PropertyChangedEvent.ChangeType != EPropertyChangeType::Interactive)
+	{
+		OnFixtureTypeChangedDelegate.Broadcast(this);
+	}
 }
+#endif // WITH_EDITOR
 
 #if WITH_EDITOR
-	/** Editor only data type change delegate */
-FDataTypeChangeDelegate UDMXEntityFixtureType::DataTypeChangeDelegate;
-#endif
+void UDMXEntityFixtureType::PostEditChangeChainProperty(FPropertyChangedChainEvent& PropertyChangedChainEvent)
+{
+	Super::PostEditChangeChainProperty(PropertyChangedChainEvent);
+	const FName PropertyName = PropertyChangedChainEvent.GetPropertyName();
+
+	// Keep Fixture Patches' ActiveMode value valid
+	if (PropertyName == GET_MEMBER_NAME_CHECKED(UDMXEntityFixtureType, Modes))
+	{
+		int32 ChangedModeIndex = PropertyChangedChainEvent.GetArrayIndex(PropertyName.ToString());
+		if (Modes.IsValidIndex(ChangedModeIndex))
+		{
+			// Notify DataType changes (Deprecated 5.0)
+			PRAGMA_DISABLE_DEPRECATION_WARNINGS
+			DataTypeChangeDelegate_DEPRECATED.Broadcast(this, Modes[ChangedModeIndex]);
+			PRAGMA_ENABLE_DEPRECATION_WARNINGS
+		}
+	}
+
+	if (PropertyChangedChainEvent.ChangeType != EPropertyChangeType::Interactive)
+	{
+		OnFixtureTypeChangedDelegate.Broadcast(this);
+	}
+}
+#endif // WITH_EDITOR
+
+#if WITH_EDITOR
+void UDMXEntityFixtureType::PostEditUndo()
+{
+	Super::PostEditUndo();
+
+	OnFixtureTypeChangedDelegate.Broadcast(this);
+}
+#endif // WITH_EDITOR
 
 #if WITH_EDITOR
 void UDMXEntityFixtureType::SetModesFromDMXImport(UDMXImport* DMXImportAsset)
@@ -160,10 +297,6 @@ void UDMXEntityFixtureType::SetModesFromDMXImport(UDMXImport* DMXImportAsset)
 			FDMXFixtureMode& Mode = Modes[Modes.Emplace()];
 			Mode.ModeName = AssetMode.Name.ToString();
 
-			// We'll keep the functions addresses from the GDTF file.
-			// For that we need to keep track of the latest occupied address after adding each function.
-			int32 LastOccupiedAddress = 0;
-
 			// Keep track of the Attributes used on this Mode's Functions because they must be unique
 			TArray<FName> MappedAttributes;
 
@@ -228,14 +361,7 @@ void UDMXEntityFixtureType::SetModesFromDMXImport(UDMXImport* DMXImportAsset)
 					}
 					const int32 NumUsedAddresses = FMath::Clamp(AddressMax - AddressMin + 1, 1, DMX_MAX_FUNCTION_SIZE);
 
-					SetFunctionSize(Function, NumUsedAddresses);
-
-					// AddressMin is the first address this function occupies. If it's not 1 after the
-					// latest occupied channel, this function is offset, skipping some addresses.
-					if (AddressMin > LastOccupiedAddress + 1)
-					{
-						Function.ChannelOffset = AddressMin - LastOccupiedAddress - 1;
-					}
+					Function.DataType = static_cast<EDMXFixtureSignalFormat>(NumUsedAddresses - 1);
 
 					// Offsets represent the value bytes in MSB format. If they are in reverse order,
 					// it means this Function uses LSB format.
@@ -249,142 +375,642 @@ void UDMXEntityFixtureType::SetModesFromDMXImport(UDMXImport* DMXImportAsset)
 					{
 						Function.bUseLSBMode = false;
 					}
-
-					// Update occupied addresses on current mode
-					LastOccupiedAddress += Function.ChannelOffset + NumChannelsToOccupy(Function.DataType);
 				}
 				else
 				{
-					SetFunctionSize(Function, 1);
-
-					// Update occupied addresses
-					++LastOccupiedAddress;
+					Function.DataType = EDMXFixtureSignalFormat::E8Bit;
 				}
 			}
+		}
 
-			// Compute mode channel span from functions' addresses and sizes
-			UpdateChannelSpan(Mode);
+		// Update the Channel Span for all Modes
+		for (int32 ModeIndex = 0; ModeIndex < Modes.Num(); ModeIndex++)
+		{
+			UpdateChannelSpan(ModeIndex);
 		}
 	}
+
+	OnFixtureTypeChangedDelegate.Broadcast(this);
 }
+#endif // WITH_EDITOR
 
-void UDMXEntityFixtureType::SetFunctionSize(FDMXFixtureFunction& InFunction, uint8 Size)
+FDMXOnFixtureTypeChangedDelegate& UDMXEntityFixtureType::GetOnFixtureTypeChanged()
 {
-	//Get New Data Type
-	EDMXFixtureSignalFormat NewDataType;
-	switch (Size)
+	return OnFixtureTypeChangedDelegate;
+}
+
+int32 UDMXEntityFixtureType::AddMode()
+{
+	FDMXFixtureMode NewMode;
+
+	// Make a unique name for the new mode
+	TSet<FString> ModeNames;
+	for (const FDMXFixtureMode& Mode : Modes)
 	{
-	case 0:
-	case 1:
-		NewDataType = EDMXFixtureSignalFormat::E8Bit;
-		break;
-	case 2:
-		NewDataType = EDMXFixtureSignalFormat::E16Bit;
-		break;
-	case 3:
-		NewDataType = EDMXFixtureSignalFormat::E24Bit;
-		break;
-	case 4:
-	default:
-		NewDataType = EDMXFixtureSignalFormat::E32Bit;
-		break;
+		ModeNames.Add(Mode.ModeName);
 	}
+	NewMode.ModeName = FDMXRuntimeUtils::GenerateUniqueNameFromExisting(ModeNames, LOCTEXT("DefaultModeName", "Mode").ToString());
 
-	InFunction.DataType = NewDataType;
-	ClampDefaultValue(InFunction);
-}
-
-#endif // WITH_EDITOR
+	const int32 NewModeIndex = Modes.Add(NewMode);
 
-uint8 UDMXEntityFixtureType::GetFunctionLastChannel(const FDMXFixtureFunction& Function)
-{
-	return Function.Channel + UDMXEntityFixtureType::NumChannelsToOccupy(Function.DataType) - 1;
+	return NewModeIndex;
 }
 
-bool UDMXEntityFixtureType::IsFunctionInModeRange(const FDMXFixtureFunction& InFunction, const FDMXFixtureMode& InMode, int32 ChannelOffset /*= 0*/)
+void UDMXEntityFixtureType::DuplicateModes(TArray<int32> InModeÎndiciesToDuplicate, TArray<int32>& OutNewModeIndices)
 {
-	const int32 LastChannel = GetFunctionLastChannel(InFunction);
-	bool bLastChannelExceedsChannelSpan = LastChannel > InMode.ChannelSpan;
-	bool bLastChannelExceedsUniverseSize = LastChannel + ChannelOffset > DMX_MAX_ADDRESS;
+	int32 NumModesDuplicated = 0;
+	for (int32 ModeIndex : InModeÎndiciesToDuplicate)
+	{
+		if (Modes.IsValidIndex(ModeIndex))
+		{
+			// Copy
+			FDMXFixtureMode NewMode = Modes[ModeIndex];
 
-	return !bLastChannelExceedsChannelSpan && !bLastChannelExceedsUniverseSize;
+			int32 IndexOfDuplicate = ModeIndex + 1 + NumModesDuplicated;
+			if (Modes.IsValidIndex(IndexOfDuplicate))
+			{
+				IndexOfDuplicate = Modes.Insert(NewMode, IndexOfDuplicate);
+			}
+			else
+			{
+				IndexOfDuplicate = Modes.Add(NewMode);
+			}
+
+			OutNewModeIndices.Add(IndexOfDuplicate);
+
+			// Make a uinque mode name by setting it
+			FString UnusedString;
+			SetModeName(IndexOfDuplicate, NewMode.ModeName, UnusedString);
+
+			NumModesDuplicated++;
+		}
+	}
 }
 
-bool UDMXEntityFixtureType::IsFixtureMatrixInModeRange(const FDMXFixtureMatrix& InFixtureMatrix, const FDMXFixtureMode& InMode, int32 ChannelOffset /*= 0*/)
+void UDMXEntityFixtureType::RemoveModes(const TArray<int32>& ModeIndicesToDelete)
 {
-	const int32 LastChannel = InFixtureMatrix.GetFixtureMatrixLastChannel();
-	bool bLastChannelExceedsChannelSpan = LastChannel > InMode.ChannelSpan;
-	bool bLastChannelExceedsUniverseSize = LastChannel + ChannelOffset > DMX_MAX_ADDRESS;
-
-	return !bLastChannelExceedsChannelSpan && !bLastChannelExceedsUniverseSize;
+	int32 NumDeletedModes = 0;
+	for (const int32 ModeIndex : ModeIndicesToDelete)
+	{
+		const int32 DeleteAtIndex = ModeIndex - NumDeletedModes;
+		if (Modes.IsValidIndex(DeleteAtIndex))
+		{
+			Modes.RemoveAt(DeleteAtIndex);
+			NumDeletedModes++;
+		}
+	}
 }
 
-void UDMXEntityFixtureType::ClampDefaultValue(FDMXFixtureFunction& InFunction)
+void UDMXEntityFixtureType::SetModeName(int32 InModeIndex, const FString& InDesiredModeName, FString& OutUniqueModeName)
 {
-	InFunction.DefaultValue = ClampValueToDataType(InFunction.DataType, FMath::Min(InFunction.DefaultValue, (int64)MAX_uint32));
+	if (ensureMsgf(Modes.IsValidIndex(InModeIndex), TEXT("Invalid Mode Index when setting Mode Name.")))
+	{
+		FDMXFixtureMode& Mode = Modes[InModeIndex];
+
+		FString OldName = Mode.ModeName;
+
+		const FString UniqueModeNameString = [InModeIndex, InDesiredModeName, this]()
+		{
+			TSet<FString> OtherModeNames;
+			for (int32 OtherModeIndex = 0; OtherModeIndex < Modes.Num(); OtherModeIndex++)
+			{
+				if (OtherModeIndex != InModeIndex)
+				{
+					OtherModeNames.Add(Modes[OtherModeIndex].ModeName);
+				}
+			}
+			return FDMXRuntimeUtils::GenerateUniqueNameFromExisting(OtherModeNames, InDesiredModeName);
+		}();
+
+		Mode.ModeName = UniqueModeNameString;
+		OutUniqueModeName = UniqueModeNameString;
+	}
 }
 
-uint8 UDMXEntityFixtureType::NumChannelsToOccupy(EDMXFixtureSignalFormat DataType)
+void UDMXEntityFixtureType::SetFixtureMatrixEnabled(int32 ModeIndex, bool bEnableMatrix)
 {
-	switch (DataType)
+	if (ensureMsgf(Modes.IsValidIndex(ModeIndex), TEXT("Trying Update Channel Span, but Mode Index is not valid.")))
 	{
-	case EDMXFixtureSignalFormat::E8Bit:
-		return 1;
+		FDMXFixtureMode& Mode = Modes[ModeIndex];
+		FDMXFixtureMatrix& Matrix = Mode.FixtureMatrixConfig;
 
-	case EDMXFixtureSignalFormat::E16Bit:
-		return 2;
+		if (bEnableMatrix != Mode.bFixtureMatrixEnabled)
+		{
+			Mode.bFixtureMatrixEnabled = bEnableMatrix;
+			
+			// Realign the Functions, only if the Matrix has at least one Channel
+			const int32 NumMatrixChannels = Matrix.GetNumChannels();
+			if (NumMatrixChannels > 0)
+			{
+				const int32 IndexOfFirstFunctionOverMatrix = Mode.Functions.IndexOfByPredicate([Matrix](const FDMXFixtureFunction& Function)
+					{
+						return Function.Channel >= Matrix.FirstCellChannel;
+					});
 
-	case EDMXFixtureSignalFormat::E24Bit:
-		return 3;
+				// Remember the old, set the new Matrix starting channel
+				const int32 OldMatrixLastChannel = Matrix.GetLastChannel();
+				Matrix.FirstCellChannel = [Mode, IndexOfFirstFunctionOverMatrix]()
+				{
+					const int32 IndexOfLastFunctionBeforeMatrix = IndexOfFirstFunctionOverMatrix - 1;
+					if (Mode.Functions.IsValidIndex(IndexOfLastFunctionBeforeMatrix))
+					{
+						return Mode.Functions[IndexOfLastFunctionBeforeMatrix].GetLastChannel() + 1;
+					}
 
-	case EDMXFixtureSignalFormat::E32Bit:
-		return 4;
+					return 1;
+				}();
 
-	default:
-		// Unhandled type
-		checkNoEntry();
-		break;
+				// Align Functions
+				const int32 MatrixDeltaChannels = Matrix.GetLastChannel() - OldMatrixLastChannel;
+				const int32 FunctionDeltaChannels = bEnableMatrix ?
+					NumMatrixChannels + MatrixDeltaChannels :
+					-(NumMatrixChannels + MatrixDeltaChannels);
+
+				for (FDMXFixtureFunction& Function : Mode.Functions)
+				{
+					if (Function.Channel >= Matrix.FirstCellChannel)
+					{
+						Function.Channel += FunctionDeltaChannels;
+					}
+				}
+			}
+		}
 	}
-	return 1;
 }
 
-uint32 UDMXEntityFixtureType::ClampValueToDataType(EDMXFixtureSignalFormat DataType, uint32 InValue)
+void UDMXEntityFixtureType::UpdateChannelSpan(int32 ModeIndex)
 {
-	switch (DataType)
+	if (ensureMsgf(Modes.IsValidIndex(ModeIndex), TEXT("Trying Update Channel Span, but Mode Index is not valid.")))
 	{
-	case EDMXFixtureSignalFormat::E8Bit:
-		return FMath::Clamp(InValue, 0u, (uint32)MAX_uint8);
+		FDMXFixtureMode& Mode = Modes[ModeIndex];
 
-	case EDMXFixtureSignalFormat::E16Bit:
-		return FMath::Clamp(InValue, 0u, (uint32)MAX_uint16);
+		if (Mode.bAutoChannelSpan)
+		{
+			int32 LowestChannel = TNumericLimits<int32>::Max();
+			int32 HighestChannel = TNumericLimits<int32>::Min();
 
-	case EDMXFixtureSignalFormat::E24Bit:
-		return FMath::Clamp(InValue, 0u, 0xFFFFFFu);
+			for (const FDMXFixtureFunction& Function : Mode.Functions)
+			{
+				const int32 FirstChannelOfFunction = Function.Channel;
+				LowestChannel = FirstChannelOfFunction < LowestChannel ? FirstChannelOfFunction : LowestChannel;
 
-	case EDMXFixtureSignalFormat::E32Bit:
-		return FMath::Clamp(InValue, 0u, MAX_uint32);
+				const int32 LastChannelOfFunction = FirstChannelOfFunction + Function.GetNumChannels() - 1;
+				HighestChannel = LastChannelOfFunction > HighestChannel ? LastChannelOfFunction : HighestChannel;
+			}
 
-	default:
-		break;
+			const int32 FirstChannelOfMatrix = Mode.FixtureMatrixConfig.FirstCellChannel;
+			LowestChannel = FirstChannelOfMatrix < LowestChannel ? FirstChannelOfMatrix : LowestChannel;
+
+			const int32 LastChannelOfMatrix = FirstChannelOfMatrix + Mode.FixtureMatrixConfig.GetNumChannels() - 1;
+			HighestChannel = LastChannelOfMatrix > HighestChannel ? LastChannelOfMatrix : HighestChannel;
+
+			Mode.ChannelSpan = FMath::Max(HighestChannel - LowestChannel + 1, 0);
+		}
 	}
-	return InValue;
 }
 
-uint32 UDMXEntityFixtureType::GetDataTypeMaxValue(EDMXFixtureSignalFormat DataType)
+int32 UDMXEntityFixtureType::AddFunction(int32 InModeIndex)
 {
-	switch (DataType)
+	if (ensureMsgf(Modes.IsValidIndex(InModeIndex), TEXT("Invalid Mode Index when setting Mode Name.")))
 	{
-	case EDMXFixtureSignalFormat::E8Bit:
-		return MAX_uint8;
-	case EDMXFixtureSignalFormat::E16Bit:
-		return MAX_uint16;
-	case EDMXFixtureSignalFormat::E24Bit:
-		return 0xFFFFFF;
-	case EDMXFixtureSignalFormat::E32Bit:
-		return MAX_uint32;
-	default:
-		checkNoEntry();
-		return 1;
+		FDMXFixtureMode& Mode = Modes[InModeIndex];
+		FDMXFixtureFunction NewFunction;
+
+		// Set a unique Name
+		TSet<FString> FunctionNames;
+		for (const FDMXFixtureFunction& Function : Mode.Functions)
+		{
+			FunctionNames.Add(Function.FunctionName);
+		}
+		NewFunction.FunctionName = FDMXRuntimeUtils::GenerateUniqueNameFromExisting(FunctionNames, LOCTEXT("DMXFixtureTypeSharedData.DefaultFunctionName", "Function").ToString());
+
+		// Update the Channel
+		NewFunction.Channel = [Mode]()
+		{
+			const int32 LastFunctionChannel = Mode.Functions.Num() > 0 ? Mode.Functions.Last().GetLastChannel() : 0;
+			const int32 LastMatrixChannel = Mode.bFixtureMatrixEnabled ? Mode.FixtureMatrixConfig.GetLastChannel() : 0;
+			return FMath::Max(LastFunctionChannel, LastMatrixChannel) + 1;
+		}();
+
+		// Add the Function and update Channel Span
+		const int32 NewFunctionIndex = Mode.Functions.Add(NewFunction);
+		UpdateChannelSpan(InModeIndex);
+
+		return NewFunctionIndex;
+	}
+
+	return INDEX_NONE;
+}
+
+int32 UDMXEntityFixtureType::InsertFunction(int32 InModeIndex, int32 InInsertAtIndex, FDMXFixtureFunction& InOutNewFunction)
+{
+	if (ensureMsgf(Modes.IsValidIndex(InModeIndex), TEXT("Invalid Mode Index when setting Mode Name.")))
+	{
+		FDMXFixtureMode& Mode = Modes[InModeIndex];
+
+		// Set a unique name
+		TSet<FString> FunctionNames;
+		for (const FDMXFixtureFunction& Function : Mode.Functions)
+		{
+			FunctionNames.Add(Function.FunctionName);
+		}
+		InOutNewFunction.FunctionName = FDMXRuntimeUtils::GenerateUniqueNameFromExisting(FunctionNames, InOutNewFunction.FunctionName);
+
+		int32 NewFunctionIndex = INDEX_NONE;
+		if(Mode.Functions.IsValidIndex(InInsertAtIndex))
+		{
+			// Add the Function, then reorder it to where it should reside
+			InOutNewFunction.Channel = TNumericLimits<int32>::Max();
+			const int32 TempFunctionIndex = Mode.Functions.Add(InOutNewFunction);
+
+			ReorderFunction(InModeIndex, TempFunctionIndex, InInsertAtIndex);
+			NewFunctionIndex = InInsertAtIndex;
+		}
+		else
+		{
+			// Add the Function after either the last Function or the last Matrix channel
+			InOutNewFunction.Channel = [Mode]()
+			{
+				const int32 LastFunctionChannel = Mode.Functions.Num() > 0 ? Mode.Functions.Last().GetLastChannel() : 0;
+				const int32 LastMatrixChannel = Mode.bFixtureMatrixEnabled ? Mode.FixtureMatrixConfig.GetLastChannel() : 0;
+				return FMath::Max(LastFunctionChannel, LastMatrixChannel) + 1;
+			}();
+
+			NewFunctionIndex = Mode.Functions.Add(InOutNewFunction);
+		}
+
+		UpdateChannelSpan(InModeIndex);
+
+		check(NewFunctionIndex != INDEX_NONE);
+		return NewFunctionIndex;
+	}
+
+	return INDEX_NONE;
+}
+
+void UDMXEntityFixtureType::DuplicateFunctions(int32 InModeIndex, const TArray<int32>& InFunctionIndicesToDuplicate, TArray<int32>& OutNewFunctionIndices)
+{
+	if (Modes.IsValidIndex(InModeIndex))
+	{
+		FDMXFixtureMode& Mode = Modes[InModeIndex];
+
+		int32 NumDuplicatedFunctions = 0;
+		for (const int32 FunctionToDuplicateIndex : InFunctionIndicesToDuplicate)
+		{
+			if (ensureMsgf(Mode.Functions.IsValidIndex(FunctionToDuplicateIndex), TEXT("Trying to duplicate Function, but Function index is not valid.")))
+			{
+				FDMXFixtureFunction DuplicatedFunction = Mode.Functions[FunctionToDuplicateIndex];
+
+				// Offset subsequent functions and possibly the matrix 
+				DuplicatedFunction.Channel = DuplicatedFunction.GetLastChannel() + 1;
+				const int32 Offset = DuplicatedFunction.GetNumChannels();
+				for (FDMXFixtureFunction& Function : Mode.Functions)
+				{
+					if (Function.Channel >= DuplicatedFunction.Channel)
+					{
+						Function.Channel += Offset;
+					}
+				}
+				if (Mode.FixtureMatrixConfig.FirstCellChannel >= DuplicatedFunction.Channel)
+				{
+					Mode.FixtureMatrixConfig.FirstCellChannel += Offset;
+				}
+
+				// Add the Function
+				int32 IndexOfNewlyAddedFunction = FunctionToDuplicateIndex + 1 + NumDuplicatedFunctions;
+				if (Mode.Functions.IsValidIndex(IndexOfNewlyAddedFunction))
+				{
+					IndexOfNewlyAddedFunction = Mode.Functions.Insert(DuplicatedFunction, IndexOfNewlyAddedFunction);
+				}
+				else
+				{
+					IndexOfNewlyAddedFunction = Mode.Functions.Add(DuplicatedFunction);
+				}
+				OutNewFunctionIndices.Add(IndexOfNewlyAddedFunction);
+
+				// Make a uinque Function Name by setting it
+				FString UnusedString;
+				SetFunctionName(InModeIndex, IndexOfNewlyAddedFunction, DuplicatedFunction.FunctionName, UnusedString);
+
+				NumDuplicatedFunctions++;
+			}
+		}
+
+		if (NumDuplicatedFunctions > 0)
+		{
+			UpdateChannelSpan(InModeIndex);
+		}
+	}		
+}
+
+void UDMXEntityFixtureType::RemoveFunctions(int32 ModeIndex, TArray<int32> FunctionIndicesToDelete)
+{
+	if (ensureMsgf(Modes.IsValidIndex(ModeIndex), TEXT("Trying to remove Function, but Mode index is invalid.")))
+	{
+		FDMXFixtureMode& Mode = Modes[ModeIndex];
+
+		// Sort to remove from last to first index
+		FunctionIndicesToDelete.Sort([](int32 IndexToDeleteA, int32 IndexToDeleteB)
+			{
+				return IndexToDeleteA > IndexToDeleteB;
+			});
+
+		for (const int32 FunctionToRemoveIndex : FunctionIndicesToDelete)
+		{
+			if (ensureMsgf(Mode.Functions.IsValidIndex(FunctionToRemoveIndex), TEXT("Trying to remove Function, but Function index is invalid.")))
+			{
+				const FDMXFixtureFunction& FunctionToRemove = Mode.Functions[FunctionToRemoveIndex];
+				const int32 Offset = FunctionToRemove.GetNumChannels();
+
+				// Realign subsequent functions
+				const int32 RealignAtIndex = FunctionToRemoveIndex + 1;
+				if (Mode.Functions.IsValidIndex(RealignAtIndex))
+				{
+					for (int32 FunctionToReorderIndex = RealignAtIndex; FunctionToReorderIndex < Mode.Functions.Num(); FunctionToReorderIndex++)
+					{
+						Mode.Functions[FunctionToReorderIndex].Channel -= Offset;
+					}
+				}
+
+				// Align the Matrix
+				if (Mode.bFixtureMatrixEnabled)
+				{
+					const int32 LastRemovedFunctionChannel = FunctionToRemove.GetLastChannel();
+					const int32 FirstMatrixChannel = Mode.FixtureMatrixConfig.FirstCellChannel;
+
+					if (FirstMatrixChannel >= LastRemovedFunctionChannel)
+					{
+						Mode.FixtureMatrixConfig.FirstCellChannel -= Offset;
+					}
+				}
+
+				Mode.Functions.RemoveAt(FunctionToRemoveIndex);
+			}
+		}
+
+		UpdateChannelSpan(ModeIndex);
+	}
+}
+
+void UDMXEntityFixtureType::ReorderFunction(int32 ModeIndex, int32 FunctionToReorderIndex, int32 InsertAtIndex)
+{
+	if (ensureMsgf(Modes.IsValidIndex(ModeIndex), TEXT("Trying to reorder Function, but Mode Index is not valid.")))
+	{
+		FDMXFixtureMode& Mode = Modes[ModeIndex];
+		FDMXFixtureMatrix& Matrix = Mode.FixtureMatrixConfig;
+
+		if (ensureMsgf(Mode.Functions.IsValidIndex(FunctionToReorderIndex) && Mode.Functions.IsValidIndex(InsertAtIndex), TEXT("Trying to reorder Function, but Function Indices are not valid.")))
+		{
+			if (FunctionToReorderIndex != InsertAtIndex)
+			{
+				FDMXFixtureFunction FunctionToReorder = Mode.Functions[FunctionToReorderIndex];
+				const FDMXFixtureFunction& InsertAtFunction = Mode.Functions[InsertAtIndex];
+
+				const bool bForward = FunctionToReorderIndex < InsertAtIndex;
+				const int32 Offset = bForward ? -FunctionToReorder.GetNumChannels() : FunctionToReorder.GetNumChannels();
+
+				// Align Matrix if enabled
+				if (Mode.bFixtureMatrixEnabled)
+				{
+					const bool bInRangeOfMatrixStart = bForward ?
+						Matrix.FirstCellChannel >= FunctionToReorder.Channel :
+						Matrix.FirstCellChannel >= InsertAtFunction.Channel;
+
+					const bool bInRangeOfMatrixEnd = bForward ?
+						Matrix.FirstCellChannel <= InsertAtFunction.GetLastChannel() :
+						Matrix.FirstCellChannel <= FunctionToReorder.GetLastChannel();
+
+					if (bInRangeOfMatrixStart && bInRangeOfMatrixEnd)
+					{
+						Matrix.FirstCellChannel += Offset;
+					}
+				}
+
+				// Set the Function Channel
+				FunctionToReorder.Channel = bForward ? InsertAtFunction.Channel + InsertAtFunction.GetNumChannels() - FunctionToReorder.GetNumChannels() : InsertAtFunction.Channel;
+
+				// Align affected Functions
+				const int32 StartIndex = bForward ? FunctionToReorderIndex + 1 : InsertAtIndex;
+				const int32 EndIndex = bForward ? InsertAtIndex : FunctionToReorderIndex - 1;
+				for (int32 FunctionIndex = StartIndex; FunctionIndex <= EndIndex; FunctionIndex++)
+				{
+					if (Mode.Functions.IsValidIndex(FunctionIndex))
+					{
+						FDMXFixtureFunction& Function = Mode.Functions[FunctionIndex];
+						Function.Channel += Offset;
+					}
+				}
+
+				// Update Functions array order
+				if (bForward && Mode.Functions.IsValidIndex(InsertAtIndex + 1))
+				{
+					Mode.Functions.Insert(FunctionToReorder, InsertAtIndex + 1);
+				}
+				else if (bForward)
+				{
+					Mode.Functions.Add(FunctionToReorder);
+				}
+				else
+				{
+					Mode.Functions.Insert(FunctionToReorder, InsertAtIndex);
+				}
+				Mode.Functions.RemoveAt(bForward ? FunctionToReorderIndex : FunctionToReorderIndex + 1);
+
+			}
+		}
+	}
+}
+
+void UDMXEntityFixtureType::SetFunctionName(int32 InModeIndex, int32 InFunctionIndex, const FString& InDesiredFunctionName, FString& OutUniqueFunctionName)
+{
+	if (ensureMsgf(Modes.IsValidIndex(InModeIndex) && Modes[InModeIndex].Functions.IsValidIndex(InFunctionIndex), TEXT("Invalid Mode Index or Function Index when setting Function Name.")))
+	{
+		FDMXFixtureMode& Mode = Modes[InModeIndex];
+		FDMXFixtureFunction& Function = Mode.Functions[InFunctionIndex];
+
+		FString OldName = Function.FunctionName;
+
+		const FString UniqueFunctionNameString = [InModeIndex, InFunctionIndex, InDesiredFunctionName, this]()
+		{
+			TSet<FString> OtherFunctionNames;
+			for (int32 OtherFunctionIndex = 0; OtherFunctionIndex < Modes[InModeIndex].Functions.Num(); OtherFunctionIndex++)
+			{
+				if (OtherFunctionIndex != InFunctionIndex)
+				{
+					OtherFunctionNames.Add(Modes[InModeIndex].Functions[OtherFunctionIndex].FunctionName);
+				}
+			}
+			return FDMXRuntimeUtils::GenerateUniqueNameFromExisting(OtherFunctionNames, InDesiredFunctionName);
+		}();
+
+		Function.FunctionName = UniqueFunctionNameString;
+		OutUniqueFunctionName = UniqueFunctionNameString;
+	}
+}
+
+void UDMXEntityFixtureType::SetFunctionStartingChannel(int32 InModeIndex, int32 InFunctionIndex, int32 InDesiredStartingChannel, int32& OutStartingChannel)
+{
+	if (ensureMsgf(Modes.IsValidIndex(InModeIndex) && Modes[InModeIndex].Functions.IsValidIndex(InFunctionIndex), TEXT("Invalid Mode Index or Function Index when setting Function Starting Channel.")))
+	{
+		FDMXFixtureMode& Mode = Modes[InModeIndex];
+		FDMXFixtureFunction& Function = Mode.Functions[InFunctionIndex];
+
+		const int32 FunctionOverNewChannelIndex = Mode.Functions.IndexOfByPredicate([InDesiredStartingChannel](const FDMXFixtureFunction& Function)
+			{
+				return
+					Function.Channel >= InDesiredStartingChannel &&
+					Function.GetLastChannel() <= InDesiredStartingChannel;
+			});
+
+		// Find the index where the Function should be moved to. If not a valid index, the Function is already at the right index.
+		const int32 DesiredIndex = [&]()
+		{
+			if (FunctionOverNewChannelIndex != INDEX_NONE)
+			{
+				return FunctionOverNewChannelIndex;
+			}
+			else
+			{
+				const bool bDesiredIndexAfterLastFunction =
+					Mode.Functions.Num() > 0 &&
+					Mode.Functions.Last().GetLastChannel() < InDesiredStartingChannel;
+
+				return bDesiredIndexAfterLastFunction ? Mode.Functions.Num() - 1 : INDEX_NONE;
+			}
+		}();
+
+		if (Mode.Functions.IsValidIndex(DesiredIndex))
+		{
+			ReorderFunction(InModeIndex, InFunctionIndex, DesiredIndex);
+		}
+	}
+}
+
+void UDMXEntityFixtureType::ClampFunctionDefautValueByDataType(int32 ModeIndex, int32 FunctionToRemoveIndex)
+{
+	if (ensureMsgf(Modes.IsValidIndex(ModeIndex), TEXT("Trying to clamp Function Default Value, but Mode Index is not valid.")))
+	{
+		FDMXFixtureMode& Mode = Modes[ModeIndex];
+
+		if (ensureMsgf(Mode.Functions.IsValidIndex(FunctionToRemoveIndex), TEXT("Trying to clamp Function Default Value, but Function Index is not valid.")))
+		{
+			FDMXFixtureFunction& Function = Mode.Functions[FunctionToRemoveIndex];
+			const uint32 SafeDefaultValue = FMath::Min(static_cast<int64>(TNumericLimits<uint32>::Max()), Function.DefaultValue);
+			const uint32 ClampedDefaultValue = FDMXConversions::ClampValueBySignalFormat(SafeDefaultValue, Function.DataType);
+
+			if (Function.DefaultValue != ClampedDefaultValue)
+			{
+				Function.DefaultValue = ClampedDefaultValue;
+			}
+		}
+	}
+}
+
+void UDMXEntityFixtureType::AddCellAttribute(int32 ModeIndex)
+{
+	if (Modes.IsValidIndex(ModeIndex))
+	{
+		FDMXFixtureMode& Mode = Modes[ModeIndex];
+		FDMXFixtureMatrix& Matrix = Mode.FixtureMatrixConfig;
+
+		FDMXFixtureCellAttribute NewAttribute;
+		TArray<FName> AvailableAttributes = FDMXAttributeName::GetPossibleValues();
+		NewAttribute.Attribute = AvailableAttributes.Num() > 0 ? FDMXAttributeName(AvailableAttributes[0]) : FDMXAttributeName();
+
+		// Disable the matrix while editing it so other functions align when enabling it again
+		SetFixtureMatrixEnabled(ModeIndex, false);
+
+		Mode.FixtureMatrixConfig.CellAttributes.Add(NewAttribute);
+
+		SetFixtureMatrixEnabled(ModeIndex, true);
+		UpdateChannelSpan(ModeIndex);
+	}
+}
+
+void UDMXEntityFixtureType::ReorderMatrix(int32 ModeIndex, int32 InsertAtFunctionIndex)
+{
+	if (ensureMsgf(Modes.IsValidIndex(ModeIndex), TEXT("Trying to reorder Matrix after Function, but Mode Index is not valid.")))
+	{
+		FDMXFixtureMode& Mode = Modes[ModeIndex];
+		FDMXFixtureMatrix& Matrix = Mode.FixtureMatrixConfig;
+
+		if (ensureMsgf(Mode.bFixtureMatrixEnabled, TEXT("Trying to reorder Matrix after Function, but Matrix is not enabled or Function Indices are not valid.")))
+		{
+			// Correct the insert at function index depending on direction
+			const bool bReorderMatrixForward = Mode.Functions[InsertAtFunctionIndex].Channel < Matrix.FirstCellChannel;
+			InsertAtFunctionIndex = bReorderMatrixForward ? InsertAtFunctionIndex : InsertAtFunctionIndex + 1;
+
+			// Reorder channels as if the matrix was removed
+			for (FDMXFixtureFunction& Function : Mode.Functions)
+			{
+				if (Function.Channel >= Matrix.FirstCellChannel)
+				{
+					Function.Channel -= Mode.FixtureMatrixConfig.GetNumChannels();
+				}
+			}
+
+			// Set the new Matrix starting channel
+			if (Mode.Functions.IsValidIndex(InsertAtFunctionIndex))
+			{
+				const FDMXFixtureFunction& InsertAtFunction = Mode.Functions[InsertAtFunctionIndex];
+				Matrix.FirstCellChannel = InsertAtFunction.Channel;
+			}
+			else if (Mode.Functions.Num() > 0)
+			{
+				Matrix.FirstCellChannel = Mode.Functions.Last().Channel + Mode.Functions.Last().GetNumChannels();
+			}
+
+			// Reorder channels that now conflict to after the matrix
+			for (FDMXFixtureFunction& Function : Mode.Functions)
+			{
+				if (Function.Channel >= Matrix.FirstCellChannel)
+				{
+					Function.Channel += Mode.FixtureMatrixConfig.GetNumChannels();
+				}
+			}
+		}
+	}
+}
+
+void UDMXEntityFixtureType::UpdateYCellsFromXCells(int32 ModeIndex)
+{
+	if (ensureMsgf(Modes.IsValidIndex(ModeIndex), TEXT("Trying to update YCells from XCells, but Mode Index is not valid.")))
+	{
+		FDMXFixtureMode& Mode = Modes[ModeIndex];
+		FDMXFixtureMatrix Matrix = Mode.FixtureMatrixConfig;
+		if (ensureMsgf(Mode.bFixtureMatrixEnabled, TEXT("Trying to update YCells from XCells, but Fixture Matrix is not enabled.")))
+		{
+			// Disable the matrix while editing it so other functions align when enabling it again
+			SetFixtureMatrixEnabled(ModeIndex, false);
+
+			constexpr int32 MaxNumCells = 512;
+			Matrix.XCells = FMath::Clamp(Matrix.XCells, 1, MaxNumCells);
+			Matrix.YCells = FMath::Clamp(Matrix.YCells, 1, MaxNumCells - Matrix.XCells + 1);
+
+			SetFixtureMatrixEnabled(ModeIndex, true);
+			UpdateChannelSpan(ModeIndex);
+		}
+	}
+}
+
+void UDMXEntityFixtureType::UpdateXCellsFromYCells(int32 ModeIndex)
+{
+	if (ensureMsgf(Modes.IsValidIndex(ModeIndex), TEXT("Trying to update XCells from YCells, but Mode Index is not valid.")))
+	{
+		FDMXFixtureMode& Mode = Modes[ModeIndex];
+		FDMXFixtureMatrix Matrix = Mode.FixtureMatrixConfig;
+		if (ensureMsgf(Mode.bFixtureMatrixEnabled, TEXT("Trying to update XCells from YCells, but Fixture Matrix is not enabled.")))
+		{
+			// Disable the matrix while editing it so other functions align when enabling it again
+			SetFixtureMatrixEnabled(ModeIndex, false);
+
+			constexpr int32 MaxNumCells = 512;
+			Matrix.YCells = FMath::Clamp(Matrix.YCells, 1, MaxNumCells);
+			Matrix.XCells = FMath::Clamp(Matrix.XCells, 1, MaxNumCells - Matrix.YCells + 1);
+
+			SetFixtureMatrixEnabled(ModeIndex, true);
+			UpdateChannelSpan(ModeIndex);
+		}
 	}
 }
 
@@ -396,10 +1022,10 @@ void UDMXEntityFixtureType::FunctionValueToBytes(const FDMXFixtureFunction& InFu
 void UDMXEntityFixtureType::IntToBytes(EDMXFixtureSignalFormat InSignalFormat, bool bUseLSB, uint32 InValue, uint8* OutBytes)
 {
 	// Make sure the input value is in the valid range for the data type
-	InValue = ClampValueToDataType(InSignalFormat, InValue);
+	InValue = FDMXConversions::ClampValueBySignalFormat(InValue, InSignalFormat);
 
 	// Number of bytes we'll have to manipulate
-	const uint8 NumBytes = NumChannelsToOccupy(InSignalFormat);
+	const uint8 NumBytes = FDMXConversions::GetSizeOfSignalFormat(InSignalFormat);
 
 	if (NumBytes == 1)
 	{
@@ -427,7 +1053,7 @@ uint32 UDMXEntityFixtureType::BytesToFunctionValue(const FDMXFixtureFunction& In
 uint32 UDMXEntityFixtureType::BytesToInt(EDMXFixtureSignalFormat InSignalFormat, bool bUseLSB, const uint8* InBytes)
 {
 	// Number of bytes we'll read
-	const uint8 NumBytes = NumChannelsToOccupy(InSignalFormat);
+	const uint8 NumBytes = FDMXConversions::GetSizeOfSignalFormat(InSignalFormat);
 
 	if (NumBytes == 1)
 	{
@@ -459,7 +1085,7 @@ void UDMXEntityFixtureType::NormalizedValueToBytes(EDMXFixtureSignalFormat InSig
 	// Make sure InValue is in the range [0.0 ... 1.0]
 	InValue = FMath::Clamp(InValue, 0.0f, 1.0f);
 
-	const uint32 IntValue = GetDataTypeMaxValue(InSignalFormat) * InValue;
+	const uint32 IntValue = FDMXConversions::GetSignalFormatMaxValue(InSignalFormat) * InValue;
 
 	// Get the individual bytes from the computed IntValue
 	IntToBytes(InSignalFormat, bUseLSB, IntValue, OutBytes);
@@ -476,204 +1102,142 @@ float UDMXEntityFixtureType::BytesToNormalizedValue(EDMXFixtureSignalFormat InSi
 	const float Value = BytesToInt(InSignalFormat, bUseLSB, InBytes);
 
 	// Normalize it
-	return Value / GetDataTypeMaxValue(InSignalFormat);
+	return Value / FDMXConversions::GetSignalFormatMaxValue(InSignalFormat);
 }
 
+
 #if WITH_EDITOR
-void UDMXEntityFixtureType::PostEditChangeProperty(FPropertyChangedEvent& PropertyChangedEvent)
-{
-	Super::PostEditChangeProperty(PropertyChangedEvent);
+/** DEPRECATED 5.0 */
+FDataTypeChangeDelegate UDMXEntityFixtureType::DataTypeChangeDelegate_DEPRECATED;
+#endif
 
-	FName PropertyName = PropertyChangedEvent.GetPropertyName();
-		
-	if (PropertyName == GET_MEMBER_NAME_CHECKED(UDMXEntityFixtureType, bFixtureMatrixEnabled))
-	{				
-		for (FDMXFixtureMode& Mode : Modes)
-		{
-			if (!bFixtureMatrixEnabled)
-			{
-				Mode.FixtureMatrixConfig.CellAttributes.Reset();
-			}
-			UpdateChannelSpan(Mode);
-		}
+uint32 UDMXEntityFixtureType::GetDataTypeMaxValue(EDMXFixtureSignalFormat DataType)
+{	
+	// DEPRECATED 5.0
+	switch (DataType)
+	{
+	case EDMXFixtureSignalFormat::E8Bit:
+		return MAX_uint8;
+	case EDMXFixtureSignalFormat::E16Bit:
+		return MAX_uint16;
+	case EDMXFixtureSignalFormat::E24Bit:
+		return 0xFFFFFF;
+	case EDMXFixtureSignalFormat::E32Bit:
+		return MAX_uint32;
+	default:
+		checkNoEntry();
+		return 1;
 	}
-
-	RebuildFixturePatchCaches();
 }
-#endif // WITH_EDITOR
 
-#if WITH_EDITOR
-void UDMXEntityFixtureType::PostEditChangeChainProperty(FPropertyChangedChainEvent& PropertyChangedEvent)
+uint8 UDMXEntityFixtureType::NumChannelsToOccupy(EDMXFixtureSignalFormat DataType)
 {
-	Super::PostEditChangeChainProperty(PropertyChangedEvent);
-	FName PropertyName = PropertyChangedEvent.GetPropertyName();
-
-	// Clamp DefaultValue from selected data type
-	if (PropertyName == GET_MEMBER_NAME_CHECKED(FDMXFixtureFunction, DefaultValue) ||
-		PropertyName == GET_MEMBER_NAME_CHECKED(FDMXFixtureFunction, DataType))
+	// DEPRECATED 5.0
+	switch (DataType)
 	{
-		const int32 ModeIndex = PropertyChangedEvent.GetArrayIndex(GET_MEMBER_NAME_CHECKED(UDMXEntityFixtureType, Modes).ToString());	
-		check(ModeIndex != INDEX_NONE);
-
-		for (FDMXFixtureFunction& Function : Modes[ModeIndex].Functions)
-		{
-			ClampDefaultValue(Function);
-		}
-	}
+	case EDMXFixtureSignalFormat::E8Bit:
+		return 1;
 
-	// Refresh ChannelSpan from functions' settings
-	if (PropertyName == GET_MEMBER_NAME_CHECKED(UDMXEntityFixtureType, bFixtureMatrixEnabled) ||
-		PropertyName == GET_MEMBER_NAME_CHECKED(FDMXFixtureMatrix, CellAttributes) ||
-		PropertyName == GET_MEMBER_NAME_CHECKED(FDMXFixtureMatrix, XCells) ||
-		PropertyName == GET_MEMBER_NAME_CHECKED(FDMXFixtureMatrix, YCells) ||	
-		PropertyName == GET_MEMBER_NAME_CHECKED(FDMXFixtureMatrix, FirstCellChannel) ||
-		PropertyName == GET_MEMBER_NAME_CHECKED(FDMXFixtureMode, Functions) ||
-		PropertyName == GET_MEMBER_NAME_CHECKED(FDMXFixtureMode, bAutoChannelSpan) ||
-		PropertyName == GET_MEMBER_NAME_CHECKED(FDMXFixtureFunction, ChannelOffset) ||
-		PropertyName == GET_MEMBER_NAME_CHECKED(FDMXFixtureFunction, DataType))
-	{
-		// If we have a specific Modes index that was modified, update its properties
-		const int32 ModeIndex = PropertyChangedEvent.GetArrayIndex(GET_MEMBER_NAME_CHECKED(UDMXEntityFixtureType, Modes).ToString());
-		if (ModeIndex != INDEX_NONE)
-		{
-			if (PropertyName == GET_MEMBER_NAME_CHECKED(FDMXFixtureMatrix, XCells))
-			{
-				UpdateYCellsFromXCells(Modes[ModeIndex]);
-			}
-			else if (PropertyName == GET_MEMBER_NAME_CHECKED(FDMXFixtureMatrix, YCells))
-			{
-				UpdateXCellsFromYCells(Modes[ModeIndex]);
-			}
+	case EDMXFixtureSignalFormat::E16Bit:
+		return 2;
 
-			UpdateChannelSpan(Modes[ModeIndex]);
-		}
-		else
-		{
-			// Unfortunately, some operations like reordering an array's values, don't give us an array index
+	case EDMXFixtureSignalFormat::E24Bit:
+		return 3;
 
-			// Check if this is a property contained in a Modes property
-			const auto HeadNode = PropertyChangedEvent.PropertyChain.GetHead();
-			if (HeadNode != nullptr && HeadNode->GetValue() != nullptr
-				&& HeadNode->GetValue()->GetFName() == GET_MEMBER_NAME_CHECKED(UDMXEntityFixtureType, Modes))
-			{
-				for (FDMXFixtureMode& Mode : Modes)
-				{
-					if (PropertyName == GET_MEMBER_NAME_CHECKED(FDMXFixtureMatrix, XCells))
-					{
-						UpdateYCellsFromXCells(Mode);
-					}
-					else if (PropertyName == GET_MEMBER_NAME_CHECKED(FDMXFixtureMatrix, YCells))
-					{
-						UpdateXCellsFromYCells(Mode);
-					}
+	case EDMXFixtureSignalFormat::E32Bit:
+		return 4;
 
-					UpdateChannelSpan(Mode);
-				}
-			}
-		}
+	default:
+		// Unhandled type
+		checkNoEntry();
+		break;
 	}
-		
-	// Keep Attributes unique across either Functions or Matrix Attributes
-	// Note: The Attribute property exists in FDMXFixtureFunction and FDMXFixtureCellAttribute, the conditions cannot be separated.
-	// In other words, the if statement here will be entered in both cases, even if either OR statement would be removed.
-	if (PropertyName == GET_MEMBER_NAME_CHECKED(FDMXFixtureFunction, Attribute) ||
-		PropertyName == GET_MEMBER_NAME_CHECKED(FDMXFixtureCellAttribute, Attribute))
-	{
-		// Find duplicates
-		const int32 ModeIndex = PropertyChangedEvent.GetArrayIndex(GET_MEMBER_NAME_CHECKED(UDMXEntityFixtureType, Modes).ToString());
-		check(ModeIndex != INDEX_NONE);
-
-		FDMXFixtureMode& Mode = Modes[ModeIndex];
+	return 1;
+}
 
-		const int32 ChangedFunctionIndex = PropertyChangedEvent.GetArrayIndex(GET_MEMBER_NAME_CHECKED(FDMXFixtureMode, Functions).ToString());
-		const int32 ChangedCellAttributeIndex = PropertyChangedEvent.GetArrayIndex(GET_MEMBER_NAME_CHECKED(FDMXFixtureMatrix, CellAttributes).ToString());
+uint8 UDMXEntityFixtureType::GetFunctionLastChannel(const FDMXFixtureFunction& Function)
+{
+	// DEPRECATED 5.0
+	return Function.Channel + Function.GetNumChannels() - 1;
+}
 
-		FDMXAttributeName ChangedAttribute;
-		if (ChangedFunctionIndex != INDEX_NONE)
-		{
-			// Unique across functions
-			check(Mode.Functions.IsValidIndex(ChangedFunctionIndex));
-			ChangedAttribute = Mode.Functions[ChangedFunctionIndex].Attribute;
+bool UDMXEntityFixtureType::IsFunctionInModeRange(const FDMXFixtureFunction& InFunction, const FDMXFixtureMode& InMode, int32 ChannelOffset /*= 0*/)
+{
+	// DEPRECATED 5.0
+	const int32 LastChannel = InFunction.GetLastChannel();
+	const bool bLastChannelExceedsChannelSpan = LastChannel > InMode.ChannelSpan;
+	const bool bLastChannelExceedsUniverseSize = LastChannel + ChannelOffset > DMX_MAX_ADDRESS;
 
-			for (int32 IdxOtherFunction = 0; IdxOtherFunction < Mode.Functions.Num(); IdxOtherFunction++)
-			{
-				if (ChangedFunctionIndex == IdxOtherFunction)
-				{
-					continue;
-				}
+	return !bLastChannelExceedsChannelSpan && !bLastChannelExceedsUniverseSize;
+}
 
-				FDMXAttributeName& Attribute = Mode.Functions[IdxOtherFunction].Attribute;
-				if (ChangedAttribute == Attribute)
-				{
-					Attribute = FDMXAttributeName::None;
-				}
-			}
-		}
-		else
-		{
-			// Unique across fixture matrix attributes
-			check(Mode.FixtureMatrixConfig.CellAttributes.IsValidIndex(ChangedCellAttributeIndex));
-			ChangedAttribute = Mode.FixtureMatrixConfig.CellAttributes[ChangedCellAttributeIndex].Attribute;
+bool UDMXEntityFixtureType::IsFixtureMatrixInModeRange(const FDMXFixtureMatrix& InFixtureMatrix, const FDMXFixtureMode& InMode, int32 ChannelOffset /*= 0*/)
+{
+	// DEPRECATED 5.0
+	const int32 LastChannel = InFixtureMatrix.GetLastChannel();
+	const bool bLastChannelExceedsChannelSpan = LastChannel > InMode.ChannelSpan;
+	const bool bLastChannelExceedsUniverseSize = LastChannel + ChannelOffset > DMX_MAX_ADDRESS;
 
-			int32 NumCellAttributes = Mode.FixtureMatrixConfig.CellAttributes.Num();
-			for (int32 IdxOtherCellAttribute = 0; IdxOtherCellAttribute < NumCellAttributes; IdxOtherCellAttribute++)
-			{
-				if (ChangedCellAttributeIndex == IdxOtherCellAttribute)
-				{
-					continue;
-				}
+	return !bLastChannelExceedsChannelSpan && !bLastChannelExceedsUniverseSize;
+}
 
-				FDMXAttributeName& Attribute = Mode.FixtureMatrixConfig.CellAttributes[IdxOtherCellAttribute].Attribute;
-				if (ChangedAttribute == Attribute)
-				{
-					Attribute = FDMXAttributeName::None;
-				}
-			}
-		}
-	}
+void UDMXEntityFixtureType::ClampDefaultValue(FDMXFixtureFunction& InFunction)
+{
+	// DEPRECATED 5.0
+	const uint32 SafeDefaultValue = FMath::Min(InFunction.DefaultValue, static_cast<int64>(TNumericLimits<uint32>::Max()));
+	InFunction.DefaultValue = FDMXConversions::ClampValueBySignalFormat(SafeDefaultValue, InFunction.DataType);
+}
 
-	// Keep Fixture Patches' ActiveMode value valid
-	if (PropertyName == GET_MEMBER_NAME_CHECKED(UDMXEntityFixtureType, Modes))
+uint32 UDMXEntityFixtureType::ClampValueToDataType(EDMXFixtureSignalFormat DataType, uint32 InValue)
+{
+	// DEPRECATED 5.0
+	switch (DataType)
 	{
-		if (PropertyChangedEvent.ChangeType == EPropertyChangeType::ArrayRemove || PropertyChangedEvent.ChangeType == EPropertyChangeType::ArrayClear)
-		{		
-			if (ParentLibrary != nullptr)
-			{
-				ParentLibrary->ForEachEntityOfType<UDMXEntityFixturePatch>([this](UDMXEntityFixturePatch* Patch)
-				{
-					if (Patch->GetFixtureType() == this)
-					{
-						Patch->ValidateActiveMode();					
-					}
-				});
-			}
-		}
+	case EDMXFixtureSignalFormat::E8Bit:
+		return FMath::Clamp(InValue, 0u, (uint32)MAX_uint8);
 
-		int32 ChangedModeIndex = PropertyChangedEvent.GetArrayIndex(PropertyName.ToString());
-		if (Modes.IsValidIndex(ChangedModeIndex))
-		{
-			// Notify DataType changes
-			DataTypeChangeDelegate.Broadcast(this, Modes[ChangedModeIndex]);
-		}
-	}
+	case EDMXFixtureSignalFormat::E16Bit:
+		return FMath::Clamp(InValue, 0u, (uint32)MAX_uint16);
+
+	case EDMXFixtureSignalFormat::E24Bit:
+		return FMath::Clamp(InValue, 0u, 0xFFFFFFu);
+
+	case EDMXFixtureSignalFormat::E32Bit:
+		return FMath::Clamp(InValue, 0u, MAX_uint32);
 
-	RebuildFixturePatchCaches();
+	default:
+		break;
+	}
+	return InValue;
 }
-#endif // WITH_EDITOR
 
 #if WITH_EDITOR
-void UDMXEntityFixtureType::PostEditUndo()
+void UDMXEntityFixtureType::SetFunctionSize(FDMXFixtureFunction& InFunction, uint8 Size)
 {
-	for (FDMXFixtureMode& Mode : Modes)
+	// DEPRECATED 5.0
+	EDMXFixtureSignalFormat NewDataType;
+	switch (Size)
 	{
-		UpdateYCellsFromXCells(Mode);
-
-		UpdateChannelSpan(Mode);
+	case 0:
+	case 1:
+		NewDataType = EDMXFixtureSignalFormat::E8Bit;
+		break;
+	case 2:
+		NewDataType = EDMXFixtureSignalFormat::E16Bit;
+		break;
+	case 3:
+		NewDataType = EDMXFixtureSignalFormat::E24Bit;
+		break;
+	case 4:
+	default:
+		NewDataType = EDMXFixtureSignalFormat::E32Bit;
+		break;
 	}
 
-	Super::PostEditUndo();
+	InFunction.DataType = NewDataType;
 
-	RebuildFixturePatchCaches();
+	const uint32 SafeDefaultValue = FMath::Min(InFunction.DefaultValue, static_cast<int64>(TNumericLimits<uint32>::Max()));
+	InFunction.DefaultValue = FDMXConversions::ClampValueBySignalFormat(SafeDefaultValue, InFunction.DataType);
 }
 #endif // WITH_EDITOR
 
@@ -681,14 +1245,16 @@ void UDMXEntityFixtureType::PostEditUndo()
 void UDMXEntityFixtureType::UpdateModeChannelProperties(FDMXFixtureMode& Mode)
 {
 	// DEPRECATED 4.27
+	PRAGMA_DISABLE_DEPRECATION_WARNINGS
 	UpdateChannelSpan(Mode);
-	RebuildFixturePatchCaches();
+	PRAGMA_ENABLE_DEPRECATION_WARNINGS
 }
 #endif // WITH_EDITOR
 
 #if WITH_EDITOR
 void UDMXEntityFixtureType::UpdateChannelSpan(FDMXFixtureMode& Mode)
 {
+	// DEPRECATED 5.0
 	if (Mode.bAutoChannelSpan)
 	{
 		if (Mode.Functions.Num() == 0 &&
@@ -703,7 +1269,7 @@ void UDMXEntityFixtureType::UpdateChannelSpan(FDMXFixtureMode& Mode)
 			// Update span from common Functions
 			for (FDMXFixtureFunction& Function : Mode.Functions)
 			{
-				Function.Channel = ChannelSpan + 1 + Function.ChannelOffset;
+				Function.Channel = ChannelSpan + 1;
 
 				switch (Function.DataType)
 				{
@@ -727,36 +1293,30 @@ void UDMXEntityFixtureType::UpdateChannelSpan(FDMXFixtureMode& Mode)
 
 			// If fixture matrix is enabled, add the channel span of the matrix
 			int32 NumCells = Mode.FixtureMatrixConfig.XCells * Mode.FixtureMatrixConfig.YCells;
-			if (bFixtureMatrixEnabled && NumCells > 0)
+			if (Mode.bFixtureMatrixEnabled && NumCells > 0)
 			{
-				// Add 'empty' channels bewtween normal functions and cell attributes to the channel span
+				// Add 'empty' channels bewtween normal Functions and the Matrix to the channel span
 				ChannelSpan = FMath::Max(ChannelSpan, Mode.FixtureMatrixConfig.FirstCellChannel);
 
-				int32 FirstCellChannel = Mode.FixtureMatrixConfig.FirstCellChannel;
-
-				// Ignore channels that are overlapping commmon functions
-				FirstCellChannel = FMath::Max(FirstCellChannel, ChannelSpan + 1);
-
-				int32 FixtureMatrixLastChannel = Mode.FixtureMatrixConfig.GetFixtureMatrixLastChannel();
-
-				ChannelSpan += FixtureMatrixLastChannel - FirstCellChannel + 1;
+				ChannelSpan += Mode.FixtureMatrixConfig.GetNumChannels();
 			}
 
 			ChannelSpan = FMath::Max(ChannelSpan, 1);
 			Mode.ChannelSpan = ChannelSpan;
 		}
 
-		// Notify DataType changes
-		DataTypeChangeDelegate.Broadcast(this, Mode);
+		// Notify DataType changes (Deprecated 5.0)
+		PRAGMA_DISABLE_DEPRECATION_WARNINGS
+		DataTypeChangeDelegate_DEPRECATED.Broadcast(this, Mode);
+		PRAGMA_ENABLE_DEPRECATION_WARNINGS
 	}
-
-	RebuildFixturePatchCaches();
 }
 #endif // WITH_EDITOR
 
 #if WITH_EDITOR
 void UDMXEntityFixtureType::UpdateYCellsFromXCells(FDMXFixtureMode& Mode)
 {
+	// Deprecated 5.0
 	const int32 MaxNumCells = 512;
 
 	Mode.FixtureMatrixConfig.XCells = FMath::Clamp(Mode.FixtureMatrixConfig.XCells, 1, MaxNumCells);
@@ -767,6 +1327,7 @@ void UDMXEntityFixtureType::UpdateYCellsFromXCells(FDMXFixtureMode& Mode)
 #if WITH_EDITOR
 void UDMXEntityFixtureType::UpdateXCellsFromYCells(FDMXFixtureMode& Mode)
 {
+	// Deprecated 5.0
 	const int32 MaxNumCells = 512;
 
 	Mode.FixtureMatrixConfig.YCells = FMath::Clamp(Mode.FixtureMatrixConfig.YCells, 1, MaxNumCells);
@@ -774,17 +1335,4 @@ void UDMXEntityFixtureType::UpdateXCellsFromYCells(FDMXFixtureMode& Mode)
 }
 #endif // WITH_EDITOR
 
-void UDMXEntityFixtureType::RebuildFixturePatchCaches()
-{
-	if (UDMXLibrary* Library = ParentLibrary.Get())
-	{
-		for (UDMXEntityFixturePatch* FixturePatch : Library->GetEntitiesTypeCast<UDMXEntityFixturePatch>())
-		{
-			if (FixturePatch->GetFixtureType() == this)
-			{
-				FixturePatch->RebuildCache();
-		}
-	}
-	}
-}
-
+#undef LOCTEXT_NAMESPACE
diff --git a/Engine/Plugins/VirtualProduction/DMX/DMXEngine/Source/DMXRuntime/Private/Sequencer/MovieSceneDMXLibrarySection.cpp b/Engine/Plugins/VirtualProduction/DMX/DMXEngine/Source/DMXRuntime/Private/Sequencer/MovieSceneDMXLibrarySection.cpp
--- a/Engine/Plugins/VirtualProduction/DMX/DMXEngine/Source/DMXRuntime/Private/Sequencer/MovieSceneDMXLibrarySection.cpp
+++ b/Engine/Plugins/VirtualProduction/DMX/DMXEngine/Source/DMXRuntime/Private/Sequencer/MovieSceneDMXLibrarySection.cpp
@@ -2,6 +2,7 @@
 
 #include "Sequencer/MovieSceneDMXLibrarySection.h"
 
+#include "DMXConversions.h"
 #include "DMXProtocolCommon.h"
 #include "DMXRuntimeLog.h"
 #include "DMXRuntimeObjectVersion.h"
@@ -71,7 +72,7 @@ void FDMXFixturePatchChannel::UpdateNumberOfChannels(bool bResetDefaultValues /*
 	{
 		// Add channels for functions that are in mode range and have an Attribute set
 
-		if (UDMXEntityFixtureType::IsFunctionInModeRange(Function, Mode, PatchChannelOffset) &&
+		if (Function.GetLastChannel() <= Patch->GetChannelSpan() &&
 			FDMXAttributeName::IsValid(Function.Attribute.GetName()))
 		{
 			bool bNewChannel = false;
@@ -95,7 +96,7 @@ void FDMXFixturePatchChannel::UpdateNumberOfChannels(bool bResetDefaultValues /*
 		}
 	}
 
-	if (UDMXEntityFixtureType::IsFixtureMatrixInModeRange(Mode.FixtureMatrixConfig, Mode, PatchChannelOffset))
+	if (Mode.FixtureMatrixConfig.GetLastChannel() <= Patch->GetChannelSpan())
 	{
 		int32 NumXCells = Mode.FixtureMatrixConfig.XCells;
 		int32 NumYCells = Mode.FixtureMatrixConfig.YCells;
@@ -199,7 +200,7 @@ FDMXCachedFunctionChannelInfo::FDMXCachedFunctionChannelInfo(const TArray<FDMXFi
 			});
 
 		const bool bMissingFunction = !AttributeNameChannelMap.Contains(FunctionChannel.AttributeName) || !CellAttributePtr;
-		if (!CellAttributePtr && bMissingFunction)
+		if (!CellAttributePtr || bMissingFunction)
 		{
 			UE_LOG(MovieSceneDMXLibrarySectionLog, Warning, TEXT("%S: Function with attribute %s from %s doesn't have a counterpart Fixture Function."), __FUNCTION__, *FunctionChannel.AttributeName.ToString(), *FixturePatch->GetDisplayName());
 			UE_LOG(MovieSceneDMXLibrarySectionLog, Warning, TEXT("%S: Further attributes may be missing. Warnings ommited to avoid overflowing the log."), __FUNCTION__);
@@ -268,42 +269,35 @@ const FDMXFixtureFunctionChannel* FDMXCachedFunctionChannelInfo::TryGetFunctionC
 
 void FDMXCachedFunctionChannelInfo::GetMatrixCellChannelsAbsoluteNoSorting(UDMXEntityFixturePatch* FixturePatch, const FIntPoint& CellCoordinate, TMap<FName, int32>& OutAttributeToAbsoluteChannelMap) const
 {
-	if (FixturePatch)
+	UDMXEntityFixtureType* FixtureType = FixturePatch ? FixturePatch->GetFixtureType() : nullptr;
+	const FDMXFixtureMode* FixtureModePtr = FixturePatch ? FixturePatch->GetActiveMode() : nullptr;
+	if (FixturePatch && FixtureType && FixtureModePtr && FixtureModePtr->bFixtureMatrixEnabled)
 	{
-		if (UDMXEntityFixtureType* FixtureType = FixturePatch->GetFixtureType())
-		{
-			if (FixtureType->bFixtureMatrixEnabled)
-			{
-				if (const FDMXFixtureMode* FixtureModePtr = FixturePatch->GetActiveMode())
-				{
-					const FDMXFixtureMatrix& FixtureMatrix = FixtureModePtr->FixtureMatrixConfig;
+		const FDMXFixtureMatrix& FixtureMatrix = FixtureModePtr->FixtureMatrixConfig;
 
-					TMap<const FDMXFixtureCellAttribute*, int32> AttributeToRelativeChannelOffsetMap;
-					int32 CellDataSize = 0;
-					int32 AttributeChannelOffset = 0;
-					for (const FDMXFixtureCellAttribute& CellAttribute : FixtureMatrix.CellAttributes)
-					{
-						AttributeToRelativeChannelOffsetMap.Add(&CellAttribute, AttributeChannelOffset);
-						const int32 AttributeSize = UDMXEntityFixtureType::NumChannelsToOccupy(CellAttribute.DataType);
+		TMap<const FDMXFixtureCellAttribute*, int32> AttributeToRelativeChannelOffsetMap;
+		int32 CellDataSize = 0;
+		int32 AttributeChannelOffset = 0;
+		for (const FDMXFixtureCellAttribute& CellAttribute : FixtureMatrix.CellAttributes)
+		{
+			AttributeToRelativeChannelOffsetMap.Add(&CellAttribute, AttributeChannelOffset);
+			const int32 AttributeSize = FDMXConversions::GetSizeOfSignalFormat(CellAttribute.DataType);
 
-						CellDataSize += AttributeSize;
-						AttributeChannelOffset += UDMXEntityFixtureType::NumChannelsToOccupy(CellAttribute.DataType);
-					}
+			CellDataSize += AttributeSize;
+			AttributeChannelOffset += FDMXConversions::GetSizeOfSignalFormat(CellAttribute.DataType);
+		}
 
-					const int32 FixtureMatrixAbsoluteStartingChannel = FixturePatch->GetStartingChannel() + FixtureMatrix.FirstCellChannel - 1;
-					const int32 CellChannelOffset = (CellCoordinate.Y * FixtureMatrix.XCells + CellCoordinate.X) * CellDataSize;
-					const int32 AbsoluteCellStartingChannel = FixtureMatrixAbsoluteStartingChannel + CellChannelOffset;
+		const int32 FixtureMatrixAbsoluteStartingChannel = FixturePatch->GetStartingChannel() + FixtureMatrix.FirstCellChannel - 1;
+		const int32 CellChannelOffset = (CellCoordinate.Y * FixtureMatrix.XCells + CellCoordinate.X) * CellDataSize;
+		const int32 AbsoluteCellStartingChannel = FixtureMatrixAbsoluteStartingChannel + CellChannelOffset;
 
-					for (const TTuple<const FDMXFixtureCellAttribute*, int32>& AttributeToRelativeChannelOffsetKvp : AttributeToRelativeChannelOffsetMap)
-					{
-						const FName FunctionAttributeName = AttributeToRelativeChannelOffsetKvp.Key->Attribute.GetName();
-						const int32 AbsoluteChannel = AbsoluteCellStartingChannel + AttributeToRelativeChannelOffsetKvp.Value;
+		for (const TTuple<const FDMXFixtureCellAttribute*, int32>& AttributeToRelativeChannelOffsetKvp : AttributeToRelativeChannelOffsetMap)
+		{
+			const FName FunctionAttributeName = AttributeToRelativeChannelOffsetKvp.Key->Attribute.GetName();
+			const int32 AbsoluteChannel = AbsoluteCellStartingChannel + AttributeToRelativeChannelOffsetKvp.Value;
 
-						check(AbsoluteChannel > 0 && AbsoluteChannel <= DMX_UNIVERSE_SIZE);
-						OutAttributeToAbsoluteChannelMap.Add(FunctionAttributeName, AbsoluteChannel);
-					}
-				}
-			}
+			check(AbsoluteChannel > 0 && AbsoluteChannel <= DMX_UNIVERSE_SIZE);
+			OutAttributeToAbsoluteChannelMap.Add(FunctionAttributeName, AbsoluteChannel);
 		}
 	}
 }
diff --git a/Engine/Plugins/VirtualProduction/DMX/DMXEngine/Source/DMXRuntime/Public/DMXRuntimeUtils.h b/Engine/Plugins/VirtualProduction/DMX/DMXEngine/Source/DMXRuntime/Public/DMXRuntimeUtils.h
--- a/Engine/Plugins/VirtualProduction/DMX/DMXEngine/Source/DMXRuntime/Public/DMXRuntimeUtils.h
+++ b/Engine/Plugins/VirtualProduction/DMX/DMXEngine/Source/DMXRuntime/Public/DMXRuntimeUtils.h
@@ -2,15 +2,23 @@
 
 #pragma once
 
-#include "CoreMinimal.h"
 #include "Library/DMXEntityFixtureType.h"
 
+#include "CoreMinimal.h"
+#include "JsonObjectConverter.h"
+
 class UDMXEntityFixturePatch;
 
 
 class DMXRUNTIME_API FDMXRuntimeUtils
 {
 public:
+	/**
+	 * Generates a unique name given a base one and a list of existing ones, by appending an index to
+	 * existing names. If InBaseName is an empty String, it returns "Default name".
+	 */
+	static FString GenerateUniqueNameFromExisting(const TSet<FString>& InExistingNames, const FString& InBaseName);
+
 	/**
 	 * Utility to separate a name from an index at the end.
 	 * @param InString	The string to be separated.
@@ -316,6 +324,9 @@ public:
 		}
 	}
 
+	/** Serializes a struct to a Sting */
+	static TOptional<FString> SerializeStructToString(const UStruct* StructClass, const void* StructData);
+
 	// can't instantiate this class
 	FDMXRuntimeUtils() = delete;
 };
diff --git a/Engine/Plugins/VirtualProduction/DMX/DMXEngine/Source/DMXRuntime/Public/Library/DMXEntityFixturePatch.h b/Engine/Plugins/VirtualProduction/DMX/DMXEngine/Source/DMXRuntime/Public/Library/DMXEntityFixturePatch.h
--- a/Engine/Plugins/VirtualProduction/DMX/DMXEngine/Source/DMXRuntime/Public/Library/DMXEntityFixturePatch.h
+++ b/Engine/Plugins/VirtualProduction/DMX/DMXEngine/Source/DMXRuntime/Public/Library/DMXEntityFixturePatch.h
@@ -40,8 +40,11 @@ class DMXRUNTIME_API UDMXEntityFixturePatch
 public:
 	UDMXEntityFixturePatch();
 
-protected:
 	// ~Begin UObject Interface
+#if WITH_EDITOR
+	virtual bool Modify(bool bAlwaysMarkDirty = true) override;
+#endif
+protected:
 	virtual void PostInitProperties() override;
 	virtual void PostLoad() override;
 
@@ -125,6 +128,14 @@ public:
 	/** Returns the manual starting channel. Used only when bAutoAssignAdress is false.  In most cases, you'll want GetStartingChannel instead. */
 	FORCEINLINE int32 GetManualStartingAddress() const { return ManualStartingAddress; }
 
+	/** 
+	 * Sets the starting channel of the Fixture Patch.
+	 * 
+	 * If Auto Assign Address was set to true, turns off Auto Assign Address.
+	 */
+	UFUNCTION(BlueprintCallable, Category = "DMX|Fixture Patch")
+	void SetStartingChannel(int32 NewStartingChannel);
+
 	/** Return the starting channel */
 	UFUNCTION(BlueprintPure, Category = "DMX|Fixture Patch")
 	int32 GetStartingChannel() const;
@@ -380,6 +391,9 @@ public:
 	bool GetAllMatrixCells(TArray<FDMXCell>& Cells);
 
 private:
+	/** Called when a Fixture Type changed */
+	void OnFixtureTypeChanged(const UDMXEntityFixtureType* FixtureType);
+
 	/** Tries to access the FixtureMatrix config of this patch and logs issues. Returns the matrix of nullptr if it isn't valid. */
 	const FDMXFixtureMatrix* GetFixtureMatrix() const;
 
diff --git a/Engine/Plugins/VirtualProduction/DMX/DMXEngine/Source/DMXRuntime/Public/Library/DMXEntityFixtureType.h b/Engine/Plugins/VirtualProduction/DMX/DMXEngine/Source/DMXRuntime/Public/Library/DMXEntityFixtureType.h
--- a/Engine/Plugins/VirtualProduction/DMX/DMXEngine/Source/DMXRuntime/Public/Library/DMXEntityFixtureType.h
+++ b/Engine/Plugins/VirtualProduction/DMX/DMXEngine/Source/DMXRuntime/Public/Library/DMXEntityFixtureType.h
@@ -40,6 +40,24 @@ struct DMXRUNTIME_API FDMXFixtureFunction
 {
 	GENERATED_BODY()
 
+	/** Constructor */
+	FDMXFixtureFunction()
+		: Attribute(FDMXNameListItem::None)
+		, FunctionName()
+		, Description()
+		, DefaultValue(0)
+		, Channel(1)
+		, ChannelOffset_DEPRECATED(0)
+		, DataType(EDMXFixtureSignalFormat::E8Bit)
+		, bUseLSBMode(false)
+	{}
+
+	/** Returns the number of channels the function spans, according to its data type */
+	FORCEINLINE uint8 GetNumChannels() const { return static_cast<uint8>(DataType) + 1; }
+
+	/** Returns the last channel of the Function */
+	uint8 GetLastChannel() const;
+
 	/**
 	 * The Attribute name to map this Function to.
 	 * This is used to easily find the Function in Blueprints, using an Attribute
@@ -47,34 +65,29 @@ struct DMXRUNTIME_API FDMXFixtureFunction
 	 * The list of Attributes can be edited on
 	 * Project Settings->Plugins->DMX Protocol->Fixture Settings->Fixture Function Attributes
 	 */
-	UPROPERTY(EditAnywhere, BlueprintReadOnly, meta = (DisplayName = "Attribute Mapping", DisplayPriority = "11"), Category = "DMX")
+	UPROPERTY(EditAnywhere, BlueprintReadOnly, meta = (DisplayName = "Attribute Mapping", DisplayPriority = "11"), Category = "Function Settings")
 	FDMXAttributeName Attribute;
 
-	UPROPERTY(EditAnywhere, BlueprintReadOnly, meta = (DisplayPriority = "10"), Category = "DMX")
+	UPROPERTY(EditAnywhere, BlueprintReadOnly, meta = (DisplayPriority = "10"), Category = "Function Settings")
 	FString FunctionName;
 
-	UPROPERTY(EditAnywhere, BlueprintReadOnly, meta = (DisplayPriority = "20"), Category = "DMX")
+	UPROPERTY(EditAnywhere, BlueprintReadOnly, meta = (DisplayPriority = "20"), Category = "Function Settings")
 	FString Description;
 
 	/** The Default Value of the function, imported from GDTF. The plugin doesn't make use of this value, but it can be used in blueprints */
-	UPROPERTY(EditAnywhere, BlueprintReadOnly, meta = (DisplayPriority = "30"), Category = "DMX")
+	UPROPERTY(EditAnywhere, BlueprintReadOnly, meta = (DisplayPriority = "30"), Category = "Function Settings")
 	int64 DefaultValue;
 
 	/** This function's starting channel */
-	UPROPERTY(VisibleAnywhere, BlueprintReadOnly, meta = (DisplayName = "Channel Assignment", ClampMin = "1", ClampMax = "512", DisplayPriority = "2"), Category = "DMX")
+	UPROPERTY(EditAnywhere, BlueprintReadOnly, meta = (DisplayName = "Channel Assignment", DisplayPriority = "2"), Category = "Function Settings")
 	int32 Channel;
 
-	/**
-	 * This function's channel offset.
-	 * E.g.: if the function's starting channel is supposed to be 10
-	 * and ChannelOffset = 5, the function's starting channel becomes 15
-	 * and all following functions follow it accordingly.
-	 */
-	UPROPERTY(EditAnywhere, BlueprintReadOnly, meta = (DisplayName = "Channel Offset", ClampMin = "0", ClampMax = "511", DisplayPriority = "1"), Category = "DMX")
-	int32 ChannelOffset;
+	/** DEPRECATED 5.0. Instead the 'Channel' property is EditAnywhere so any function can be assigned freely */
+	UPROPERTY(meta = (DeprecatedProperty, DeprecationMessage = "Deprecated since the Channel property can be set in the DMX Library Editor."))
+	int32 ChannelOffset_DEPRECATED;
 
 	/** This function's data type. Defines the used number of channels (bytes) */
-	UPROPERTY(EditAnywhere, BlueprintReadOnly, meta = (DisplayPriority = "5"), Category = "DMX")
+	UPROPERTY(EditAnywhere, BlueprintReadOnly, meta = (DisplayPriority = "5"), Category = "Function Settings")
 	EDMXFixtureSignalFormat DataType;
 
 	/**
@@ -89,28 +102,26 @@ struct DMXRUNTIME_API FDMXFixtureFunction
 	 * In MSB mode, the example above would be interpreted in binary as 0x00 0x01, which means 1.
 	 * The first byte (0) became the highest part in binary form and the following byte (1), the lowest.
 	 */
-	UPROPERTY(EditAnywhere, BlueprintReadOnly, meta = (DisplayName = "Use LSB Mode", DisplayPriority = "29"), Category = "DMX")
+	UPROPERTY(EditAnywhere, BlueprintReadOnly, meta = (DisplayName = "Use LSB Mode", DisplayPriority = "29"), Category = "Function Settings")
 	bool bUseLSBMode = false;
+};
 
-	/** Returns the number of channels the function spans, according to its data type */
-	FORCEINLINE uint8 GetNumChannels() const { return static_cast<uint8>(DataType) + 1;	}
+USTRUCT(BlueprintType)
+struct DMXRUNTIME_API FDMXFixtureCellAttribute
+{
+	GENERATED_BODY()
 
-	FDMXFixtureFunction()
+	/** Constructor */
+	FDMXFixtureCellAttribute()
 		: Attribute(FDMXNameListItem::None)
-		, FunctionName()
 		, Description()
 		, DefaultValue(0)
-		, Channel(1)
-		, ChannelOffset(0)
 		, DataType(EDMXFixtureSignalFormat::E8Bit)
 		, bUseLSBMode(false)
 	{}
-};
 
-USTRUCT(BlueprintType)
-struct DMXRUNTIME_API FDMXFixtureCellAttribute
-{
-	GENERATED_BODY()
+	/** Returns the number of channels of the attribute */
+	uint8 GetNumChannels() const { return static_cast<uint8>(DataType) + 1; }
 
 	/**
 	 * The Attribute name to map this Function to.
@@ -134,6 +145,7 @@ struct DMXRUNTIME_API FDMXFixtureCellAttribute
 	EDMXFixtureSignalFormat DataType;
 
 	/**
+	 * The Endianess of the Attribute:
 	 * Least Significant Byte mode makes the individual bytes (channels) of the function be
 	 * interpreted with the first bytes being the lowest part of the number.
 	 *
@@ -147,17 +159,6 @@ struct DMXRUNTIME_API FDMXFixtureCellAttribute
 	 */
 	UPROPERTY(EditAnywhere, BlueprintReadOnly, meta = (DisplayName = "Use LSB Mode", DisplayPriority = "29"), Category = "DMX")
 	bool bUseLSBMode;
-
-	FDMXFixtureCellAttribute()
-		: Attribute(FDMXNameListItem::None)
-		, Description()
-		, DefaultValue(0)
-		, DataType(EDMXFixtureSignalFormat::E8Bit)
-		, bUseLSBMode(false)
-	{}
-
-	/** Returns the number of channels of the attribute */
-	uint8 GetNumChannels() const { return static_cast<uint8>(DataType) + 1; }
 };
 
 USTRUCT(BlueprintType)
@@ -165,30 +166,35 @@ struct DMXRUNTIME_API FDMXFixtureMatrix
 {
 	GENERATED_BODY()
 
-	UPROPERTY(EditAnywhere, BlueprintReadOnly, meta = (DisplayPriority = "60", DisplayName = "Cell Attributes"), Category = "DMX")
-	TArray<FDMXFixtureCellAttribute> CellAttributes;
+	/** Constructor */
+	FDMXFixtureMatrix();
 
-	UPROPERTY(EditAnywhere, BlueprintReadOnly, meta = (DisplayPriority = "20", DisplayName = "First Cell Channel", ClampMin = "1", ClampMax = "512"), Category = "DMX")
-	int32 FirstCellChannel;
+	/** Returns the number of channels of the Matrix */
+	int32 GetNumChannels() const;
 
-	UPROPERTY(EditAnywhere, BlueprintReadOnly, meta = (DisplayPriority = "30", DisplayName = "X Cells", ClampMin = "0"), Category = "DMX")
-	int32 XCells;
+	/** Returns the last channel of the Matrix */
+	int32 GetLastChannel() const;
 
-	UPROPERTY(EditAnywhere, BlueprintReadOnly, meta = (DisplayPriority = "40", DisplayName = "Y Cells", ClampMin = "0"), Category = "DMX")
-	int32 YCells;
+	UE_DEPRECATED(5.0, "Deprecated for more consistent naming. Instead use FDMXFixtureMatrix::GetLastChannel")
+	int32 GetFixtureMatrixLastChannel() const;
 
-	UPROPERTY(EditAnywhere, BlueprintReadOnly, meta = (DisplayPriority = "50", DisplayName = "PixelMapping Distribution"), Category = "DMX")
-	EDMXPixelMappingDistribution PixelMappingDistribution;
+	UE_DEPRECATED(5.0, "Deprecated to reduce redundant code. Instead use UDMXFixturePatch::GetMatrixCellChannelsRelative")
+	bool GetChannelsFromCell(FIntPoint CellCoordinate, FDMXAttributeName Attribute, TArray<int32>& Channels) const;
 
-	FDMXFixtureMatrix()
-		: FirstCellChannel(1)
-		, XCells(1)
-		, YCells(1)
-		, PixelMappingDistribution(EDMXPixelMappingDistribution::TopLeftToRight)
-	{}
+	UPROPERTY(EditAnywhere, BlueprintReadOnly, meta = (DisplayPriority = "60", DisplayName = "Cell Attributes"), Category = "Mode Settings")
+	TArray<FDMXFixtureCellAttribute> CellAttributes;
 
-	bool GetChannelsFromCell(FIntPoint CellCoordinate, FDMXAttributeName Attribute, TArray<int32>& Channels) const;
-	int32 GetFixtureMatrixLastChannel() const;
+	UPROPERTY(EditAnywhere, BlueprintReadOnly, meta = (DisplayPriority = "20", DisplayName = "First Cell Channel", ClampMin = "1", ClampMax = "512"), Category = "Mode Settings")
+	int32 FirstCellChannel = 1;
+
+	UPROPERTY(EditAnywhere, BlueprintReadOnly, meta = (DisplayPriority = "30", DisplayName = "X Cells", ClampMin = "1"), Category = "Mode Settings")
+	int32 XCells = 1;
+
+	UPROPERTY(EditAnywhere, BlueprintReadOnly, meta = (DisplayPriority = "40", DisplayName = "Y Cells", ClampMin = "1"), Category = "Mode Settings")
+	int32 YCells = 1;
+
+	UPROPERTY(EditAnywhere, BlueprintReadOnly, meta = (DisplayPriority = "50", DisplayName = "PixelMapping Distribution"), Category = "Mode Settings")
+	EDMXPixelMappingDistribution PixelMappingDistribution = EDMXPixelMappingDistribution::TopLeftToRight;
 };
 
 USTRUCT(BlueprintType)
@@ -215,105 +221,249 @@ struct DMXRUNTIME_API FDMXFixtureMode
 {
 	GENERATED_BODY()
 
-	UPROPERTY(EditAnywhere, BlueprintReadOnly, meta = (DisplayPriority = "1"), Category = "DMX")
+	UPROPERTY(EditAnywhere, BlueprintReadOnly, meta = (DisplayPriority = "10"), Category = "Mode Settings")
 	FString ModeName;
 
-	UPROPERTY(EditAnywhere, BlueprintReadOnly, meta = (DisplayPriority = "15"), Category = "DMX")
+	UPROPERTY(EditAnywhere, BlueprintReadOnly, meta = (DisplayPriority = "20"), Category = "Mode Settings")
 	TArray<FDMXFixtureFunction> Functions;
 
-	/** Number of channels (bytes) used by this mode's functions */
-	UPROPERTY(EditAnywhere, BlueprintReadOnly, meta = (ClampMin = "0", ClampMax = "512", DisplayPriority = "10", EditCondition = "!bAutoChannelSpan"), Category = "DMX")
-	int32 ChannelSpan;
-
 	/**
 	 * When enabled, ChannelSpan is automatically set based on the created functions and their data types.
 	 * If disabled, ChannelSpan can be manually set and functions and functions' channels beyond the
 	 * specified span will be ignored.
 	 */
-	UPROPERTY(EditAnywhere, meta = (DisplayPriority = "5"), Category = "DMX")
-	bool bAutoChannelSpan;
+	UPROPERTY(EditAnywhere, Category = "Mode Settings", meta = (DisplayPriority = "30"))
+	bool bAutoChannelSpan = true;
 
-	UPROPERTY(EditAnywhere, Category = "DMX")
-	FDMXFixtureMatrix FixtureMatrixConfig;
+	/** Number of channels (bytes) used by this mode's functions */
+	UPROPERTY(EditAnywhere, BlueprintReadOnly, Category = "Mode Settings", meta = (ClampMin = "4", ClampMax = "512", DisplayPriority = "40", EditCondition = "!bAutoChannelSpan"))
+	int32 ChannelSpan = 0;
 
-	FDMXFixtureMode()
-		: ChannelSpan(0)
-		, bAutoChannelSpan(true)
-	{}
+	/** 
+	 * Modulators applied right before a patch of this type is received. 
+	 * NOTE: Modulators only affect the patch's normalized values! Untouched values are still available when accesing raw values. 
+	 */
+	UPROPERTY(EditAnywhere, Instanced, Category = "Mode Settings", meta = (DisplayPriority = "50"))
+	TArray<UDMXModulator*> InputModulators;
+
+	UPROPERTY(EditAnywhere, BlueprintReadOnly, Category = "Mode Settings", meta = (DisplayPriority = "60"))
+	bool bFixtureMatrixEnabled = false;
+
+	UPROPERTY(EditAnywhere, Category = "Mode Settings", meta = (DisplayPriority = "70"))
+	FDMXFixtureMatrix FixtureMatrixConfig;
 
 #if WITH_EDITOR
-	/**
-	 * Adding the function into the mode, checking the channel offset
-	 */
-	int32 AddOrInsertFunction(int32 IndexOfFunction, const FDMXFixtureFunction& InFunction);
+	/** DEPRECATED 5.0 */
+	UE_DEPRECATED(5.0, "Removed in favor of UDMXEntityFixtureType::AddFunction and UDMXEntityFixtureType::InsertFunction")
+	int32 AddOrInsertFunction(int32 IndexOfFunction, FDMXFixtureFunction InFunction);
 #endif
 };
 
+
 #if WITH_EDITOR
-	/** Notification when data type changed */
-	DECLARE_MULTICAST_DELEGATE_TwoParams(FDataTypeChangeDelegate, const UDMXEntityFixtureType*, const FDMXFixtureMode&);
+/** Notification when data type changed */
+DECLARE_MULTICAST_DELEGATE_TwoParams(FDataTypeChangeDelegate, const UDMXEntityFixtureType*, const FDMXFixtureMode&);
 #endif
 
+DECLARE_MULTICAST_DELEGATE_OneParam(FDMXOnFixtureTypeChangedDelegate, const UDMXEntityFixtureType* /** ChangedFixtureType */);
+
 
+/** 
+ * Class to describe a type of Fixture. Fixture Patches can be created from Fixture Types (see UDMXEntityFixturePatch::ParentFixtureTypeTemplate).
+ */
 UCLASS(BlueprintType, Blueprintable, meta = (DisplayName = "DMX Fixture Type"))
 class DMXRUNTIME_API UDMXEntityFixtureType
 	: public UDMXEntity
 {
 	GENERATED_BODY()
 
+	//~ Begin UObject interface
+protected:
+	virtual void Serialize(FArchive& Ar) override;
 public:
-	UDMXEntityFixtureType()
-		: bFixtureMatrixEnabled(false)
-	{}
+#if WITH_EDITOR
+	virtual bool Modify(bool bAlwaysMarkDirty = true) override;
+	virtual void PostEditChangeProperty(FPropertyChangedEvent& PropertyChangedEvent) override;
+	virtual void PostEditChangeChainProperty(FPropertyChangedChainEvent& PropertyChangedChainEvent) override;
+	virtual void PostEditUndo() override;
+#endif
+	//~ End UObject interface
 
+public:
+#if WITH_EDITOR
+	UFUNCTION(BlueprintCallable, Category = "Fixture Settings")
+	void SetModesFromDMXImport(UDMXImport* DMXImportAsset);
+#endif // WITH_EDITOR
+	
+	/** Returns a delegate that is and should be broadcast whenever a Fixture Type changed */
+	static FDMXOnFixtureTypeChangedDelegate& GetOnFixtureTypeChanged();
+
+	/** The GDTF file from which the Fixture Type was setup */
 	UPROPERTY(EditAnywhere, BlueprintReadOnly, Category = "Fixture Settings")
 	UDMXImport* DMXImport;
-
+	
+	/** The Category of the Fixture, useful for Filtering */
 	UPROPERTY(EditAnywhere, BlueprintReadOnly, Category = "Fixture Settings", meta = (DisplayName = "DMX Category"))
 	FDMXFixtureCategory DMXCategory;
 
-	UPROPERTY(EditAnywhere, BlueprintReadOnly, Category = "Fixture Settings", meta = (DisplayName = "Matrix Fixture"))
-	bool bFixtureMatrixEnabled;
-
 	UPROPERTY(EditAnywhere, BlueprintReadOnly, Category = "Fixture Settings")
 	TArray<FDMXFixtureMode> Modes;
 
+private:
+	/** Delegate that should be broadcast whenever a fixture type changed */
+	static FDMXOnFixtureTypeChangedDelegate OnFixtureTypeChangedDelegate;
+
+
+	//////////////////////////////////////////////////
+	// Helpers to edit the Fixture Type
+
+	// Fixture Mode related
+public:
+	/**
+	 * Adds a Mode to the Modes Array
+	 *
+	 * @return								The Index of the newly added Mode.
+	 */	
+	int32 AddMode();
+
 	/** 
-	 * Modulators applied right before a patch of this type is received. 
-	 * NOTE: Modulators only affect the patch's normalized values! Untouched values are still available when accesing raw values. 
+	 * Duplicates the Modes at specified Indices 
+	 *
+	 * @param InModeIndices					The indicies of the Modes to duplicate
+	 * @param OutNewModeIndices				The indices of the newly created Modes.
 	 */
-	UPROPERTY(EditAnywhere, Instanced, Category = "Fixture Settings")
-	TArray<UDMXModulator*> InputModulators;
+	void DuplicateModes(TArray<int32> InModeIndicesToDuplicate, TArray<int32>& OutNewModeIndices);
+
+	/** Deletes the Modes at specified Indices */
+	void RemoveModes(const TArray<int32>& ModeIndicesToDelete);
+
+	/** Sets a Mode Name for specified Mode 
+	 * 
+	 * @param InModeIndex					The index of the Mode in the Modes array
+	 * @param InDesiredModeName				The desired Name that should be set
+	 * @param OutUniqueModeName				The unique Name that was set
+	*/
+	void SetModeName(int32 InModeIndex, const FString& InDesiredModeName, FString& OutUniqueModeName);
+
+	/** Enables or disables the Matrix, reorders Function channels accordingly
+	 *
+	 * @param ModeIndex						The index of the Mode for which the Matrix should be enabled or disabled.
+	 * @param bEnableMatrix					Whether to enable or disable the Matrix
+	*/
+	void SetFixtureMatrixEnabled(int32 ModeIndex, bool bEnableMatrix);
 
+	/**
+	 * Updates the channel span of the Mode.
+	 *
+	 * @param ModeIndex						The Index of the Mode for which the Channel Span should be updated.
+	 */
+	void UpdateChannelSpan(int32 ModeIndex);
+
+
+	// Fixture Function related
 public:
-#if WITH_EDITOR
-	UFUNCTION(BlueprintCallable, Category = "Fixture Settings")
-	void SetModesFromDMXImport(UDMXImport* DMXImportAsset);
+	/** 
+	 * Adds a new Function to the Mode's Functions array
+	 *
+	 * @param InModeIndex					The index of the Mode, that will have the Function added to its Functions array.
+	 * @return								The Index of the newly added Function.
+	 */
+	int32 AddFunction(int32 InModeIndex);
 
-	static void SetFunctionSize(FDMXFixtureFunction& InFunction, uint8 Size);
+	/** 
+	 * Inserts a Function to the Mode's Function Array
+	 *
+	 * @param InModeIndex					The index of the Mode, that will have the Function added to its Functions array.
+	 * @param InInsertAtIndex				If a valid Index, the Function will be inserted at this Index, and subsequent Function's Channels will be reodered after it.
+	 * @param InOutNewFunction				The function that will be inserted.
+	 * @return								The Index of the newly added Function.
+	 */
+	int32 InsertFunction(int32 InModeIndex, int32 InInsertAtIndex, FDMXFixtureFunction& InOutNewFunction);
 
-	static FDataTypeChangeDelegate& GetDataTypeChangeDelegate() { return DataTypeChangeDelegate; }
-#endif // WITH_EDITOR
+	/**
+	 * Adds a Function to the Mode's Function Array
+	 *
+	 * @param InModeIndex					The index of the Mode in which the Functions to duplicate reside.
+	 * @param InFunctionIndicesToDuplicate	The indices of the Functions to duplicate.
+	 * @param OutNewFunctionIndices			The Function indices where the newly added Functions reside
+	 */
+	void DuplicateFunctions(int32 InModeIndex, const TArray<int32>& InFunctionIndicesToDuplicate, TArray<int32>& OutNewFunctionIndices);
 
-	/** Gets the last channel occupied by the Function */
-	static uint8 GetFunctionLastChannel(const FDMXFixtureFunction& Function);
+	/**
+	 * Removes Functions from the Mode's Function Array
+	 *
+	 * @param InModeIndex					The index of the Mode in which the Functions to remove reside
+	 * @param FunctionIndicesToDelete		The indices of the Functions to remove.
+	 */
+	void RemoveFunctions(int32 ModeIndex, TArray<int32> FunctionIndicesToDelete);
 
 	/**
-	 * Return true if a Function's occupied channels are within the Mode's Channel Span.
-	 * Optionally add an offset to the function address
+	 * Reorders a function to reside at the Insert At Index, subsequently reorders other affected Functions
+	 *
+	 * @param ModeIndex						The Index of the Mode in which the Functions reside
+	 * @param FunctionIndex					The Index of the Function that is reorderd
+	 * @param InsertAtIndex					The Index of the Function where the function is inserted.
 	 */
-	static bool IsFunctionInModeRange(const FDMXFixtureFunction& InFunction, const FDMXFixtureMode& InMode, int32 ChannelOffset = 0);
+	void ReorderFunction(int32 ModeIndex, int32 FunctionToReorderIndex, int32 InsertAtIndex);
 
-	static bool IsFixtureMatrixInModeRange(const FDMXFixtureMatrix& InFixtureMatrix, const FDMXFixtureMode& InMode, int32 ChannelOffset = 0);
+	/** Sets a Mode Name for specified Mode
+	 *
+	 * @param InModeIndex					The index of the Mode in the Modes array
+	 * @param InFunctionIndex				The index of the Function in the Mode's Function array
+	 * @param DesiredFunctionName			The desired Name that should be set
+	 * @param OutUniqueFunctionName			The unique Name that was set
+	*/
+	void SetFunctionName(int32 InModeIndex, int32 InFunctionIndex, const FString& InDesiredFunctionName, FString& OutUniqueFunctionName);
+
+	/** Sets a Starting Channel for the Function, aligns it to other functions
+	 *
+	 * @param InModeIndex					The index of the Mode in the Modes array
+	 * @param InFunctionIndex				The index of the Function in the Mode's Function array
+	 * @param InDesiredStartingChannel		The desired Starting Channel that should be set
+	 * @param OutStartingChannel			The resulting Starting Channel that was set
+	*/
+	void SetFunctionStartingChannel(int32 InModeIndex, int32 InFunctionIndex, int32 InDesiredStartingChannel, int32& OutStartingChannel);
 
-	static void ClampDefaultValue(FDMXFixtureFunction& InFunction);
+	/**
+	 * Clamps the Default Value of the Function by its Data Type
+	 *
+	 * @param ModeIndex						The Index of the Mode in which the Functions reside
+	 * @param FunctionToRemoveIndex			The Index of the Function for which the Default Value is clamped
+	 */
+	void ClampFunctionDefautValueByDataType(int32 ModeIndex, int32 FunctionToRemoveIndex);
 
-	static uint8 NumChannelsToOccupy(EDMXFixtureSignalFormat DataType);
 
-	static uint32 ClampValueToDataType(EDMXFixtureSignalFormat DataType, uint32 InValue);
+	// Fixture Matrix related
+public:
+	/** Adds a new cell attribute to the selected Mode */
+	void AddCellAttribute(int32 ModeIndex);
 
-	static uint32 GetDataTypeMaxValue(EDMXFixtureSignalFormat DataType);
+	/**
+	 * Reorders the Fixture Matrix to reside after a function, subsequently reorders other affected Functions
+	 *
+	 * @param FixtureType					The Fixture Type in which the Functions reside
+	 * @param ModeIndex						The Index of the Mode in which the Functions reside
+	 * @param InsertAfterFunctionIndex		The Index of the Function after which the Matrix is inserted. If an invalid index is specified, the Matrix will added after the last Function Channel.
+	 */
+	void ReorderMatrix(int32 ModeIndex, int32 InsertAtFunctionIndex);
+
+	/**
+	 * Updates the channel span of the Mode.
+	 *
+	 * @param ModeIndex						The Index of the Mode for which the YCells should be updated
+	 */
+	void UpdateYCellsFromXCells(int32 ModeIndex);
+
+	/**
+	 * Updates the channel span of the Mode.
+	 *
+	 * @param ModeIndex						The Index of the Mode for which the XCells should be updated
+	 */
+	void UpdateXCellsFromYCells(int32 ModeIndex);
 
+
+	// Conversions. @TODO, move to FDMXConversions
+public:
 	//~ Conversions to/from Bytes, Int and Normalized Float values.
 	static void FunctionValueToBytes(const FDMXFixtureFunction& InFunction, uint32 InValue, uint8* OutBytes);
 	static void IntToBytes(EDMXFixtureSignalFormat InSignalFormat, bool bUseLSB, uint32 InValue, uint8* OutBytes);
@@ -327,32 +477,62 @@ public:
 	static float BytesToFunctionNormalizedValue(const FDMXFixtureFunction& InFunction, const uint8* InBytes);
 	static float BytesToNormalizedValue(EDMXFixtureSignalFormat InSignalFormat, bool bUseLSB, const uint8* InBytes);
 
+
+
+	//////////////////////////////////////////////////
+	// Deprecated Members
+	
+public:
+	UE_DEPRECATED(5.0, "Deprecated in favor of FDMXConversions::GetSignalFormatMaxValue.")
+	static uint32 GetDataTypeMaxValue(EDMXFixtureSignalFormat DataType);
+
+	UE_DEPRECATED(5.0, "Deprecated in favor of FDMXConversions::SizeOfSignalFormat.")
+	static uint8 NumChannelsToOccupy(EDMXFixtureSignalFormat DataType);
+
+	UE_DEPRECATED(5.0, "Deprecated in favor of FDMXFixtureFunction::GetLastChannel.")
+	static uint8 GetFunctionLastChannel(const FDMXFixtureFunction& Function);
+
+	UE_DEPRECATED(5.0, "Deprecated to reduce redundant code. Instead use FDMXFixtureFunction::GetNumChannels() and FDMXFixtureMode::ChannelSpan")
+	static bool IsFunctionInModeRange(const FDMXFixtureFunction& InFunction, const FDMXFixtureMode& InMode, int32 ChannelOffset = 0);
+
+	UE_DEPRECATED(5.0, "Deprecated since the use of the function and its name were not clear, leading to hard to read, complicated code.")
+	static bool IsFixtureMatrixInModeRange(const FDMXFixtureMatrix& InFixtureMatrix, const FDMXFixtureMode& InMode, int32 ChannelOffset = 0);
+
+	UE_DEPRECATED(5.0, "Deprecated to reduce redundant code. Use FDMXConversions::ClampValueBySignalFormat locally where clamping is required.")
+	static void ClampDefaultValue(FDMXFixtureFunction& InFunction);
+
+	UE_DEPRECATED(5.0, "Deprecated in favor of FDMXConversions now holding all conversions. Use FDMXConversions::ClampValueBySignalFormat instead.")
+	static uint32 ClampValueToDataType(EDMXFixtureSignalFormat DataType, uint32 InValue);
+
 #if WITH_EDITOR
-	virtual void PostEditChangeProperty(FPropertyChangedEvent& PropertyChangedEvent) override;
-	virtual void PostEditChangeChainProperty(FPropertyChangedChainEvent& PropertyChangedEvent) override;
-	virtual void PostEditUndo() override;
+	UE_DEPRECATED(5.0, "Deprecated because of unclear use. Set via FDMXFixtureFunction::DataType directly instead.")
+	static void SetFunctionSize(FDMXFixtureFunction& InFunction, uint8 Size);
 
-	/** DEPRECATED 4.27 */
-	UE_DEPRECATED(4.27, "Use UpdateChannelSpan instead.")
+	UE_DEPRECATED(5.0, "Deprecated in favor of the more generic GetOnFixtureTypeChanged which is supported in non-editor builds as well.")
+	static FDataTypeChangeDelegate& GetDataTypeChangeDelegate() { return DataTypeChangeDelegate_DEPRECATED; }
+
+	UE_DEPRECATED(4.27, "Use MakeValid instead.")
 	void UpdateModeChannelProperties(FDMXFixtureMode& Mode);
 
-	/** Updates the channel span of the Mode */
+	UE_DEPRECATED(5.0, "Deprecated in favor of UDMXEntityFixtureType::UpdateChannelSpan.")
 	void UpdateChannelSpan(FDMXFixtureMode& Mode);
 
-	/** Updates the FixtureMatrixConfig's YCells property given num XCells for the specified Mode */
+	UE_DEPRECATED(5.0, "Deprecated  in favor of UDMXEntityFixtureType::UpdateYCellsFromXCells.")
 	void UpdateYCellsFromXCells(FDMXFixtureMode& Mode);
 
-	/** Updates the FixtureMatrixConfig's XCells property given num YCells for the specified Mode */
+	UE_DEPRECATED(5.0, "Deprecated to in favor of UDMXEntityFixtureType::UpdateXCellsFromYCells.")
 	void UpdateXCellsFromYCells(FDMXFixtureMode& Mode);
-
 #endif // WITH_EDITOR
 
-private:
-	/** Rebuilds the cache of fixture patches that use this fixture type */
-	void RebuildFixturePatchCaches();
+	UPROPERTY(meta = (DeprecatedProperty, DeprecationMessage = "FixtureMatrixEnabled is deprecated. Instead now each Mode has a FixtureMatrixEnabled property."))
+	bool bFixtureMatrixEnabled = false;
+	
+	UPROPERTY(meta = (DeprecatedProperty, DeprecationMessage = "Fixture Type's 'Input Modulators' property is deprecated. Instead now each Mode has its own Input Modulators."))
+	TArray<UDMXModulator*> InputModulators;
 
+private:
 #if WITH_EDITOR
 	/** Editor only data type change delagate */
-	static FDataTypeChangeDelegate DataTypeChangeDelegate;
+	static FDataTypeChangeDelegate DataTypeChangeDelegate_DEPRECATED;
 #endif
 };
diff --git a/Engine/Plugins/VirtualProduction/DMX/DMXPixelMapping/Source/DMXPixelMappingEditor/Private/Customizations/DMXPixelMappingDetailCustomization_FixtureGroup.cpp b/Engine/Plugins/VirtualProduction/DMX/DMXPixelMapping/Source/DMXPixelMappingEditor/Private/Customizations/DMXPixelMappingDetailCustomization_FixtureGroup.cpp
--- a/Engine/Plugins/VirtualProduction/DMX/DMXPixelMapping/Source/DMXPixelMappingEditor/Private/Customizations/DMXPixelMappingDetailCustomization_FixtureGroup.cpp
+++ b/Engine/Plugins/VirtualProduction/DMX/DMXPixelMapping/Source/DMXPixelMappingEditor/Private/Customizations/DMXPixelMappingDetailCustomization_FixtureGroup.cpp
@@ -259,20 +259,20 @@ FReply FDMXPixelMappingDetailCustomization_FixtureGroup::OnFixturePatchesDragged
 		TArray<TSharedPtr<FDMXPixelMappingComponentTemplate>> Templates;
 		for (const FDMXEntityFixturePatchRef& FixturePatchRef : SelectedFixturePatches)
 		{
-			if (UDMXEntityFixturePatch* FixturePatch = FixturePatchRef.GetFixturePatch())
+			const UDMXEntityFixturePatch* FixturePatch = FixturePatchRef.GetFixturePatch();
+			const UDMXEntityFixtureType* FixtureType = FixturePatch ? FixturePatch->GetFixtureType() : nullptr;
+			const FDMXFixtureMode* ActiveModePtr = FixturePatch ? FixturePatch->GetActiveMode() : nullptr;
+			if (FixturePatch && FixtureType && ActiveModePtr)
 			{
-				if (UDMXEntityFixtureType* FixtureType = FixturePatch->GetFixtureType())
+				if (ActiveModePtr->bFixtureMatrixEnabled)
 				{
-					if (FixtureType->bFixtureMatrixEnabled)
-					{
-						TSharedRef<FDMXPixelMappingComponentTemplate> FixturePatchMatrixTemplate = MakeShared<FDMXPixelMappingComponentTemplate>(UDMXPixelMappingMatrixComponent::StaticClass(), FixturePatchRef);
-						Templates.Add(FixturePatchMatrixTemplate);
-					}
-					else
-					{
-						TSharedRef<FDMXPixelMappingComponentTemplate> FixturePatchItemTemplate = MakeShared<FDMXPixelMappingComponentTemplate>(UDMXPixelMappingFixtureGroupItemComponent::StaticClass(), FixturePatchRef);
-						Templates.Add(FixturePatchItemTemplate);
-					}
+					TSharedRef<FDMXPixelMappingComponentTemplate> FixturePatchMatrixTemplate = MakeShared<FDMXPixelMappingComponentTemplate>(UDMXPixelMappingMatrixComponent::StaticClass(), FixturePatchRef);
+					Templates.Add(FixturePatchMatrixTemplate);
+				}
+				else
+				{
+					TSharedRef<FDMXPixelMappingComponentTemplate> FixturePatchItemTemplate = MakeShared<FDMXPixelMappingComponentTemplate>(UDMXPixelMappingFixtureGroupItemComponent::StaticClass(), FixturePatchRef);
+					Templates.Add(FixturePatchItemTemplate);
 				}
 			}
 		}
diff --git a/Engine/Plugins/VirtualProduction/DMX/DMXPixelMapping/Source/DMXPixelMappingEditor/Private/Toolkits/DMXPixelMappingToolkit.cpp b/Engine/Plugins/VirtualProduction/DMX/DMXPixelMapping/Source/DMXPixelMappingEditor/Private/Toolkits/DMXPixelMappingToolkit.cpp
--- a/Engine/Plugins/VirtualProduction/DMX/DMXPixelMapping/Source/DMXPixelMappingEditor/Private/Toolkits/DMXPixelMappingToolkit.cpp
+++ b/Engine/Plugins/VirtualProduction/DMX/DMXPixelMapping/Source/DMXPixelMappingEditor/Private/Toolkits/DMXPixelMappingToolkit.cpp
@@ -27,7 +27,6 @@
 #include "Library/DMXEntityFixtureType.h"
 #include "Library/DMXEntityFixturePatch.h"
 #include "Library/DMXLibrary.h"
-#include "DMXUtils.h"
 #include "DMXPixelMappingComponentWidget.h"
 
 #include "ScopedTransaction.h"
diff --git a/Engine/Plugins/VirtualProduction/DMX/DMXPixelMapping/Source/DMXPixelMappingEditor/Private/Widgets/SDMXPixelMappingFixturePatchDetailRow.cpp b/Engine/Plugins/VirtualProduction/DMX/DMXPixelMapping/Source/DMXPixelMappingEditor/Private/Widgets/SDMXPixelMappingFixturePatchDetailRow.cpp
--- a/Engine/Plugins/VirtualProduction/DMX/DMXPixelMapping/Source/DMXPixelMappingEditor/Private/Widgets/SDMXPixelMappingFixturePatchDetailRow.cpp
+++ b/Engine/Plugins/VirtualProduction/DMX/DMXPixelMapping/Source/DMXPixelMappingEditor/Private/Widgets/SDMXPixelMappingFixturePatchDetailRow.cpp
@@ -46,9 +46,11 @@ void SDMXPixelMappingFixturePatchDetailRow::Construct(const FArguments& InArgs)
 				SAssignNew(FixturePatchNameTextBlock, STextBlock)
 				.Text_Lambda([FixturePatch]()
 					{
-						if (FixturePatch.IsValid())
+						const UDMXEntityFixtureType* FixtureType = FixturePatch.IsValid() ? FixturePatch->GetFixtureType() : nullptr;
+						const FDMXFixtureMode* FixtureModePtr = FixturePatch.IsValid() ? FixturePatch->GetActiveMode() : nullptr;
+						if (FixtureType && FixtureModePtr)
 						{
-							if (FixturePatch->GetFixtureType() && FixturePatch->GetFixtureType()->bFixtureMatrixEnabled)
+							if (FixtureModePtr->bFixtureMatrixEnabled)
 							{
 								return FText::Format(LOCTEXT("MatrixFixturePatchName", "Matrix: {0}"), FText::FromString(FixturePatch->Name));
 							}
