DMX - Fix setting the universe ID in channel monitor does not reset the channel values

diff --git a/Engine/Plugins/VirtualProduction/DMX/DMXEngine/Source/DMXEditor/Private/Widgets/Monitors/SDMXChannelsMonitor.cpp b/Engine/Plugins/VirtualProduction/DMX/DMXEngine/Source/DMXEditor/Private/Widgets/Monitors/SDMXChannelsMonitor.cpp
--- a/Engine/Plugins/VirtualProduction/DMX/DMXEngine/Source/DMXEditor/Private/Widgets/Monitors/SDMXChannelsMonitor.cpp
+++ b/Engine/Plugins/VirtualProduction/DMX/DMXEngine/Source/DMXEditor/Private/Widgets/Monitors/SDMXChannelsMonitor.cpp
@@ -292,6 +292,8 @@ void SDMXChannelsMonitor::OnUniverseIDValueCommitted(const FText& InNewText, ETe
 				UniverseIDEditableTextBox->SetText(FText::FromString(FString::FromInt(0)));
 			}
 		}
+
+		ZeroChannelValues();
 	}
 }
 
