DMX - Fix matrix is limited to 64x64 cells (now 512 total), make dmx fixture params blueprint editable, unhides all categories for fixture actors and fixture components for reevaluation

diff --git a/Engine/Plugins/VirtualProduction/DMX/DMXEngine/Source/DMXEditor/Private/DMXFixtureTypeSharedData.cpp b/Engine/Plugins/VirtualProduction/DMX/DMXEngine/Source/DMXEditor/Private/DMXFixtureTypeSharedData.cpp
--- a/Engine/Plugins/VirtualProduction/DMX/DMXEngine/Source/DMXEditor/Private/DMXFixtureTypeSharedData.cpp
+++ b/Engine/Plugins/VirtualProduction/DMX/DMXEngine/Source/DMXEditor/Private/DMXFixtureTypeSharedData.cpp
@@ -578,7 +578,7 @@ void FDMXFixtureTypeSharedData::AddFunctionToSelectedMode()
 		ModeRef.FixtureType->Modify();
 
 		Mode.Functions.Add(NewFunction);
-		ModeRef.FixtureType->UpdateModeChannelProperties(Mode);
+		ModeRef.FixtureType->UpdateChannelSpan(Mode);
 
 		ModeRef.FixtureType->PostEditChange();
 
diff --git a/Engine/Plugins/VirtualProduction/DMX/DMXEngine/Source/DMXRuntime/Private/Library/DMXEntityFixtureType.cpp b/Engine/Plugins/VirtualProduction/DMX/DMXEngine/Source/DMXRuntime/Private/Library/DMXEntityFixtureType.cpp
--- a/Engine/Plugins/VirtualProduction/DMX/DMXEngine/Source/DMXRuntime/Private/Library/DMXEntityFixtureType.cpp
+++ b/Engine/Plugins/VirtualProduction/DMX/DMXEngine/Source/DMXRuntime/Private/Library/DMXEntityFixtureType.cpp
@@ -154,7 +154,7 @@ void UDMXEntityFixtureType::SetModesFromDMXImport(UDMXImport* DMXImportAsset)
 			}
 
 			// Compute mode channel span from functions' addresses and sizes
-			UpdateModeChannelProperties(Mode);
+			UpdateChannelSpan(Mode);
 		}
 	}
 }
@@ -463,7 +463,7 @@ void UDMXEntityFixtureType::PostEditChangeProperty(FPropertyChangedEvent& Proper
 			{
 				Mode.FixtureMatrixConfig.CellAttributes.Reset();
 			}
-			UpdateModeChannelProperties(Mode);
+			UpdateChannelSpan(Mode);
 		}
 	}
 }
@@ -503,7 +503,16 @@ void UDMXEntityFixtureType::PostEditChangeChainProperty(FPropertyChangedChainEve
 		const int32 ModeIndex = PropertyChangedEvent.GetArrayIndex(GET_MEMBER_NAME_CHECKED(UDMXEntityFixtureType, Modes).ToString());
 		if (ModeIndex != INDEX_NONE)
 		{
-			UpdateModeChannelProperties(Modes[ModeIndex]);
+			UpdateChannelSpan(Modes[ModeIndex]);
+
+			if (PropertyName == GET_MEMBER_NAME_CHECKED(FDMXFixtureMatrix, XCells))
+			{
+				UpdateYCellsFromXCells(Modes[ModeIndex]);
+			}
+			else if (PropertyName == GET_MEMBER_NAME_CHECKED(FDMXFixtureMatrix, YCells))
+			{
+				UpdateXCellsFromYCells(Modes[ModeIndex]);
+			}
 		}
 		else
 		{
@@ -516,7 +525,16 @@ void UDMXEntityFixtureType::PostEditChangeChainProperty(FPropertyChangedChainEve
 			{
 				for (FDMXFixtureMode& Mode : Modes)
 				{
-					UpdateModeChannelProperties(Mode);
+					UpdateChannelSpan(Mode);
+
+					if (PropertyName == GET_MEMBER_NAME_CHECKED(FDMXFixtureMatrix, XCells))
+					{
+						UpdateYCellsFromXCells(Mode);
+					}
+					else if (PropertyName == GET_MEMBER_NAME_CHECKED(FDMXFixtureMatrix, YCells))
+					{
+						UpdateXCellsFromYCells(Mode);
+					}
 				}
 			}
 		}
@@ -604,7 +622,8 @@ void UDMXEntityFixtureType::PostEditUndo()
 {
 	for (FDMXFixtureMode& Mode : Modes)
 	{
-		UpdateModeChannelProperties(Mode);
+		UpdateChannelSpan(Mode);
+		UpdateYCellsFromXCells(Mode);
 	}
 
 	Super::PostEditUndo();
@@ -613,6 +632,14 @@ void UDMXEntityFixtureType::PostEditUndo()
 
 #if WITH_EDITOR
 void UDMXEntityFixtureType::UpdateModeChannelProperties(FDMXFixtureMode& Mode)
+{
+	// DEPRECATED 4.27
+	UpdateChannelSpan(Mode);
+}
+#endif // WITH_EDITOR
+
+#if WITH_EDITOR
+void UDMXEntityFixtureType::UpdateChannelSpan(FDMXFixtureMode& Mode)
 {
 	if (Mode.bAutoChannelSpan)
 	{
@@ -623,7 +650,6 @@ void UDMXEntityFixtureType::UpdateModeChannelProperties(FDMXFixtureMode& Mode)
 		}
 		else
 		{
-
 			int32 ChannelSpan = 0;
 
 			// Update span from common Functions
@@ -678,6 +704,26 @@ void UDMXEntityFixtureType::UpdateModeChannelProperties(FDMXFixtureMode& Mode)
 }
 #endif // WITH_EDITOR
 
+#if WITH_EDITOR
+void UDMXEntityFixtureType::UpdateYCellsFromXCells(FDMXFixtureMode& Mode)
+{
+	const int32 MaxNumCells = 512;
+
+	Mode.FixtureMatrixConfig.XCells = FMath::Clamp(Mode.FixtureMatrixConfig.XCells, 1, MaxNumCells);
+	Mode.FixtureMatrixConfig.YCells = FMath::Clamp(Mode.FixtureMatrixConfig.YCells, 1, MaxNumCells - Mode.FixtureMatrixConfig.XCells + 1);
+}
+#endif // WITH_EDITOR
+
+#if WITH_EDITOR
+void UDMXEntityFixtureType::UpdateXCellsFromYCells(FDMXFixtureMode& Mode)
+{
+	const int32 MaxNumCells = 512;
+
+	Mode.FixtureMatrixConfig.YCells = FMath::Clamp(Mode.FixtureMatrixConfig.YCells, 1, MaxNumCells);
+	Mode.FixtureMatrixConfig.XCells = FMath::Clamp(Mode.FixtureMatrixConfig.XCells, 1, MaxNumCells - Mode.FixtureMatrixConfig.YCells + 1);
+}
+#endif // WITH_EDITOR
+
 #if WITH_EDITOR
 int32 FDMXFixtureMode::AddOrInsertFunction(int32 IndexOfFunction, const FDMXFixtureFunction& InFunction)
 {
@@ -707,4 +753,4 @@ int32 FDMXFixtureMode::AddOrInsertFunction(int32 IndexOfFunction, const FDMXFixt
 	return Index;
 }
 
-#endif // WITH_EDITOR
\ No newline at end of file
+#endif // WITH_EDITOR
diff --git a/Engine/Plugins/VirtualProduction/DMX/DMXEngine/Source/DMXRuntime/Public/Library/DMXEntityFixtureType.h b/Engine/Plugins/VirtualProduction/DMX/DMXEngine/Source/DMXRuntime/Public/Library/DMXEntityFixtureType.h
--- a/Engine/Plugins/VirtualProduction/DMX/DMXEngine/Source/DMXRuntime/Public/Library/DMXEntityFixtureType.h
+++ b/Engine/Plugins/VirtualProduction/DMX/DMXEngine/Source/DMXRuntime/Public/Library/DMXEntityFixtureType.h
@@ -322,7 +322,19 @@ public:
 	virtual void PostEditChangeChainProperty(FPropertyChangedChainEvent& PropertyChangedEvent) override;
 	virtual void PostEditUndo() override;
 
+	/** DEPRECATED 4.27 */
+	UE_DEPRECATED(4.27, "Use UpdateChannelSpan instead.")
 	void UpdateModeChannelProperties(FDMXFixtureMode& Mode);
+
+	/** Updates the channel span of the Mode */
+	void UpdateChannelSpan(FDMXFixtureMode& Mode);
+
+	/** Updates the FixtureMatrixConfig's YCells property given num XCells for the specified Mode */
+	void UpdateYCellsFromXCells(FDMXFixtureMode& Mode);
+
+	/** Updates the FixtureMatrixConfig's XCells property given num YCells for the specified Mode */
+	void UpdateXCellsFromYCells(FDMXFixtureMode& Mode);
+
 #endif // WITH_EDITOR
 
 private:
diff --git a/Engine/Plugins/VirtualProduction/DMX/DMXFixtures/Source/DMXFixtures/Private/DMXFixtureActorMatrix.cpp b/Engine/Plugins/VirtualProduction/DMX/DMXFixtures/Source/DMXFixtures/Private/DMXFixtureActorMatrix.cpp
--- a/Engine/Plugins/VirtualProduction/DMX/DMXFixtures/Source/DMXFixtures/Private/DMXFixtureActorMatrix.cpp
+++ b/Engine/Plugins/VirtualProduction/DMX/DMXFixtures/Source/DMXFixtures/Private/DMXFixtureActorMatrix.cpp
@@ -2,6 +2,7 @@
 
 #include "DMXFixtureActorMatrix.h"
 
+#include "DMXProtocolConstants.h"
 #include "DMXStats.h"
 
 #include "Rendering/Texture2DResource.h"
@@ -97,11 +98,13 @@ void ADMXFixtureActorMatrix::InitializeMatrixFixture()
 	XCells = MatrixProperties.XCells;
 	YCells = MatrixProperties.YCells;
 
-	// Limit cells [1-64]
+	// Limit cells [1-DMX_UNIVERSE_SIZE]
+	constexpr int DMXUniverseSize = DMX_UNIVERSE_SIZE;
+
 	XCells = FMath::Max(XCells, 1);
 	YCells = FMath::Max(YCells, 1);
-	XCells = FMath::Min(XCells, 64);
-	YCells = FMath::Min(YCells, 64);
+	XCells = FMath::Min(XCells, DMXUniverseSize);
+	YCells = FMath::Min(YCells, DMXUniverseSize);
 
 	int NbrCells = XCells * YCells;
 
@@ -415,11 +418,13 @@ void ADMXFixtureActorMatrix::GenerateEditorMatrixMesh()
 		XCells = MatrixProperties.XCells;
 		YCells = MatrixProperties.YCells;
 
-		// Limit cells [1-64]
+		// Limit cells [1-DMX_UNIVERSE_SIZE]
+		constexpr int DMXUniverseSize = DMX_UNIVERSE_SIZE;
+
 		XCells = FMath::Max(XCells, 1);
 		YCells = FMath::Max(YCells, 1);
-		XCells = FMath::Min(XCells, 64);
-		YCells = FMath::Min(YCells, 64);
+		XCells = FMath::Min(XCells, DMXUniverseSize);
+		YCells = FMath::Min(YCells, DMXUniverseSize);
 
 		MatrixHead->ClearAllMeshSections();
 		GenerateMatrixCells();
diff --git a/Engine/Plugins/VirtualProduction/DMX/DMXFixtures/Source/DMXFixtures/Public/DMXFixtureActor.h b/Engine/Plugins/VirtualProduction/DMX/DMXFixtures/Source/DMXFixtures/Public/DMXFixtureActor.h
--- a/Engine/Plugins/VirtualProduction/DMX/DMXFixtures/Source/DMXFixtures/Public/DMXFixtureActor.h
+++ b/Engine/Plugins/VirtualProduction/DMX/DMXFixtures/Source/DMXFixtures/Public/DMXFixtureActor.h
@@ -30,7 +30,7 @@ enum EDMXFixtureQualityLevel
 };
 
 
-UCLASS(HideCategories = ("Rendering", "Variable", "Input", "Tags", "Activation", "Cooking", "Replication", "AssetUserData", "Collision", "LOD", "Actor"))
+UCLASS()
 class DMXFIXTURES_API ADMXFixtureActor : public AActor
 {
 	GENERATED_BODY()
@@ -49,7 +49,7 @@ public:
 	void FeedFixtureData();
 
 	// VISUAL QUALITY LEVEL----------------------
-	UPROPERTY(EditAnywhere, BlueprintReadOnly, Category = "DMX Light Fixture", meta = (DisplayPriority = 0))
+	UPROPERTY(EditAnywhere, BlueprintReadWrite, Category = "DMX Light Fixture", meta = (DisplayPriority = 0))
 	TEnumAsByte<EDMXFixtureQualityLevel> QualityLevel;
 
 	// HIERARCHY---------------------------------
@@ -94,25 +94,25 @@ public:
 	void InterpolateDMXComponents(float DeltaSeconds);
 	
 	// PARAMETERS---------------------------------
-	UPROPERTY(EditAnywhere, BlueprintReadOnly, Category = "DMX Light Fixture")
+	UPROPERTY(EditAnywhere, BlueprintReadWrite, Category = "DMX Light Fixture")
 	float LightIntensityMax;
 
-	UPROPERTY(EditAnywhere, BlueprintReadOnly, Category = "DMX Light Fixture")
+	UPROPERTY(EditAnywhere, BlueprintReadWrite, Category = "DMX Light Fixture")
 	float LightDistanceMax;
 
-	UPROPERTY(EditAnywhere, BlueprintReadOnly, Category = "DMX Light Fixture")
+	UPROPERTY(EditAnywhere, BlueprintReadWrite, Category = "DMX Light Fixture")
 	float LightColorTemp;
 
-	UPROPERTY(EditAnywhere, BlueprintReadOnly, Category = "DMX Light Fixture")
+	UPROPERTY(EditAnywhere, BlueprintReadWrite, Category = "DMX Light Fixture")
 	float SpotlightIntensityScale;
 
-	UPROPERTY(EditAnywhere, BlueprintReadOnly, Category = "DMX Light Fixture")
+	UPROPERTY(EditAnywhere, BlueprintReadWrite, Category = "DMX Light Fixture")
 	float PointlightIntensityScale;
 
-	UPROPERTY(EditAnywhere, BlueprintReadOnly, Category = "DMX Light Fixture")
+	UPROPERTY(EditAnywhere, BlueprintReadWrite, Category = "DMX Light Fixture")
 	bool LightCastShadow;
 
-	UPROPERTY(EditAnywhere, BlueprintReadOnly, Category = "DMX Light Fixture")
+	UPROPERTY(EditAnywhere, BlueprintReadWrite, Category = "DMX Light Fixture")
 	bool UseDynamicOcclusion;
 
 	// DMX COMPONENT -----------------------------
diff --git a/Engine/Plugins/VirtualProduction/DMX/DMXFixtures/Source/DMXFixtures/Public/DMXFixtureActorMatrix.h b/Engine/Plugins/VirtualProduction/DMX/DMXFixtures/Source/DMXFixtures/Public/DMXFixtureActorMatrix.h
--- a/Engine/Plugins/VirtualProduction/DMX/DMXFixtures/Source/DMXFixtures/Public/DMXFixtureActorMatrix.h
+++ b/Engine/Plugins/VirtualProduction/DMX/DMXFixtures/Source/DMXFixtures/Public/DMXFixtureActorMatrix.h
@@ -7,7 +7,7 @@
 #include "Rendering/Texture2DResource.h"
 #include "DMXFixtureActorMatrix.generated.h"
 
-UCLASS(HideCategories = ("Rendering", "Variable", "Input", "Tags", "Activation", "Cooking", "Replication", "AssetUserData", "Collision", "LOD", "Actor", "HLOD"))
+UCLASS()
 class DMXFIXTURES_API ADMXFixtureActorMatrix : public ADMXFixtureActor
 {
 	GENERATED_BODY()
diff --git a/Engine/Plugins/VirtualProduction/DMX/DMXFixtures/Source/DMXFixtures/Public/DMXFixtureComponentColor.h b/Engine/Plugins/VirtualProduction/DMX/DMXFixtures/Source/DMXFixtures/Public/DMXFixtureComponentColor.h
--- a/Engine/Plugins/VirtualProduction/DMX/DMXFixtures/Source/DMXFixtures/Public/DMXFixtureComponentColor.h
+++ b/Engine/Plugins/VirtualProduction/DMX/DMXFixtures/Source/DMXFixtures/Public/DMXFixtureComponentColor.h
@@ -7,7 +7,7 @@
 #include "DMXFixtureComponentColor.generated.h"
 
 // Specific class to handle color mixing using 4 channels (rgb, cmy, rgbw)
-UCLASS(ClassGroup = DMXFixtureComponent, meta = (IsBlueprintBase = true), HideCategories = ("Variable", "Tags", "Activation", "Cooking", "ComponentReplication", "AssetUserData", "Collision", "Sockets"))
+UCLASS(ClassGroup = DMXFixtureComponent, meta = (IsBlueprintBase = true))
 class DMXFIXTURES_API UDMXFixtureComponentColor : public UDMXFixtureComponent
 {
 	GENERATED_BODY()
diff --git a/Engine/Plugins/VirtualProduction/DMX/DMXFixtures/Source/DMXFixtures/Public/DMXFixtureComponentDouble.h b/Engine/Plugins/VirtualProduction/DMX/DMXFixtures/Source/DMXFixtures/Public/DMXFixtureComponentDouble.h
--- a/Engine/Plugins/VirtualProduction/DMX/DMXFixtures/Source/DMXFixtures/Public/DMXFixtureComponentDouble.h
+++ b/Engine/Plugins/VirtualProduction/DMX/DMXFixtures/Source/DMXFixtures/Public/DMXFixtureComponentDouble.h
@@ -7,7 +7,7 @@
 #include "DMXFixtureComponentDouble.generated.h"
 
 // Component that uses 2 DMX channels
-UCLASS(ClassGroup = FixtureComponent, meta=(IsBlueprintBase = true), HideCategories = ("Variable", "Tags", "Activation", "Cooking", "ComponentReplication", "AssetUserData", "Collision", "Sockets"))
+UCLASS(ClassGroup = FixtureComponent, meta=(IsBlueprintBase = true))
 class DMXFIXTURES_API UDMXFixtureComponentDouble : public UDMXFixtureComponent
 {
 	GENERATED_BODY()
diff --git a/Engine/Plugins/VirtualProduction/DMX/DMXFixtures/Source/DMXFixtures/Public/DMXFixtureComponentSingle.h b/Engine/Plugins/VirtualProduction/DMX/DMXFixtures/Source/DMXFixtures/Public/DMXFixtureComponentSingle.h
--- a/Engine/Plugins/VirtualProduction/DMX/DMXFixtures/Source/DMXFixtures/Public/DMXFixtureComponentSingle.h
+++ b/Engine/Plugins/VirtualProduction/DMX/DMXFixtures/Source/DMXFixtures/Public/DMXFixtureComponentSingle.h
@@ -7,7 +7,7 @@
 #include "DMXFixtureComponentSingle.generated.h"
 
 // Component that uses 1 DMX channel
-UCLASS(ClassGroup = FixtureComponent, meta=(IsBlueprintBase = true), HideCategories = ("Variable", "Tags", "Activation", "Cooking", "ComponentReplication", "AssetUserData", "Collision", "Sockets"))
+UCLASS(ClassGroup = FixtureComponent, meta=(IsBlueprintBase = true))
 class DMXFIXTURES_API UDMXFixtureComponentSingle : public UDMXFixtureComponent
 {
 	GENERATED_BODY()
