DMX - Fix crash after copying fixture in DMX Library

diff --git a/Engine/Plugins/VirtualProduction/DMX/DMXEngine/Source/DMXEditor/Private/Customizations/DMXEntityFixturePatchDetails.cpp b/Engine/Plugins/VirtualProduction/DMX/DMXEngine/Source/DMXEditor/Private/Customizations/DMXEntityFixturePatchDetails.cpp
--- a/Engine/Plugins/VirtualProduction/DMX/DMXEngine/Source/DMXEditor/Private/Customizations/DMXEntityFixturePatchDetails.cpp
+++ b/Engine/Plugins/VirtualProduction/DMX/DMXEngine/Source/DMXEditor/Private/Customizations/DMXEntityFixturePatchDetails.cpp
@@ -123,6 +123,17 @@ void FDMXEntityFixturePatchDetails::OnParentFixtureTypeChanged(UDMXEntity* NewTe
 
 void FDMXEntityFixturePatchDetails::OnFixtureTypeChanged(const UDMXEntityFixtureType* FixtureType)
 {
+	// Keep the active mode valid
+	int32 ActiveMode;
+	if (ActiveModeHandle->GetValue(ActiveMode) == FPropertyAccess::Success)
+	{
+		if (!FixtureType->Modes.IsValidIndex(ActiveMode))
+		{
+			const int32 NewActiveMode = FixtureType->Modes.Num() > 0 ? 0 : INDEX_NONE;
+			ActiveModeHandle->SetValue(NewActiveMode);
+		}
+	}
+
 	PropertyUtilities->ForceRefresh();
 }
 
diff --git a/Engine/Plugins/VirtualProduction/DMX/DMXEngine/Source/DMXEditor/Private/DMXEditorUtils.cpp b/Engine/Plugins/VirtualProduction/DMX/DMXEngine/Source/DMXEditor/Private/DMXEditorUtils.cpp
--- a/Engine/Plugins/VirtualProduction/DMX/DMXEngine/Source/DMXEditor/Private/DMXEditorUtils.cpp
+++ b/Engine/Plugins/VirtualProduction/DMX/DMXEngine/Source/DMXEditor/Private/DMXEditorUtils.cpp
@@ -32,31 +32,37 @@ struct FDMXEntityObjectTextFactory : public FCustomizableTextObjectFactory
 	/** Entities instantiated */
 	TArray<UDMXEntity*> NewEntities;
 
-	/** Constructs a new object factory from the given text buffer */
-	static TSharedRef<FDMXEntityObjectTextFactory> Get(const FString& InTextBuffer)
+	static bool CanCreate(const FString& InTextBuffer)
 	{
-		// Construct a new instance
-		TSharedPtr<FDMXEntityObjectTextFactory> FactoryPtr = MakeShareable(new FDMXEntityObjectTextFactory());
-		check(FactoryPtr.IsValid());
+		TSharedRef<FDMXEntityObjectTextFactory> Factory = MakeShareable(new FDMXEntityObjectTextFactory());
 
 		// Create new objects if we're allowed to
-		if (FactoryPtr->CanCreateObjectsFromText(InTextBuffer))
+		return Factory->CanCreateObjectsFromText(InTextBuffer);
+	}
+
+	/** Constructs a new object factory from the given text buffer. Returns the factor or nullptr if no factory can be created */
+	static TSharedPtr<FDMXEntityObjectTextFactory> Create(const FString& InTextBuffer, UDMXLibrary* InParentLibrary)
+	{
+		TSharedRef<FDMXEntityObjectTextFactory> Factory = MakeShareable(new FDMXEntityObjectTextFactory());
+
+		// Create new objects if we're allowed to
+		if (IsValid(InParentLibrary) && Factory->CanCreateObjectsFromText(InTextBuffer))
 		{
 			EObjectFlags ObjectFlags = RF_Transactional;
 
-			// Use the transient package initially for creating the objects, since the variable name is used when copying
-			FactoryPtr->ProcessBuffer(GetTransientPackage(), ObjectFlags, InTextBuffer);
+			Factory->ProcessBuffer(InParentLibrary, ObjectFlags, InTextBuffer);
+
+			return Factory;
 		}
 
-		return FactoryPtr.ToSharedRef();
+		return nullptr;
 	}
 
 protected:
 	/** Constructor; protected to only allow this class to instance itself */
 	FDMXEntityObjectTextFactory()
 		: FCustomizableTextObjectFactory(GWarn)
-	{
-	}
+	{}
 
 	//~ Begin FCustomizableTextObjectFactory implementation
 	virtual bool CanCreateClass(UClass* ObjectClass, bool& bOmitSubObjs) const override
@@ -71,15 +77,6 @@ protected:
 
 		if (UDMXEntity* NewEntity = Cast<UDMXEntity>(NewObject))
 		{
-			// If this is a Fixture Type and the first object was a Patch, we
-			// don't add the Type to the array. SDMXEntityList will deal with it
-			// from the reference on the Patch(es)
-			if (NewEntities.Num() > 0 && NewEntities[0]->GetClass()->IsChildOf<UDMXEntityFixturePatch>()
-				&& NewEntity->GetClass()->IsChildOf<UDMXEntityFixtureType>())
-			{
-				return;
-			}
-
 			NewEntities.Add(NewEntity);
 		}
 	}
@@ -425,67 +422,8 @@ void FDMXEditorUtils::CopyEntities(const TArray<UDMXEntity*>&& EntitiesToCopy)
 	// Export the component object(s) to text for copying
 	for (UDMXEntity* Entity : EntitiesToCopy)
 	{
-		check(Entity && Entity->GetParentLibrary());
-
-		// Fixture Patches require copying their template because it's a reference to a private object
-		if (UDMXEntityFixturePatch* AsPatch = Cast<UDMXEntityFixturePatch>(Entity))
-		{
-			if (UDMXEntityFixtureType* FixtureType = AsPatch->GetFixtureType())
-			{
-				bool bExportedTemplate = false;
-
-				// Try to get a cached duplicate of the template
-				UDMXEntityFixtureType* DuplicateFixtureType;
-				if (UDMXEntityFixtureType** CachedTemplate = CopiedPatchTemplates.Find(FixtureType->GetFName()))
-				{
-					DuplicateFixtureType = *CachedTemplate;
-					bExportedTemplate = true;
-				}
-				else
-				{
-					// Copy the template to the transient package to make the Patch reference the copy
-					FObjectDuplicationParameters DuplicationParams(FixtureType, GetTransientPackage());
-					DuplicationParams.DestName = FixtureType->GetFName();
-
-					DuplicateFixtureType = CastChecked<UDMXEntityFixtureType>(StaticDuplicateObjectEx(DuplicationParams));
-					// Keep same entity ID to find the original Template when pasting
-					DuplicateFixtureType->ReplicateID(FixtureType);
-
-					// Cache this copy so we don't copy the same template over and over for several Patches
-					CopiedPatchTemplates.Add(DuplicateFixtureType->GetFName(), DuplicateFixtureType);
-				}
-
-				// We'll temporarily change the ParentFixtureTypeTemplate of the Patch to copy it
-				// with a reference to the duplicate Fixture Type
-				UDMXEntityFixtureType* OriginalTemplate = FixtureType;
-				FixtureType = DuplicateFixtureType;
-				// Export the Patch referencing the duplicate template
-				UExporter::ExportToOutputDevice(&Context, AsPatch, nullptr, Archive, TEXT("copy"), 0, PPF_ExportsNotFullyQualified | PPF_Copy | PPF_Delimited, false, GetTransientPackage());
-				if (!bExportedTemplate)
-				{
-					// Export the template after the Patch, to make interpretation easier when pasting it back
-					UExporter::ExportToOutputDevice(&Context, DuplicateFixtureType, nullptr, Archive, TEXT("copy"), 4, PPF_ExportsNotFullyQualified | PPF_Copy | PPF_Delimited, false, GetTransientPackage());
-				}
-				// Revert the patch to it's original, private template
-				FixtureType = OriginalTemplate;
-			}
-			else // Template is null
-			{
-				// Export the entity object to the given string
-				UExporter::ExportToOutputDevice(&Context, AsPatch, nullptr, Archive, TEXT("copy"), 0, PPF_ExportsNotFullyQualified | PPF_Copy | PPF_Delimited, false, GetTransientPackage());
-			}
-		}
-		else
-		{
-			// Export the entity object to the given string
-			UExporter::ExportToOutputDevice(&Context, Entity, nullptr, Archive, TEXT("copy"), 0, PPF_ExportsNotFullyQualified | PPF_Copy | PPF_Delimited, false, GetTransientPackage());
-		}
-	}
-
-	// Speed up the deletion of the copies. We don't need them anymore
-	for (const TPair<FName, UDMXEntityFixtureType*>& CopiedTemplate : CopiedPatchTemplates)
-	{
-		CopiedTemplate.Value->ConditionalBeginDestroy();
+		// Export the entity object to the given string
+		UExporter::ExportToOutputDevice(&Context, Entity, nullptr, Archive, TEXT("copy"), 0, PPF_ExportsNotFullyQualified | PPF_Copy | PPF_Delimited, false, GetTransientPackage());
 	}
 
 	// Copy text to clipboard
@@ -493,27 +431,34 @@ void FDMXEditorUtils::CopyEntities(const TArray<UDMXEntity*>&& EntitiesToCopy)
 	FPlatformApplicationMisc::ClipboardCopy(*ExportedText);
 }
 
-bool FDMXEditorUtils::CanPasteEntities()
+bool FDMXEditorUtils::CanPasteEntities(UDMXLibrary* ParentLibrary)
 {
 	FString ClipboardContent;
 	FPlatformApplicationMisc::ClipboardPaste(ClipboardContent);
 
 	// Obtain the entity object text factory for the clipboard content and return whether or not we can use it
-	TSharedRef<FDMXEntityObjectTextFactory> Factory = FDMXEntityObjectTextFactory::Get(ClipboardContent);
-	return Factory->NewEntities.Num() > 0;
+	return FDMXEntityObjectTextFactory::CanCreate(ClipboardContent);
 }
 
 void FDMXEditorUtils::GetEntitiesFromClipboard(TArray<UDMXEntity*>& OutNewObjects)
+{
+	// DEPRECATED 5.0, no longer creates a meaningful result
+	OutNewObjects = TArray<UDMXEntity*>();
+}
+
+TArray<UDMXEntity*> FDMXEditorUtils::CreateEntitiesFromClipboard(UDMXLibrary* ParentLibrary)
 {
 	// Get the text from the clipboard
 	FString TextToImport;
 	FPlatformApplicationMisc::ClipboardPaste(TextToImport);
 
 	// Get a new component object factory for the clipboard content
-	TSharedRef<FDMXEntityObjectTextFactory> Factory = FDMXEntityObjectTextFactory::Get(TextToImport);
+	if (TSharedPtr<FDMXEntityObjectTextFactory> Factory = FDMXEntityObjectTextFactory::Create(TextToImport, ParentLibrary))
+	{
+		return Factory->NewEntities;
+	}
 
-	// Return the created component mappings
-	OutNewObjects = MoveTemp(Factory->NewEntities);
+	return TArray<UDMXEntity*>();
 }
 
 bool FDMXEditorUtils::AreFixtureTypesIdentical(const UDMXEntityFixtureType* A, const UDMXEntityFixtureType* B)
diff --git a/Engine/Plugins/VirtualProduction/DMX/DMXEngine/Source/DMXEditor/Private/Widgets/FixturePatch/SDMXFixturePatchTree.cpp b/Engine/Plugins/VirtualProduction/DMX/DMXEngine/Source/DMXEditor/Private/Widgets/FixturePatch/SDMXFixturePatchTree.cpp
--- a/Engine/Plugins/VirtualProduction/DMX/DMXEngine/Source/DMXEditor/Private/Widgets/FixturePatch/SDMXFixturePatchTree.cpp
+++ b/Engine/Plugins/VirtualProduction/DMX/DMXEngine/Source/DMXEditor/Private/Widgets/FixturePatch/SDMXFixturePatchTree.cpp
@@ -461,15 +461,14 @@ bool SDMXFixturePatchTree::CanCopyNodes() const
 
 void SDMXFixturePatchTree::OnPasteNodes()
 {
+	// Get the library that's being edited
+	UDMXLibrary* Library = GetDMXLibrary();
+	check(Library);
+
 	// Get the Entities to paste from the clipboard
-	TArray<UDMXEntity*> NewObjects;
-	FDMXEditorUtils::GetEntitiesFromClipboard(NewObjects);
+	TArray<UDMXEntity*> NewObjects = FDMXEditorUtils::CreateEntitiesFromClipboard(Library);
 	if (NewObjects.Num() != 0)
 	{
-		// Get the library that's being edited
-		UDMXLibrary* Library = GetDMXLibrary();
-		check(Library);
-
 		// Start transaction for Undo and take a snapshot of the current Library state
 		const FScopedTransaction PasteEntities(NewObjects.Num() > 1 ? LOCTEXT("PasteFixturePatches", "Paste Fixture Patches") : LOCTEXT("PasteFixturePatch", "Paste Fixture Patch"));
 		Library->Modify();
@@ -478,34 +477,37 @@ void SDMXFixturePatchTree::OnPasteNodes()
 		TArray<TWeakObjectPtr<UDMXEntityFixturePatch>> NewFixturePatches;
 		for (UDMXEntity* NewEntity : NewObjects)
 		{
-			UDMXEntityFixturePatch* FixturePatch = CastChecked<UDMXEntityFixturePatch>(NewEntity);
-			NewFixturePatches.Add(FixturePatch);
+			if (UDMXEntityFixturePatch* FixturePatch = Cast<UDMXEntityFixturePatch>(NewEntity))
+			{
+				NewFixturePatches.Add(FixturePatch);
 
-			// Search for a suitable replacement for the pasted Fixture Type, with identical
-			// properties, except for the Name, ID and Parent Library
-			const TArray<UDMXEntityFixturePatch*> FixturePatches = Library->GetEntitiesTypeCast<UDMXEntityFixturePatch>();
-			const UDMXEntityFixturePatch* const* SourceFixturePatchPtr = FixturePatches.FindByPredicate([FixturePatch](const UDMXEntityFixturePatch* SourceFixturePatch)
+				// Search for a suitable replacement for the pasted Fixture Type, with identical
+				// properties, except for the Name, ID and Parent Library
+				const TArray<UDMXEntityFixturePatch*> FixturePatches = Library->GetEntitiesTypeCast<UDMXEntityFixturePatch>();
+				const UDMXEntityFixturePatch* const* SourceFixturePatchPtr = FixturePatches.FindByPredicate([FixturePatch](const UDMXEntityFixturePatch* SourceFixturePatch)
+					{
+						return SourceFixturePatch->GetID() == FixturePatch->GetID();
+					});
+
+				if (SourceFixturePatchPtr)
 				{
-					return SourceFixturePatch->GetID() == FixturePatch->GetID();
-				});
+					UDMXEntityFixtureType* FixtureType = (*SourceFixturePatchPtr)->GetFixtureType();
+					FixturePatch->SetFixtureType(FixtureType);
+				}
 
-			if (SourceFixturePatchPtr)
-			{
-				FixturePatch->SetFixtureType((*SourceFixturePatchPtr)->GetFixtureType());
-			}
+				// Move the Fixture Type template from the transient package into the Library package
+				FixturePatch->Rename(*MakeUniqueObjectName(Library, UDMXEntityFixtureType::StaticClass()).ToString(), Library, REN_DoNotDirty | REN_DontCreateRedirectors);
+				FixturePatch->RefreshID();
 
-			// Move the Fixture Type template from the transient package into the Library package
-			NewEntity->Rename(*MakeUniqueObjectName(Library, UDMXEntityFixtureType::StaticClass()).ToString(), Library, REN_DoNotDirty | REN_DontCreateRedirectors);
-			NewEntity->RefreshID();
+				// Add to the Library
+				AutoAssignCopiedPatch(FixturePatch);
 
-			// Add to the Library
-			AutoAssignCopiedPatch(FixturePatch);
+				// Move the Entity from the transient package into the Library package
+				FixturePatch->Rename(*MakeUniqueObjectName(Library, FixturePatch->GetClass()).ToString(), Library, REN_DoNotDirty | REN_DontCreateRedirectors);
 
-			// Move the Entity from the transient package into the Library package
-			NewEntity->Rename(*MakeUniqueObjectName(Library, NewEntity->GetClass()).ToString(), Library, REN_DoNotDirty | REN_DontCreateRedirectors);
-			
-			// Make sure the Entity's name won't collide with existing ones
-			NewEntity->SetName(FDMXRuntimeUtils::FindUniqueEntityName(Library, NewEntity->GetClass(), NewEntity->GetDisplayName()));
+				// Make sure the Entity's name won't collide with existing ones
+				FixturePatch->SetName(FDMXRuntimeUtils::FindUniqueEntityName(Library, FixturePatch->GetClass(), FixturePatch->GetDisplayName()));
+			}
 		}
 
 		FixturePatchSharedData->SelectFixturePatches(NewFixturePatches);
@@ -516,7 +518,10 @@ void SDMXFixturePatchTree::OnPasteNodes()
 
 bool SDMXFixturePatchTree::CanPasteNodes() const
 {
-	return FDMXEditorUtils::CanPasteEntities();
+	UDMXLibrary* Library = GetDMXLibrary();
+	check(IsValid(Library));
+
+	return FDMXEditorUtils::CanPasteEntities(Library);
 }
 
 void SDMXFixturePatchTree::OnDuplicateNodes()
diff --git a/Engine/Plugins/VirtualProduction/DMX/DMXEngine/Source/DMXEditor/Private/Widgets/FixtureType/SDMXFixtureTypeTree.cpp b/Engine/Plugins/VirtualProduction/DMX/DMXEngine/Source/DMXEditor/Private/Widgets/FixtureType/SDMXFixtureTypeTree.cpp
--- a/Engine/Plugins/VirtualProduction/DMX/DMXEngine/Source/DMXEditor/Private/Widgets/FixtureType/SDMXFixtureTypeTree.cpp
+++ b/Engine/Plugins/VirtualProduction/DMX/DMXEngine/Source/DMXEditor/Private/Widgets/FixtureType/SDMXFixtureTypeTree.cpp
@@ -113,7 +113,7 @@ void SDMXFixtureTypeTree::RebuildNodes(const TSharedPtr<FDMXEntityTreeRootNode>&
 	EntityNodeToEntityRowMap.Reset();
 
 	UDMXLibrary* Library = GetDMXLibrary();
-	check(Library);
+	check(IsValid(Library));
 
 	Library->ForEachEntityOfType<UDMXEntityFixtureType>([this](UDMXEntityFixtureType* FixtureType)
 		{
@@ -284,16 +284,15 @@ bool SDMXFixtureTypeTree::CanCopyNodes() const
 
 void SDMXFixtureTypeTree::OnPasteNodes()
 {
+	// Get the library that's being edited
+	UDMXLibrary* Library = GetDMXLibrary();
+	check(IsValid(Library));
+
 	// Get the Entities to paste from the clipboard
-	TArray<UDMXEntity*> NewObjects;
-	FDMXEditorUtils::GetEntitiesFromClipboard(NewObjects);
+	TArray<UDMXEntity*> NewObjects = FDMXEditorUtils::CreateEntitiesFromClipboard(Library);
 
 	if (NewObjects.Num() != 0)
 	{
-		// Get the library that's being edited
-		UDMXLibrary* Library = GetDMXLibrary();
-		check(Library);
-
 		// Start transaction for Undo and take a snapshot of the current Library state
 		const FScopedTransaction PasteEntities(NewObjects.Num() > 1 ? LOCTEXT("PasteFixtureTypes", "Paste Fixture Types") : LOCTEXT("PasteFixtureType", "Paste Fixture Type"));
 		Library->Modify();
@@ -321,12 +320,17 @@ void SDMXFixtureTypeTree::OnPasteNodes()
 
 bool SDMXFixtureTypeTree::CanPasteNodes() const
 {
-	return FDMXEditorUtils::CanPasteEntities();
+	UDMXLibrary* Library = GetDMXLibrary();
+	check(IsValid(Library));
+
+	return FDMXEditorUtils::CanPasteEntities(Library);
 }
 
 void SDMXFixtureTypeTree::OnDuplicateNodes()
 {
 	UDMXLibrary* Library = GetDMXLibrary();
+	check(IsValid(Library));
+
 	TArray<UDMXEntity*> SelectedEntities = GetSelectedEntities();
 
 	if (Library && SelectedEntities.Num() > 0)
diff --git a/Engine/Plugins/VirtualProduction/DMX/DMXEngine/Source/DMXEditor/Public/DMXEditorUtils.h b/Engine/Plugins/VirtualProduction/DMX/DMXEngine/Source/DMXEditor/Public/DMXEditorUtils.h
--- a/Engine/Plugins/VirtualProduction/DMX/DMXEngine/Source/DMXEditor/Public/DMXEditorUtils.h
+++ b/Engine/Plugins/VirtualProduction/DMX/DMXEngine/Source/DMXEditor/Public/DMXEditorUtils.h
@@ -55,13 +55,19 @@ public:
 	static void CopyEntities(const TArray<UDMXEntity*>&& EntitiesToCopy);
 
 	/**  Determines whether the current contents of the clipboard contain paste-able DMX Entity information */
-	static bool CanPasteEntities();
+	static bool CanPasteEntities(UDMXLibrary* ParentLibrary);
+
+	/** DEPRECATED 5.0 */
+	UE_DEPRECATED(5.0, "Replaced with FDMXEditorUtils::CreateEntitiesFromClipboard. NOTE: GetEntitiesFromClipboard no longer returns any entities, as entities can only be created within a DMXLibrary.")
+	static void GetEntitiesFromClipboard(TArray<UDMXEntity*>& OutNewObjects);
 
 	/**
-	 * Gets the copied DMX Entities from the clipboard without attempting to paste/apply them in any way
-	 * @param OutNewObjectMap			Contains the name->instance object mapping of the copied DMX Entities
+	 * Creates the copied DMX Entities from the clipboard without attempting to paste/apply them in any way
+	 * 
+	 * @param ParentLibrary			The library in which the entities are created.
+	 * @return						The array of newly created enities.
 	 */
-	static void GetEntitiesFromClipboard(TArray<UDMXEntity*>& OutNewObjects);
+	static TArray<UDMXEntity*> CreateEntitiesFromClipboard(UDMXLibrary* ParentLibrary);
 
 	/**
 	 * Compares the property values of two Fixture Types, including properties in arrays,
diff --git a/Engine/Plugins/VirtualProduction/DMX/DMXEngine/Source/DMXRuntime/Private/Library/DMXEntityFixturePatch.cpp b/Engine/Plugins/VirtualProduction/DMX/DMXEngine/Source/DMXRuntime/Private/Library/DMXEntityFixturePatch.cpp
--- a/Engine/Plugins/VirtualProduction/DMX/DMXEngine/Source/DMXRuntime/Private/Library/DMXEntityFixturePatch.cpp
+++ b/Engine/Plugins/VirtualProduction/DMX/DMXEngine/Source/DMXRuntime/Private/Library/DMXEntityFixturePatch.cpp
@@ -365,16 +365,16 @@ void UDMXEntityFixturePatch::SetFixtureType(UDMXEntityFixtureType* NewFixtureTyp
 {
 	Modify();
 
-	if (NewFixtureType && NewFixtureType != ParentFixtureTypeTemplate)
+	if (!IsValid(NewFixtureType))
+	{
+		ParentFixtureTypeTemplate = nullptr;
+	}
+	else if (NewFixtureType && NewFixtureType != ParentFixtureTypeTemplate)
 	{
 		ParentFixtureTypeTemplate = NewFixtureType;
 		
 		ActiveMode = ParentFixtureTypeTemplate->Modes.Num() > 0 ? 0 : INDEX_NONE;
 	}
-	else
-	{
-		ParentFixtureTypeTemplate = nullptr;
-	}
 
 	RebuildCache();
 }
