DMX - DMX Library Patch editor now expands universes in a more user friendly way

diff --git a/Engine/Plugins/VirtualProduction/DMX/DMXEngine/Source/DMXEditor/Private/Widgets/FixturePatch/SDMXFixturePatchTree.cpp b/Engine/Plugins/VirtualProduction/DMX/DMXEngine/Source/DMXEditor/Private/Widgets/FixturePatch/SDMXFixturePatchTree.cpp
--- a/Engine/Plugins/VirtualProduction/DMX/DMXEngine/Source/DMXEditor/Private/Widgets/FixturePatch/SDMXFixturePatchTree.cpp
+++ b/Engine/Plugins/VirtualProduction/DMX/DMXEngine/Source/DMXEditor/Private/Widgets/FixturePatch/SDMXFixturePatchTree.cpp
@@ -47,6 +47,7 @@ void SDMXFixturePatchTree::Construct(const FArguments& InArgs)
 
 		// Bind to selection changes
 		FixturePatchSharedData->OnFixturePatchSelectionChanged.AddSP(this, &SDMXFixturePatchTree::OnFixturePatchesSelected);
+		FixturePatchSharedData->OnUniverseSelectionChanged.AddSP(this, &SDMXFixturePatchTree::OnUniverseSelected);
 
 		// Bind to library changes
 		PinnedDMXEditor->GetDMXLibrary()->GetOnEntitiesAdded().AddSP(this, &SDMXFixturePatchTree::OnEntitiesAddedOrRemoved);
@@ -56,71 +57,15 @@ void SDMXFixturePatchTree::Construct(const FArguments& InArgs)
 		UDMXEntityFixturePatch::GetOnFixturePatchChanged().AddSP(this, &SDMXFixturePatchTree::OnFixturePatchChanged);
 
 		// Make an initial selection
-		TArray<UDMXEntityFixturePatch*> FixturePatches = GetDMXLibrary()->GetEntitiesTypeCast<UDMXEntityFixturePatch>();
+		const TArray<UDMXEntityFixturePatch*> FixturePatches = GetDMXLibrary()->GetEntitiesTypeCast<UDMXEntityFixturePatch>();
 		if (FixturePatches.Num() > 0)
 		{
-			TArray<TWeakObjectPtr<UDMXEntityFixturePatch>> InitialSelection = { FixturePatches[0] };
+			const TArray<TWeakObjectPtr<UDMXEntityFixturePatch>> InitialSelection = { FixturePatches[0] };
 			FixturePatchSharedData->SelectFixturePatches(InitialSelection);
 		}
 	}
 }
 
-TSharedPtr<FDMXEntityTreeEntityNode> SDMXFixturePatchTree::CreateEntityNode(UDMXEntity* Entity)
-{
-	check(Entity && Entity->GetClass() == UDMXEntityFixturePatch::StaticClass());
-
-	TSharedPtr<FDMXEntityTreeEntityNode> NewNode = MakeShared<FDMXEntityTreeEntityNode>(Entity);
-	RefreshFilteredState(NewNode, false);
-
-	// Error status
-	FText InvalidReason;
-	if (!Entity->IsValidEntity(InvalidReason))
-	{
-		NewNode->SetErrorStatus(InvalidReason);
-	}
-
-	return NewNode;
-}
-
-void SDMXFixturePatchTree::OnEntitiesAddedOrRemoved(UDMXLibrary* DMXLibrary, TArray<UDMXEntity*> Entities)
-{
-	if (DMXLibrary && AddButtonDropdownList.IsValid())
-	{
-		AddButtonDropdownList->RefreshEntitiesList();
-		UpdateTree();
-	}
-}
-
-void SDMXFixturePatchTree::OnFixturePatchChanged(const UDMXEntityFixturePatch* FixturePatch)
-{
-	if (TSharedPtr<FDMXEditor> PinnedDMXEditor = DMXEditor.Pin())
-	{
-		if (PinnedDMXEditor->GetDMXLibrary() == FixturePatch->GetParentLibrary())
-		{
-			UpdateTree();
-		}
-	}
-}
-
-void SDMXFixturePatchTree::OnFixturePatchesSelected()
-{
-	if (!bChangingSelection)
-	{
-		TArray<TWeakObjectPtr<UDMXEntityFixturePatch>> SelectedFixturePatches = FixturePatchSharedData->GetSelectedFixturePatches();
-
-		TArray<UDMXEntity*> NewSelection;
-		for (TWeakObjectPtr<UDMXEntityFixturePatch> SelectedFixturePatch : SelectedFixturePatches)
-		{
-			if (UDMXEntityFixturePatch* SelectedEntity = SelectedFixturePatch.Get())
-			{
-				NewSelection.Add(SelectedEntity);
-			}
-		}
-
-		SelectItemsByEntities(NewSelection);
-	}
-}
-
 TSharedRef<SWidget> SDMXFixturePatchTree::GenerateAddNewEntityButton()
 {
 	// Top part, with the  [+ Add New] button and the filter box
@@ -194,7 +139,7 @@ void SDMXFixturePatchTree::RebuildNodes(const TSharedPtr<FDMXEntityTreeRootNode>
 					{
 						constexpr FDMXEntityTreeCategoryNode::ECategoryType CategoryType = FDMXEntityTreeCategoryNode::ECategoryType::UniverseID;
 
-						TSharedPtr<FDMXEntityTreeCategoryNode> UniverseCategoryNode = FixturePatchTree->GetOrCreateCategoryNode(
+						TSharedRef<FDMXEntityTreeCategoryNode> UniverseCategoryNode = FixturePatchTree->GetOrCreateCategoryNode(
 							CategoryType,
 							FText::Format(LOCTEXT("UniverseSubcategoryLabel", "Universe {0}"),	FText::AsNumber(FixturePatch->GetUniverseID())),
 							FixturePatch->GetUniverseID(),
@@ -202,6 +147,28 @@ void SDMXFixturePatchTree::RebuildNodes(const TSharedPtr<FDMXEntityTreeRootNode>
 						);
 
 						UniverseCategoryNode->AddChild(FixturePatchNode);
+
+						// Retain expansion states
+						const int32 UserExpandedCategoryIndex = FixturePatchTree->UserExpandedUniverseCategoryNodes.IndexOfByPredicate([UniverseCategoryNode](const TSharedPtr<FDMXEntityTreeCategoryNode>& OtherUniverseCategoryNode)
+							{
+								return UniverseCategoryNode->GetIntValue() == OtherUniverseCategoryNode->GetIntValue();
+							});
+
+						if (UserExpandedCategoryIndex != INDEX_NONE)
+						{
+							FixturePatchTree->UserExpandedUniverseCategoryNodes[UserExpandedCategoryIndex] = UniverseCategoryNode;
+							FixturePatchTree->SetNodeExpansion(FixturePatchTree->UserExpandedUniverseCategoryNodes[UserExpandedCategoryIndex], true);
+						}
+						else if (!FixturePatchTree->AutoExpandedUniverseCategoryNode.IsValid() ||
+							FixturePatchTree->AutoExpandedUniverseCategoryNode->GetIntValue() == UniverseCategoryNode->GetIntValue())
+						{
+							FixturePatchTree->SetNodeExpansion(UniverseCategoryNode, true);
+							FixturePatchTree->AutoExpandedUniverseCategoryNode = UniverseCategoryNode;
+						}
+						else
+						{
+							FixturePatchTree->SetNodeExpansion(UniverseCategoryNode, false);
+						}
 					}
 					else
 					{
@@ -303,8 +270,8 @@ void SDMXFixturePatchTree::RebuildNodes(const TSharedPtr<FDMXEntityTreeRootNode>
 	
 	InRootNode->AddChild(AssignedFixturesCategoryNode);
 	InRootNode->AddChild(UnassignedFixturesCategoryNode);
-	EntitiesTreeWidget->SetItemExpansion(AssignedFixturesCategoryNode, true);
-	EntitiesTreeWidget->SetItemExpansion(UnassignedFixturesCategoryNode, true);
+	SetNodeExpansion(AssignedFixturesCategoryNode, true);
+	SetNodeExpansion(UnassignedFixturesCategoryNode, true);
 	
 	RefreshFilteredState(AssignedFixturesCategoryNode, false);
 	RefreshFilteredState(UnassignedFixturesCategoryNode, false);
@@ -329,12 +296,12 @@ void SDMXFixturePatchTree::RebuildNodes(const TSharedPtr<FDMXEntityTreeRootNode>
 	SelectItemsByEntities(SelectedEntities);
 }
 
-TSharedRef<ITableRow> SDMXFixturePatchTree::OnGenerateRow(TSharedPtr<FDMXEntityTreeNodeBase> InNodePtr, const TSharedRef<STableViewBase>& OwnerTable)
+TSharedRef<ITableRow> SDMXFixturePatchTree::OnGenerateRow(TSharedPtr<FDMXEntityTreeNodeBase> InNode, const TSharedRef<STableViewBase>& OwnerTable)
 {
 	// Create the node of the appropriate type
-	if (InNodePtr->GetNodeType() == FDMXEntityTreeNodeBase::ENodeType::CategoryNode)
+	if (InNode->GetNodeType() == FDMXEntityTreeNodeBase::ENodeType::CategoryNode)
 	{
-		TSharedPtr<FDMXEntityTreeCategoryNode> CategoryNode = StaticCastSharedPtr<FDMXEntityTreeCategoryNode>(InNodePtr);
+		TSharedPtr<FDMXEntityTreeCategoryNode> CategoryNode = StaticCastSharedPtr<FDMXEntityTreeCategoryNode>(InNode);
 		const bool bIsRootCategory = CategoryNode->GetCategoryType() != FDMXEntityTreeCategoryNode::ECategoryType::UniverseID;
 
 		return 
@@ -342,13 +309,13 @@ TSharedRef<ITableRow> SDMXFixturePatchTree::OnGenerateRow(TSharedPtr<FDMXEntityT
 			.OnFixturePatchOrderChanged(OnEntityOrderChangedDelegate)
 			[
 				SNew(STextBlock)
-				.Text(InNodePtr->GetDisplayNameText())
+				.Text(InNode->GetDisplayNameText())
 				.TextStyle(FEditorStyle::Get(), "DetailsView.CategoryTextStyle")
 			];
 	}
 	else
 	{
-		TSharedPtr<FDMXEntityTreeEntityNode> EntityNode = StaticCastSharedPtr<FDMXEntityTreeEntityNode>(InNodePtr);
+		TSharedPtr<FDMXEntityTreeEntityNode> EntityNode = StaticCastSharedPtr<FDMXEntityTreeEntityNode>(InNode);
 
 		TSharedRef<SDMXFixturePatchTreeFixturePatchRow> FixturePatchRow = SNew(SDMXFixturePatchTreeFixturePatchRow, EntityNode, OwnerTable, SharedThis(this))
 			.OnGetFilterText(this, &SDMXEntityTreeViewBase::GetFilterText)
@@ -362,6 +329,39 @@ TSharedRef<ITableRow> SDMXFixturePatchTree::OnGenerateRow(TSharedPtr<FDMXEntityT
 	}
 }
 
+void SDMXFixturePatchTree::OnExpansionChanged(TSharedPtr<FDMXEntityTreeNodeBase> Node, bool bInExpansionState)
+{
+	SDMXEntityTreeViewBase::OnExpansionChanged(Node, bInExpansionState);
+
+	if (Node.IsValid())
+	{
+		if (bInExpansionState &&
+			Node->GetNodeType() == FDMXEntityTreeNodeBase::ENodeType::CategoryNode &&
+			Node != AutoExpandedUniverseCategoryNode)
+		{
+			TSharedRef<FDMXEntityTreeCategoryNode> CategoryNode = StaticCastSharedRef<FDMXEntityTreeCategoryNode>(Node.ToSharedRef());
+			if (CategoryNode->GetCategoryType() == FDMXEntityTreeCategoryNode::ECategoryType::UniverseID)
+			{
+				UserExpandedUniverseCategoryNodes.AddUnique(CategoryNode);
+			}
+		}
+		else if (!bInExpansionState &&
+			Node->GetNodeType() == FDMXEntityTreeNodeBase::ENodeType::CategoryNode)
+		{
+			TSharedRef<FDMXEntityTreeCategoryNode> CategoryNode = StaticCastSharedRef<FDMXEntityTreeCategoryNode>(Node.ToSharedRef());
+			if (CategoryNode->GetCategoryType() == FDMXEntityTreeCategoryNode::ECategoryType::UniverseID)
+			{
+				UserExpandedUniverseCategoryNodes.RemoveSingle(CategoryNode);
+
+				if (AutoExpandedUniverseCategoryNode == CategoryNode)
+				{
+					AutoExpandedUniverseCategoryNode.Reset();
+				}
+			}
+		}
+	}
+}
+
 TSharedPtr<SWidget> SDMXFixturePatchTree::OnContextMenuOpen()
 {
 	const bool bCloseWindowAfterMenuSelection = true;
@@ -390,11 +390,11 @@ TSharedPtr<SWidget> SDMXFixturePatchTree::OnContextMenuOpen()
 	return MenuBuilder.MakeWidget();
 }
 
-void SDMXFixturePatchTree::OnSelectionChanged(TSharedPtr<FDMXEntityTreeNodeBase> InSelectedNodePtr, ESelectInfo::Type SelectInfo)
+void SDMXFixturePatchTree::OnSelectionChanged(TSharedPtr<FDMXEntityTreeNodeBase> InSelectedNode, ESelectInfo::Type SelectInfo)
 {
 	const TArray<UDMXEntity*> NewSelection = GetSelectedEntities();
 
-	TGuardValue<bool> RecursionGuard(bChangingSelection, true);
+	TGuardValue<bool> RecursionGuard(bChangingFixturePatchSelection, true);
 
 	// Never clear the selection
 	if (GetSelectedEntities().Num() == 0)
@@ -417,17 +417,40 @@ void SDMXFixturePatchTree::OnSelectionChanged(TSharedPtr<FDMXEntityTreeNodeBase>
 	else
 	{
 		// Select selected Fixture Types in Fixture Type Shared Data
-		TArray<TWeakObjectPtr<UDMXEntityFixturePatch>> SelectedFixtureTypes;
+		TArray<TWeakObjectPtr<UDMXEntityFixturePatch>> SelectedFixturePatches;
 
 		for (UDMXEntity* Entity : NewSelection)
 		{
 			if (UDMXEntityFixturePatch* FixtureType = Cast<UDMXEntityFixturePatch>(Entity))
 			{
-				SelectedFixtureTypes.Add(FixtureType);
+				SelectedFixturePatches.Add(FixtureType);
 			}
 		}
 
-		FixturePatchSharedData->SelectFixturePatches(SelectedFixtureTypes);
+		// Scroll into view
+		if (SelectedFixturePatches.Num() > 0)
+		{
+			// The weak ptrs are known to be valid here, as they were created from the strong pointers of NewSelection above
+			SelectedFixturePatches.Sort([](const TWeakObjectPtr<UDMXEntityFixturePatch>& WeakFixturePatchA, const TWeakObjectPtr<UDMXEntityFixturePatch>& WeakFixturePatchB)
+				{
+					const UDMXEntityFixturePatch* FixturePatchA = WeakFixturePatchA.Get();
+					const UDMXEntityFixturePatch* FixturePatchB = WeakFixturePatchB.Get();
+
+					const bool bUniverseIsLower = FixturePatchA->GetUniverseID() < FixturePatchB->GetUniverseID();
+					const bool bChannelIsLower =
+						FixturePatchA->GetUniverseID() == FixturePatchB->GetUniverseID() &&
+						FixturePatchA->GetStartingChannel() <= FixturePatchB->GetStartingChannel();
+
+					return bUniverseIsLower || bChannelIsLower;
+				});
+
+			if (const TSharedPtr<FDMXEntityTreeEntityNode> EntityNode = FindNodeByEntity(SelectedFixturePatches[0].Get())) 
+			{
+				RequestScrollIntoView(EntityNode);
+			}
+		}
+
+		FixturePatchSharedData->SelectFixturePatches(SelectedFixturePatches);
 	}
 }
 
@@ -600,7 +623,7 @@ bool SDMXFixturePatchTree::CanDeleteNodes() const
 
 void SDMXFixturePatchTree::OnRenameNode()
 {
-	TArray<TSharedPtr<FDMXEntityTreeNodeBase>> SelectedItems = EntitiesTreeWidget->GetSelectedItems();
+	TArray<TSharedPtr<FDMXEntityTreeNodeBase>> SelectedItems = GetSelectedNodes();
 
 	if (SelectedItems.Num() == 1 && SelectedItems[0]->GetNodeType() == FDMXEntityTreeNodeBase::ENodeType::EntityNode)
 	{
@@ -616,10 +639,27 @@ void SDMXFixturePatchTree::OnRenameNode()
 
 bool SDMXFixturePatchTree::CanRenameNode() const
 {
-	return EntitiesTreeWidget->GetSelectedItems().Num() == 1 && EntitiesTreeWidget->GetSelectedItems()[0]->CanRename();
+	return GetSelectedNodes().Num() == 1 && GetSelectedNodes()[0]->CanRename();
+}
+
+TSharedPtr<FDMXEntityTreeEntityNode> SDMXFixturePatchTree::CreateEntityNode(UDMXEntity* Entity)
+{
+	check(Entity && Entity->GetClass() == UDMXEntityFixturePatch::StaticClass());
+
+	TSharedPtr<FDMXEntityTreeEntityNode> NewNode = MakeShared<FDMXEntityTreeEntityNode>(Entity);
+	RefreshFilteredState(NewNode, false);
+
+	// Error status
+	FText InvalidReason;
+	if (!Entity->IsValidEntity(InvalidReason))
+	{
+		NewNode->SetErrorStatus(InvalidReason);
+	}
+
+	return NewNode;
 }
 
-TSharedPtr<SDMXFixturePatchTreeFixturePatchRow> SDMXFixturePatchTree::FindEntityRowByNode(const TSharedRef<FDMXEntityTreeEntityNode>& EntityNode)
+TSharedPtr<SDMXFixturePatchTreeFixturePatchRow> SDMXFixturePatchTree::FindEntityRowByNode(const TSharedRef<FDMXEntityTreeEntityNode>& EntityNode) const
 {
 	if (const TSharedRef<SDMXFixturePatchTreeFixturePatchRow>* RowPtr = EntityNodeToEntityRowMap.Find(EntityNode))
 	{
@@ -629,6 +669,105 @@ TSharedPtr<SDMXFixturePatchTreeFixturePatchRow> SDMXFixturePatchTree::FindEntity
 	return nullptr;
 }
 
+TSharedPtr<FDMXEntityTreeCategoryNode> SDMXFixturePatchTree::FindCategoryNodeByUniverseID(int32 UniverseID, TSharedPtr<FDMXEntityTreeNodeBase> StartNode) const
+{
+	// Start at root node if none was provided
+	if (!StartNode.IsValid())
+	{
+		StartNode = GetRootNode();
+	}
+
+	// Test the StartNode
+	if (StartNode.IsValid() && StartNode->GetNodeType() == FDMXEntityTreeNodeBase::ENodeType::CategoryNode)
+	{
+		TSharedPtr<FDMXEntityTreeCategoryNode> CategoryNode = StaticCastSharedPtr<FDMXEntityTreeCategoryNode>(StartNode);
+		if (CategoryNode->GetCategoryType() == FDMXEntityTreeCategoryNode::ECategoryType::UniverseID &&
+			CategoryNode->GetIntValue() == UniverseID)
+		{
+			return CategoryNode;
+		}
+	}
+
+	// Test children recursively 
+	for (const TSharedPtr<FDMXEntityTreeNodeBase>& ChildNode : StartNode->GetChildren())
+	{
+		TSharedPtr<FDMXEntityTreeCategoryNode> CategoryNode = FindCategoryNodeByUniverseID(UniverseID, ChildNode);
+		if (CategoryNode.IsValid())
+		{
+			return CategoryNode;
+		}
+	}
+
+	return nullptr;
+}
+
+void SDMXFixturePatchTree::OnEntitiesAddedOrRemoved(UDMXLibrary* DMXLibrary, TArray<UDMXEntity*> Entities)
+{
+	if (DMXLibrary && AddButtonDropdownList.IsValid())
+	{
+		AddButtonDropdownList->RefreshEntitiesList();
+		UpdateTree();
+	}
+}
+
+void SDMXFixturePatchTree::OnFixturePatchChanged(const UDMXEntityFixturePatch* FixturePatch)
+{
+	if (TSharedPtr<FDMXEditor> PinnedDMXEditor = DMXEditor.Pin())
+	{
+		if (PinnedDMXEditor->GetDMXLibrary() == FixturePatch->GetParentLibrary())
+		{
+			UpdateTree();
+		}
+	}
+}
+
+void SDMXFixturePatchTree::OnFixturePatchesSelected()
+{
+	if (!bChangingFixturePatchSelection)
+	{
+		// Apply selection
+		TArray<TWeakObjectPtr<UDMXEntityFixturePatch>> SelectedFixturePatches = FixturePatchSharedData->GetSelectedFixturePatches();
+
+		TArray<UDMXEntity*> NewSelection;
+		for (TWeakObjectPtr<UDMXEntityFixturePatch> SelectedFixturePatch : SelectedFixturePatches)
+		{
+			if (UDMXEntityFixturePatch* SelectedEntity = SelectedFixturePatch.Get())
+			{
+				NewSelection.Add(SelectedEntity);
+			}
+		}
+
+		SelectItemsByEntities(NewSelection);
+
+		// Scroll into view
+		if (NewSelection.Num() > 0)
+		{
+			if (const TSharedPtr<FDMXEntityTreeEntityNode> EntityNode = FindNodeByEntity(NewSelection[0]))
+			{
+				RequestScrollIntoView(EntityNode);
+			}
+		}
+	}
+}
+
+void SDMXFixturePatchTree::OnUniverseSelected()
+{
+	const int32 SelectedUniverse = FixturePatchSharedData->GetSelectedUniverse();
+
+	if (AutoExpandedUniverseCategoryNode.IsValid() &&
+		UserExpandedUniverseCategoryNodes.Contains(AutoExpandedUniverseCategoryNode))
+	{
+		SetNodeExpansion(AutoExpandedUniverseCategoryNode, false);
+	}
+
+	AutoExpandedUniverseCategoryNode = FindCategoryNodeByUniverseID(SelectedUniverse);
+	if (AutoExpandedUniverseCategoryNode.IsValid())
+	{
+		SetNodeExpansion(AutoExpandedUniverseCategoryNode, true);
+		RequestScrollIntoView(AutoExpandedUniverseCategoryNode);
+	}
+}
+
 void SDMXFixturePatchTree::AutoAssignCopiedPatch(UDMXEntityFixturePatch* Patch) const
 {
 	check(Patch);
@@ -753,14 +892,14 @@ void SDMXFixturePatchTree::OnAddNewFixturePatchClicked(UDMXEntity* InSelectedFix
 	}
 }
 
-void SDMXFixturePatchTree::OnAutoAssignChannelStateChanged(bool NewState, TSharedPtr<FDMXEntityTreeEntityNode> InNodePtr)
+void SDMXFixturePatchTree::OnAutoAssignChannelStateChanged(bool NewState, TSharedPtr<FDMXEntityTreeEntityNode> InNode)
 {
 	const FScopedTransaction Transaction(LOCTEXT("SetAutoAssignChannelTransaction", "Set Auto Assign Channel"));
 
 	TArray<UDMXEntityFixturePatch*> ChangedPatches;
 
 	// Was the changed entity one of the selected ones?
-	if (EntitiesTreeWidget->IsItemSelected(InNodePtr))
+	if (IsNodeSelected(InNode))
 	{
 		for (UDMXEntity* Entity : GetSelectedEntities())
 		{
@@ -782,7 +921,7 @@ void SDMXFixturePatchTree::OnAutoAssignChannelStateChanged(bool NewState, TShare
 	}
 	else
 	{
-		if (UDMXEntityFixturePatch* FixturePatch = Cast<UDMXEntityFixturePatch>(InNodePtr->GetEntity()))
+		if (UDMXEntityFixturePatch* FixturePatch = Cast<UDMXEntityFixturePatch>(InNode->GetEntity()))
 		{
 			FixturePatch->Modify();
 			FixturePatch->PreEditChange(UDMXEntityFixturePatch::StaticClass()->FindPropertyByName(UDMXEntityFixturePatch::GetAutoAssignAddressPropertyNameChecked()));
diff --git a/Engine/Plugins/VirtualProduction/DMX/DMXEngine/Source/DMXEditor/Private/Widgets/FixturePatch/SDMXFixturePatchTree.h b/Engine/Plugins/VirtualProduction/DMX/DMXEngine/Source/DMXEditor/Private/Widgets/FixturePatch/SDMXFixturePatchTree.h
--- a/Engine/Plugins/VirtualProduction/DMX/DMXEngine/Source/DMXEditor/Private/Widgets/FixturePatch/SDMXFixturePatchTree.h
+++ b/Engine/Plugins/VirtualProduction/DMX/DMXEngine/Source/DMXEditor/Private/Widgets/FixturePatch/SDMXFixturePatchTree.h
@@ -36,16 +36,14 @@ public:
 	/** Destructor */
 	virtual ~SDMXFixturePatchTree() {};
 
-	/** Creates a Node for the entity */
-	TSharedPtr<FDMXEntityTreeEntityNode> CreateEntityNode(UDMXEntity* Entity);
-
 protected:
 	//~ Begin SDMXEntityTreeViewBase interface
 	virtual TSharedRef<SWidget> GenerateAddNewEntityButton();
 	virtual void RebuildNodes(const TSharedPtr<FDMXEntityTreeRootNode>& InRootNode);
 	virtual TSharedRef<ITableRow> OnGenerateRow(TSharedPtr<FDMXEntityTreeNodeBase> Node, const TSharedRef<STableViewBase>& OwnerTable);
+	virtual void OnExpansionChanged(TSharedPtr<FDMXEntityTreeNodeBase> Node, bool bInExpansionState) override;
 	virtual TSharedPtr<SWidget> OnContextMenuOpen();
-	virtual void OnSelectionChanged(TSharedPtr<FDMXEntityTreeNodeBase> InSelectedNodePtr, ESelectInfo::Type SelectInfo);
+	virtual void OnSelectionChanged(TSharedPtr<FDMXEntityTreeNodeBase> InSelectedNode, ESelectInfo::Type SelectInfo);
 	virtual void OnCutSelectedNodes();
 	virtual bool CanCutNodes() const;
 	virtual void OnCopySelectedNodes();
@@ -61,6 +59,15 @@ protected:
 	//~ End SDMXEntityTreeViewBase interface
 
 private:
+	/** Creates a Node for the entity */
+	TSharedPtr<FDMXEntityTreeEntityNode> CreateEntityNode(UDMXEntity* Entity);
+
+	/** Returns the row that corresponds to the node */
+	TSharedPtr<SDMXFixturePatchTreeFixturePatchRow> FindEntityRowByNode(const TSharedRef<FDMXEntityTreeEntityNode>& EntityNode) const;
+
+	/** Returns the category node that corresponds to the Universe ID */
+	TSharedPtr<FDMXEntityTreeCategoryNode> FindCategoryNodeByUniverseID(int32 UniverseID, TSharedPtr<FDMXEntityTreeNodeBase> StartNode = nullptr) const;
+
 	/** Called when Entities in the DMX Library were added or removed */
 	void OnEntitiesAddedOrRemoved(UDMXLibrary* DMXLibrary, TArray<UDMXEntity*> Entities);
 
@@ -70,8 +77,8 @@ private:
 	/** Called when Fixture Patches were selected in Fixture Patch Shared Data */
 	void OnFixturePatchesSelected();
 
-	/** Returns the row that corresponds to the node */
-	TSharedPtr<SDMXFixturePatchTreeFixturePatchRow> FindEntityRowByNode(const TSharedRef<FDMXEntityTreeEntityNode>& EntityNode);
+	/** Called when a Universe was selected in Fixture Patch Shared Data */
+	void OnUniverseSelected();
 
 	/** Auto assigns a patch */
 	void AutoAssignCopiedPatch(UDMXEntityFixturePatch* Patch) const;
@@ -83,10 +90,16 @@ private:
 	void OnAddNewFixturePatchClicked(UDMXEntity* InSelectedFixtureType);
 
 	/** Called when Auto Assign Channel is changed for a Patch */
-	void OnAutoAssignChannelStateChanged(bool NewState, TSharedPtr<FDMXEntityTreeEntityNode> InNodePtr);
+	void OnAutoAssignChannelStateChanged(bool NewState, TSharedPtr<FDMXEntityTreeEntityNode> InNode);
 
 	/** True while the widget is changing the selection */
-	bool bChangingSelection = false;
+	bool bChangingFixturePatchSelection = false;
+
+	/** The auto expanded universe category node */
+	TSharedPtr<FDMXEntityTreeCategoryNode> AutoExpandedUniverseCategoryNode;
+
+	/** The user expanded universe category nodes */
+	TArray<TSharedRef<FDMXEntityTreeCategoryNode>> UserExpandedUniverseCategoryNodes;
 
 	/** Map of entity nodes and their row in the tree */
 	TMap<TSharedRef<FDMXEntityTreeEntityNode>, TSharedRef<SDMXFixturePatchTreeFixturePatchRow>> EntityNodeToEntityRowMap;
diff --git a/Engine/Plugins/VirtualProduction/DMX/DMXEngine/Source/DMXEditor/Private/Widgets/FixtureType/SDMXFixtureTypeTree.cpp b/Engine/Plugins/VirtualProduction/DMX/DMXEngine/Source/DMXEditor/Private/Widgets/FixtureType/SDMXFixtureTypeTree.cpp
--- a/Engine/Plugins/VirtualProduction/DMX/DMXEngine/Source/DMXEditor/Private/Widgets/FixtureType/SDMXFixtureTypeTree.cpp
+++ b/Engine/Plugins/VirtualProduction/DMX/DMXEngine/Source/DMXEditor/Private/Widgets/FixtureType/SDMXFixtureTypeTree.cpp
@@ -149,9 +149,11 @@ void SDMXFixtureTypeTree::RebuildNodes(const TSharedPtr<FDMXEntityTreeRootNode>&
 
 TSharedRef<ITableRow> SDMXFixtureTypeTree::OnGenerateRow(TSharedPtr<FDMXEntityTreeNodeBase> Node, const TSharedRef<STableViewBase>& OwnerTable)
 {
-		// Create the node of the appropriate type
+	// Create the node of the appropriate type
 	if (Node->GetNodeType() == FDMXEntityTreeNodeBase::ENodeType::CategoryNode)
 	{
+		SetNodeExpansion(Node, true);
+
 		const TSharedPtr<FDMXEntityTreeCategoryNode> CategoryNode = StaticCastSharedPtr<FDMXEntityTreeCategoryNode>(Node);
 		constexpr bool bIsRootCategory = true;
 
@@ -400,11 +402,11 @@ bool SDMXFixtureTypeTree::CanDeleteNodes() const
 
 void SDMXFixtureTypeTree::OnRenameNode()
 {
-	TArray<TSharedPtr<FDMXEntityTreeNodeBase>> SelectedItems = EntitiesTreeWidget->GetSelectedItems();
+	TArray<TSharedPtr<FDMXEntityTreeNodeBase>> SelectedItems = GetSelectedNodes();
 
 	if (SelectedItems.Num() == 1 && SelectedItems[0].IsValid() && SelectedItems[0]->GetNodeType() == FDMXEntityTreeNodeBase::ENodeType::EntityNode)
 	{
-		EntitiesTreeWidget->RequestScrollIntoView(SelectedItems[0].ToSharedRef());
+		RequestScrollIntoView(SelectedItems[0].ToSharedRef());
 
 		const TSharedPtr<FDMXEntityTreeEntityNode> EntityNode = StaticCastSharedPtr<FDMXEntityTreeEntityNode>(SelectedItems[0]);
 		const TSharedPtr<SDMXFixtureTypeTreeFixtureTypeRow> FixtureTypeRow = FindEntityRowByNode(EntityNode.ToSharedRef());
@@ -418,7 +420,7 @@ void SDMXFixtureTypeTree::OnRenameNode()
 
 bool SDMXFixtureTypeTree::CanRenameNode() const
 {
-	TArray<TSharedPtr<FDMXEntityTreeNodeBase>> SelectedItems = EntitiesTreeWidget->GetSelectedItems();
+	TArray<TSharedPtr<FDMXEntityTreeNodeBase>> SelectedItems = GetSelectedNodes();
 
 	return 
 		SelectedItems.Num() == 1 && 
diff --git a/Engine/Plugins/VirtualProduction/DMX/DMXEngine/Source/DMXEditor/Private/Widgets/FixtureType/SDMXFixtureTypeTree.h b/Engine/Plugins/VirtualProduction/DMX/DMXEngine/Source/DMXEditor/Private/Widgets/FixtureType/SDMXFixtureTypeTree.h
--- a/Engine/Plugins/VirtualProduction/DMX/DMXEngine/Source/DMXEditor/Private/Widgets/FixtureType/SDMXFixtureTypeTree.h
+++ b/Engine/Plugins/VirtualProduction/DMX/DMXEngine/Source/DMXEditor/Private/Widgets/FixtureType/SDMXFixtureTypeTree.h
@@ -72,8 +72,6 @@ private:
 	/** Called when a Fixture Type changed */
 	void OnFixtureTypeChanged(const UDMXEntityFixtureType* FixtureType);
 
-	/** Called when the selection changed */
-
 	/** Called when the Add button was clicked */
 	FReply OnAddNewFixtureTypeClicked();
 
diff --git a/Engine/Plugins/VirtualProduction/DMX/DMXEngine/Source/DMXEditor/Private/Widgets/SDMXEntityTreeViewBase.cpp b/Engine/Plugins/VirtualProduction/DMX/DMXEngine/Source/DMXEditor/Private/Widgets/SDMXEntityTreeViewBase.cpp
--- a/Engine/Plugins/VirtualProduction/DMX/DMXEngine/Source/DMXEditor/Private/Widgets/SDMXEntityTreeViewBase.cpp
+++ b/Engine/Plugins/VirtualProduction/DMX/DMXEngine/Source/DMXEditor/Private/Widgets/SDMXEntityTreeViewBase.cpp
@@ -145,21 +145,21 @@ void SDMXEntityTreeViewBase::Tick(const FGeometry& AllottedGeometry, const doubl
 		// Clear search filter
 		FilterBox->SetText(FText::GetEmpty());
 
-		// Expand the parent nodes
-		for (const TSharedPtr<FDMXEntityTreeNodeBase>& SelectedNode : PendingSelection)
-		{
-			for (TSharedPtr<FDMXEntityTreeNodeBase> ParentNode = SelectedNode->GetParent().Pin(); ParentNode.IsValid(); ParentNode = ParentNode->GetParent().Pin())
-			{
-				EntitiesTreeWidget->SetItemExpansion(ParentNode, true);
-			}
-		}
-
 		EntitiesTreeWidget->ClearSelection();
 		EntitiesTreeWidget->SetItemSelection(PendingSelection, true, ESelectInfo::OnMouseClick);
-		EntitiesTreeWidget->RequestScrollIntoView(PendingSelection[0]);
 
 		PendingSelection.Reset();
 	}
+	
+	if (PendingExpandItemMap.Num() > 0)
+	{
+		for (const TTuple<TSharedRef<FDMXEntityTreeNodeBase>, bool>& NodeToExpansionStatePair : PendingExpandItemMap)
+		{
+			EntitiesTreeWidget->SetItemExpansion(NodeToExpansionStatePair.Key, NodeToExpansionStatePair.Value);
+		}
+
+		PendingExpandItemMap.Reset();
+	}
 }
 
 void SDMXEntityTreeViewBase::UpdateTreeInternal(bool bRegenerateTreeNodes /*= true*/)
@@ -358,21 +358,6 @@ void SDMXEntityTreeViewBase::SelectItemByName(const FString& ItemName, ESelectIn
 	}
 }
 
-TArray<TSharedPtr<FDMXEntityTreeEntityNode>> SDMXEntityTreeViewBase::GetSelectedNodes() const
-{
-	TArray<TSharedPtr<FDMXEntityTreeEntityNode>> Result;
-	TArray<TSharedPtr<FDMXEntityTreeNodeBase>> SelectedItems = EntitiesTreeWidget->GetSelectedItems();
-	for (TSharedPtr<FDMXEntityTreeNodeBase> SelectedItem : SelectedItems)
-	{
-		if (SelectedItem->GetNodeType() == FDMXEntityTreeNodeBase::ENodeType::EntityNode)
-		{
-			Result.Add(StaticCastSharedPtr<FDMXEntityTreeEntityNode>(SelectedItem));
-		}
-	}
-
-	return Result;
-}
-
 TArray<UDMXEntity*> SDMXEntityTreeViewBase::GetSelectedEntities() const
 {
 	TArray<UDMXEntity*> SelectedEntities;
@@ -393,6 +378,39 @@ TArray<UDMXEntity*> SDMXEntityTreeViewBase::GetSelectedEntities() const
 	return SelectedEntities;
 }
 
+
+bool SDMXEntityTreeViewBase::IsNodeSelected(const TSharedPtr<FDMXEntityTreeNodeBase>& Node) const
+{
+	return EntitiesTreeWidget->IsItemSelected(Node);
+}
+
+TArray<TSharedPtr<FDMXEntityTreeNodeBase>> SDMXEntityTreeViewBase::GetSelectedNodes() const
+{
+	return EntitiesTreeWidget->GetSelectedItems();
+}
+
+int32 SDMXEntityTreeViewBase::GetNumExpandedNodes() const
+{
+	TSet<TSharedPtr<FDMXEntityTreeNodeBase>> ExpandedTreeNodes;
+	EntitiesTreeWidget->GetExpandedItems(ExpandedTreeNodes);
+
+	return ExpandedTreeNodes.Num();
+}
+
+void SDMXEntityTreeViewBase::SetNodeExpansion(TSharedPtr<FDMXEntityTreeNodeBase> Node, bool bExpandItem)
+{
+	if (Node.IsValid())
+	{
+		PendingExpandItemMap.Add(Node.ToSharedRef(), bExpandItem);
+	}
+}
+
+void SDMXEntityTreeViewBase::RequestScrollIntoView(TSharedPtr<FDMXEntityTreeNodeBase> Node)
+{
+	EntitiesTreeWidget->RequestScrollIntoView(Node);
+}
+
+
 FText SDMXEntityTreeViewBase::GetFilterText() const
 {
 	return FilterBox->GetText();
@@ -507,7 +525,7 @@ void SDMXEntityTreeViewBase::OnExpansionChanged(TSharedPtr<FDMXEntityTreeNodeBas
 	{
 		Node->SetExpansionState(bInExpansionState);
 	}
-}
+}	
 
 bool SDMXEntityTreeViewBase::RefreshFilteredState(TSharedPtr<FDMXEntityTreeNodeBase> Node, bool bRecursive)
 {
@@ -557,10 +575,8 @@ void SDMXEntityTreeViewBase::OnFilterTextChanged(const FText& InFilterText)
 	// Clears selection to make UpdateTree automatically select the first visible node
 	EntitiesTreeWidget->ClearSelection();
 
-	UpdateTree(/*bRegenerateTreeNodes =*/false);
-
-	// If we reset the filter, recover nodes expansion states
-	UpdateNodesExpansion(RootNode.ToSharedRef(), GetFilterText().IsEmpty());
+	constexpr bool bRegenerateTreeNodes = false;
+	UpdateTree(bRegenerateTreeNodes);
 }
 
 void SDMXEntityTreeViewBase::GetCollapsedNodes(TSet<TSharedPtr<FDMXEntityTreeNodeBase>>& OutCollapsedNodes, TSharedPtr<FDMXEntityTreeNodeBase> InParentNodePtr /*= nullptr*/) const
@@ -586,33 +602,6 @@ void SDMXEntityTreeViewBase::GetCollapsedNodes(TSet<TSharedPtr<FDMXEntityTreeNod
 	}
 }
 
-void SDMXEntityTreeViewBase::SetNodeExpansionState(TSharedPtr<FDMXEntityTreeNodeBase> InNodeToChange, const bool bIsExpanded)
-{
-	if (EntitiesTreeWidget.IsValid() && InNodeToChange.IsValid())
-	{
-		EntitiesTreeWidget->SetItemExpansion(InNodeToChange, bIsExpanded);
-	}
-}
-
-void SDMXEntityTreeViewBase::UpdateNodesExpansion(TSharedRef<FDMXEntityTreeNodeBase> InRootNode, bool bFilterIsEmpty)
-{
-	// Only category nodes have children and need expansion
-	if (InRootNode->GetNodeType() != FDMXEntityTreeNodeBase::ENodeType::EntityNode)
-	{
-		// If the filter is not empty, all nodes should be expanded
-		const bool bExpandNodes = !bFilterIsEmpty || InRootNode->GetExpansionState();
-		EntitiesTreeWidget->SetItemExpansion(InRootNode, bExpandNodes);
-
-		for (const TSharedPtr<FDMXEntityTreeNodeBase>& Child : InRootNode->GetChildren())
-		{
-			if (Child.IsValid() && Child->GetNodeType() != FDMXEntityTreeNodeBase::ENodeType::EntityNode)
-			{
-				UpdateNodesExpansion(Child.ToSharedRef(), bFilterIsEmpty);
-			}
-		}
-	}
-}
-
 void SDMXEntityTreeViewBase::OnActiveTabChanged(TSharedPtr<SDockTab> PreviouslyActive, TSharedPtr<SDockTab> NewlyActivated)
 {
 	if (IsInTab(NewlyActivated))
diff --git a/Engine/Plugins/VirtualProduction/DMX/DMXEngine/Source/DMXEditor/Private/Widgets/SDMXEntityTreeViewBase.h b/Engine/Plugins/VirtualProduction/DMX/DMXEngine/Source/DMXEditor/Private/Widgets/SDMXEntityTreeViewBase.h
--- a/Engine/Plugins/VirtualProduction/DMX/DMXEngine/Source/DMXEditor/Private/Widgets/SDMXEntityTreeViewBase.h
+++ b/Engine/Plugins/VirtualProduction/DMX/DMXEngine/Source/DMXEditor/Private/Widgets/SDMXEntityTreeViewBase.h
@@ -70,7 +70,6 @@ public:
 				TSharedPtr<FDMXEntityTreeCategoryNode> ExistingCategoryNode = StaticCastSharedPtr<FDMXEntityTreeCategoryNode>(Node);
 				if (ExistingCategoryNode->GetCategoryType() == InCategoryType && ExistingCategoryNode->GetDisplayNameText().CompareTo(InCategoryName) == 0)
 				{
-					EntitiesTreeWidget->SetItemExpansion(Node, true);
 					return ExistingCategoryNode.ToSharedRef();
 				}
 			}
@@ -90,7 +89,7 @@ public:
 		constexpr bool bRefreshFilteredStateRecursive = false;
 		RefreshFilteredState(NewNode, bRefreshFilteredStateRecursive);
 
-		constexpr bool bExpandItem = true;
+		constexpr bool bExpandItem = false;
 		EntitiesTreeWidget->SetItemExpansion(NewNode, bExpandItem);
 		NewNode->SetExpansionState(bExpandItem);
 
@@ -118,18 +117,33 @@ public:
 	/** Selects an item by name */
 	void SelectItemByName(const FString& ItemName, ESelectInfo::Type SelectInfo = ESelectInfo::Direct);
 
-	/** Gets selected nodes */
-	TArray<TSharedPtr<FDMXEntityTreeEntityNode>> GetSelectedNodes() const;
-
 	/** Get only the valid selected entities */
 	TArray<UDMXEntity*> GetSelectedEntities() const;
 
+	/** Returns true if the node is selected */
+	bool IsNodeSelected(const TSharedPtr<FDMXEntityTreeNodeBase>& Node) const;
+
+	/** Returns the selcted nodes */
+	TArray<TSharedPtr<FDMXEntityTreeNodeBase>> GetSelectedNodes() const;
+
+	/** Returns the number of expanded nodes */
+	int32 GetNumExpandedNodes() const;
+
+	/** Sets the node to be expaned or collapsed */
+	void SetNodeExpansion(TSharedPtr<FDMXEntityTreeNodeBase> Node, bool bExpandItem);
+
+	/** Requests to scroll the node into view */
+	void RequestScrollIntoView(TSharedPtr<FDMXEntityTreeNodeBase> Node);
+
 	/** Gets current filter from the FilterBox */
 	FText GetFilterText() const;
 
 	/** Gets the DMX Library object being edited */
 	UDMXLibrary* GetDMXLibrary() const;
 
+	/** Returns the root node of the tree */
+	FORCEINLINE const TSharedPtr<FDMXEntityTreeRootNode>& GetRootNode() const { return RootNode; }
+
 	/** Handler for when an entity from the list is dragged */
 	virtual FReply OnEntitiesDragged(TSharedPtr<FDMXEntityTreeNodeBase> InNodePtr, const FPointerEvent& MouseEvent);
 
@@ -195,9 +209,6 @@ protected:
 	 */
 	bool RefreshFilteredState(TSharedPtr<FDMXEntityTreeNodeBase> Node, bool bRecursive);
 
-	/** The actual Tree widget */
-	TSharedPtr<STreeView<TSharedPtr<FDMXEntityTreeNodeBase>>> EntitiesTreeWidget;
-
 	/** Broadcast whent he selection updated. */
 	FDMXOnSelectionChanged OnSelectionChangedDelegate;
 
@@ -223,17 +234,6 @@ private:
 	/** Returns the set of expandable nodes that are currently collapsed in the UI */
 	void GetCollapsedNodes(TSet<TSharedPtr<FDMXEntityTreeNodeBase>>& OutCollapsedNodes, TSharedPtr<FDMXEntityTreeNodeBase> InParentNode = nullptr) const;
 
-	/**
-	 * Set the expansion state of a node
-	 *
-	 * @param InNodeToChange	The node to be expanded/collapsed
-	 * @param bIsExpanded		True to expand the node, false to collapse it
-	 */
-	void SetNodeExpansionState(TSharedPtr<FDMXEntityTreeNodeBase> InNodeToChange, const bool bIsExpanded);
-
-	/** Expand all categories during filtering and resets nodes expansion state after filtering is cleared */
-	void UpdateNodesExpansion(TSharedRef<FDMXEntityTreeNodeBase> InRootNode, bool bFilterIsEmpty);
-
 	/** Called when the active tab in the editor changes */
 	void OnActiveTabChanged(TSharedPtr<SDockTab> PreviouslyActive, TSharedPtr<SDockTab> NewlyActivated);
 
@@ -250,6 +250,9 @@ private:
 	/** The filter box that handles filtering entities */
 	TSharedPtr<SSearchBox> FilterBox;
 
+	/** Map of items that should be expanded or collapsed next tick */
+	TMap<TSharedRef<FDMXEntityTreeNodeBase>, bool> PendingExpandItemMap;
+
 	/** Gate to prevent changing the selection while selection change is being broadcast. */
 	bool bUpdatingSelection = false;
 
@@ -269,4 +272,7 @@ private:
 
 	/** Delegate handle bound to the FGlobalTabmanager::OnActiveTabChanged delegate */
 	FDelegateHandle OnActiveTabChangedDelegateHandle;
+
+	/** The actual Tree widget */
+	TSharedPtr<STreeView<TSharedPtr<FDMXEntityTreeNodeBase>>> EntitiesTreeWidget;
 };
