DMX - Improve UI styles for better looks in DMX project settings, DMX channel widget and DMX Fader

diff --git a/Engine/Plugins/VirtualProduction/DMX/DMXEngine/Source/DMXEditor/Private/Style/DMXEditorStyle.cpp b/Engine/Plugins/VirtualProduction/DMX/DMXEngine/Source/DMXEditor/Private/Style/DMXEditorStyle.cpp
--- a/Engine/Plugins/VirtualProduction/DMX/DMXEngine/Source/DMXEditor/Private/Style/DMXEditorStyle.cpp
+++ b/Engine/Plugins/VirtualProduction/DMX/DMXEngine/Source/DMXEditor/Private/Style/DMXEditorStyle.cpp
@@ -9,6 +9,7 @@
 #include "Interfaces/IPluginManager.h"
 #include "Engine/Font.h"
 
+
 TSharedPtr<FSlateStyleSet> FDMXEditorStyle::StyleInstance = nullptr;
 
 void FDMXEditorStyle::Initialize()
@@ -34,7 +35,8 @@ FName FDMXEditorStyle::GetStyleSetName()
 }
 
 #define IMAGE_BRUSH(RelativePath, ...) FSlateImageBrush(Style->RootToContentDir(RelativePath, TEXT(".png")), __VA_ARGS__)
-#define BOX_BRUSH(RelativePath, ...) FSlateBoxBrush(Style->RootToContentDir(RelativePath, TEXT(".png") ), __VA_ARGS__)
+#define BOX_BRUSH(RelativePath, ...) FSlateBoxBrush(Style->RootToContentDir(RelativePath, TEXT(".png")), __VA_ARGS__)
+#define CORE_BOX_BRUSH( RelativePath, ... ) FSlateBoxBrush(Style->RootToCoreContentDir(RelativePath, TEXT(".png") ), __VA_ARGS__)
 
 TSharedRef<FSlateStyleSet> FDMXEditorStyle::Create()
 {
@@ -48,6 +50,7 @@ TSharedRef<FSlateStyleSet> FDMXEditorStyle::Create()
 
 	TSharedRef<FSlateStyleSet> Style = MakeShared<FSlateStyleSet>(GetStyleSetName());
 	Style->SetContentRoot(IPluginManager::Get().FindPlugin(TEXT("DMXEngine"))->GetBaseDir() / TEXT("Resources"));
+	Style->SetCoreContentRoot(FPaths::EngineContentDir());
 
 	// Default color brushes
 	{
@@ -115,18 +118,15 @@ TSharedRef<FSlateStyleSet> FDMXEditorStyle::Create()
 		Style->Set("DMXEditor.OutputConsole.MacroMin", new IMAGE_BRUSH("Icons/icon_MacroMin_51x30", Icon51x30));
 		Style->Set("DMXEditor.OutputConsole.MacroMax", new IMAGE_BRUSH("Icons/icon_MacroMax_51x30", Icon51x30));
 
-		static const FLinearColor DefaultFaderBackColor = FLinearColor::Black;
+		static const FLinearColor DefaultFaderBackColor = FLinearColor::FromSRGBColor(FColor::FromHex("191919"));
 		static const FLinearColor DefaultFaderFillColor = FLinearColor::FromSRGBColor(FColor::FromHex("00aeef"));
 		static const FLinearColor DefeaultFaderForeColor = FLinearColor::FromSRGBColor(FColor::FromHex("ffffff"));
 
-		static FSlateBrush FaderBackBrush = BOX_BRUSH("Common/Spinbox", FMargin(4.0f / 16.0f), DefaultFaderBackColor);
-		static FSlateBrush FaderFillBrush = BOX_BRUSH("Common/Spinbox_Fill", FMargin(4.0f / 16.0f), DefaultFaderFillColor);
-
 		Style->Set("DMXEditor.OutputConsole.Fader", FSpinBoxStyle(FEditorStyle::GetWidgetStyle<FSpinBoxStyle>("SpinBox"))
-			.SetBackgroundBrush(FaderBackBrush)
-			.SetHoveredBackgroundBrush(FaderBackBrush)
-			.SetActiveFillBrush(FaderFillBrush)
-			.SetInactiveFillBrush(FaderFillBrush)
+			.SetBackgroundBrush(CORE_BOX_BRUSH("Slate/Common/Spinbox", FMargin(4.0f / 16.0f), DefaultFaderBackColor))
+			.SetHoveredBackgroundBrush(CORE_BOX_BRUSH("Slate/Common/Spinbox", FMargin(4.0f / 16.0f), DefaultFaderBackColor))
+			.SetActiveFillBrush(CORE_BOX_BRUSH("Slate/Common/Spinbox_Fill", FMargin(4.0f / 16.0f), DefaultFaderFillColor))
+			.SetInactiveFillBrush(CORE_BOX_BRUSH("Slate/Common/Spinbox_Fill", FMargin(4.0f / 16.0f), DefaultFaderFillColor))
 			.SetForegroundColor(DefeaultFaderForeColor)
 			.SetArrowsImage(FSlateNoResource()));
 	}
@@ -135,6 +135,7 @@ TSharedRef<FSlateStyleSet> FDMXEditorStyle::Create()
 }
 #undef IMAGE_BRUSH
 #undef BOX_BRUSH
+#undef CORE_BOX_BRUSH
 
 void FDMXEditorStyle::ReloadTextures()
 {
diff --git a/Engine/Plugins/VirtualProduction/DMX/DMXEngine/Source/DMXEditor/Private/Widgets/SDMXChannel.cpp b/Engine/Plugins/VirtualProduction/DMX/DMXEngine/Source/DMXEditor/Private/Widgets/SDMXChannel.cpp
--- a/Engine/Plugins/VirtualProduction/DMX/DMXEngine/Source/DMXEditor/Private/Widgets/SDMXChannel.cpp
+++ b/Engine/Plugins/VirtualProduction/DMX/DMXEngine/Source/DMXEditor/Private/Widgets/SDMXChannel.cpp
@@ -40,7 +40,7 @@ void SDMXChannel::Construct(const FArguments& InArgs)
 		SNew(STextBlock)
 		.Text(GetChannelIDText())
 		.ColorAndOpacity(FSlateColor(IDColor))
-		.MinDesiredWidth(23.0f)
+		.MinDesiredWidth(24.0f)
 		.Justification(ETextJustify::Center)
 		.Font(FDMXEditorStyle::Get().GetFontStyle("DMXEditor.Font.InputChannelID"));
 
@@ -48,7 +48,7 @@ void SDMXChannel::Construct(const FArguments& InArgs)
 		SNew(STextBlock)
 		.Text(GetValueText())
 		.ColorAndOpacity(FSlateColor(ValueColor))
-		.MinDesiredWidth(23.0f)
+		.MinDesiredWidth(24.0f)
 		.Justification(ETextJustify::Center)
 		.Font(FDMXEditorStyle::Get().GetFontStyle("DMXEditor.Font.InputChannelValue"));
 
diff --git a/Engine/Plugins/VirtualProduction/DMX/DMXProtocol/Source/DMXProtocolEditor/Private/DetailsCustomizations/Widgets/SDMXDelayEditWidget.cpp b/Engine/Plugins/VirtualProduction/DMX/DMXProtocol/Source/DMXProtocolEditor/Private/DetailsCustomizations/Widgets/SDMXDelayEditWidget.cpp
--- a/Engine/Plugins/VirtualProduction/DMX/DMXProtocol/Source/DMXProtocolEditor/Private/DetailsCustomizations/Widgets/SDMXDelayEditWidget.cpp
+++ b/Engine/Plugins/VirtualProduction/DMX/DMXProtocol/Source/DMXProtocolEditor/Private/DetailsCustomizations/Widgets/SDMXDelayEditWidget.cpp
@@ -38,6 +38,7 @@ void SDMXDelayEditWidget::Construct(const FArguments& InArgs)
 		.VAlign(VAlign_Center)
 		[
 			SAssignNew(ValueEditBox, SEditableTextBox)
+			.Font(FEditorStyle::GetFontStyle(TEXT("PropertyWindow.NormalFont")))
 			.Text(InitialDelayText)
 			.OnVerifyTextChanged(this, &SDMXDelayEditWidget::OnVerifyValueTextChanged)
 			.OnTextCommitted(this, &SDMXDelayEditWidget::OnValueTextCommitted)
@@ -50,6 +51,7 @@ void SDMXDelayEditWidget::Construct(const FArguments& InArgs)
 		.AutoWidth()
 		[
 			SNew(STextBlock)
+			.Font(FEditorStyle::GetFontStyle(TEXT("PropertyWindow.NormalFont")))
 			.Text_Lambda([this]
 				{
 					if (Delay == 1.0)
@@ -74,6 +76,7 @@ void SDMXDelayEditWidget::Construct(const FArguments& InArgs)
 			.ButtonContent()
 			[
 				SNew(STextBlock)
+				.Font(FEditorStyle::GetFontStyle(TEXT("PropertyWindow.NormalFont")))
 				.Text_Lambda([this]()
 					{
 						const double DelayFrameRateDecimal = DelayFrameRate.AsDecimal();
