DMX - Fix a Crash on Unix based systems when clicking anywhere in the Editor after renaming a function in DMX Library

diff --git a/Engine/Plugins/VirtualProduction/DMX/DMXEngine/Source/DMXEditor/Private/Widgets/FixtureType/SDMXFixtureFunctionEditor.cpp b/Engine/Plugins/VirtualProduction/DMX/DMXEngine/Source/DMXEditor/Private/Widgets/FixtureType/SDMXFixtureFunctionEditor.cpp
--- a/Engine/Plugins/VirtualProduction/DMX/DMXEngine/Source/DMXEditor/Private/Widgets/FixtureType/SDMXFixtureFunctionEditor.cpp
+++ b/Engine/Plugins/VirtualProduction/DMX/DMXEngine/Source/DMXEditor/Private/Widgets/FixtureType/SDMXFixtureFunctionEditor.cpp
@@ -172,7 +172,8 @@ void SDMXFixtureFunctionEditor::SetFunction(UDMXEntityFixtureType* InFixtureType
 			FStructureDetailsViewArgs StructureDetailsViewArgs;
 
 			FPropertyEditorModule& PropertyModule = FModuleManager::LoadModuleChecked<FPropertyEditorModule>("PropertyEditor");
-			StructDetailsView = PropertyModule.CreateStructureDetailView(DetailsViewArgs, StructureDetailsViewArgs, FunctionStructOnScope);
+			
+			const TSharedRef<IStructureDetailsView> StructDetailsView = PropertyModule.CreateStructureDetailView(DetailsViewArgs, StructureDetailsViewArgs, FunctionStructOnScope);
 
 			StructDetailsView->GetDetailsView()->SetIsPropertyVisibleDelegate(FIsPropertyVisible::CreateSP(this, &SDMXFixtureFunctionEditor::IsPropertyVisible));
 			StructDetailsView->GetDetailsView()->ForceRefresh();
diff --git a/Engine/Plugins/VirtualProduction/DMX/DMXEngine/Source/DMXEditor/Private/Widgets/FixtureType/SDMXFixtureFunctionEditor.h b/Engine/Plugins/VirtualProduction/DMX/DMXEngine/Source/DMXEditor/Private/Widgets/FixtureType/SDMXFixtureFunctionEditor.h
--- a/Engine/Plugins/VirtualProduction/DMX/DMXEngine/Source/DMXEditor/Private/Widgets/FixtureType/SDMXFixtureFunctionEditor.h
+++ b/Engine/Plugins/VirtualProduction/DMX/DMXEngine/Source/DMXEditor/Private/Widgets/FixtureType/SDMXFixtureFunctionEditor.h
@@ -69,10 +69,7 @@ private:
 
 	/** The current transaction or nullptr if there's no transaction ongoing */
 	TUniquePtr<FScopedTransaction> Transaction;
-
-	/** Details view to display the Mode struct */
-	TSharedPtr<IStructureDetailsView> StructDetailsView;
-
+	
 	/** Border that holds the content of the editor */
 	TSharedPtr<SBorder> ContentBorder;
 
diff --git a/Engine/Plugins/VirtualProduction/DMX/DMXEngine/Source/DMXEditor/Private/Widgets/FixtureType/SDMXFixtureModeEditor.cpp b/Engine/Plugins/VirtualProduction/DMX/DMXEngine/Source/DMXEditor/Private/Widgets/FixtureType/SDMXFixtureModeEditor.cpp
--- a/Engine/Plugins/VirtualProduction/DMX/DMXEngine/Source/DMXEditor/Private/Widgets/FixtureType/SDMXFixtureModeEditor.cpp
+++ b/Engine/Plugins/VirtualProduction/DMX/DMXEngine/Source/DMXEditor/Private/Widgets/FixtureType/SDMXFixtureModeEditor.cpp
@@ -196,7 +196,7 @@ void SDMXFixtureModeEditor::SetMode(UDMXEntityFixtureType* InFixtureType, int32
 		FStructureDetailsViewArgs StructureDetailsViewArgs;
 
 		FPropertyEditorModule& PropertyModule = FModuleManager::LoadModuleChecked<FPropertyEditorModule>("PropertyEditor");
-		StructDetailsView = PropertyModule.CreateStructureDetailView(DetailsViewArgs, StructureDetailsViewArgs, ModeStructOnScope);
+		const TSharedRef<IStructureDetailsView> StructDetailsView = PropertyModule.CreateStructureDetailView(DetailsViewArgs, StructureDetailsViewArgs, ModeStructOnScope);
 		StructDetailsView->GetDetailsView()->SetCustomFilterLabel(LOCTEXT("SearchModePropertiesLabel", "Search Mode Properties"));
 
 		StructDetailsView->GetDetailsView()->SetIsPropertyVisibleDelegate(FIsPropertyVisible::CreateSP(this, &SDMXFixtureModeEditor::IsPropertyVisible));
diff --git a/Engine/Plugins/VirtualProduction/DMX/DMXEngine/Source/DMXEditor/Private/Widgets/FixtureType/SDMXFixtureModeEditor.h b/Engine/Plugins/VirtualProduction/DMX/DMXEngine/Source/DMXEditor/Private/Widgets/FixtureType/SDMXFixtureModeEditor.h
--- a/Engine/Plugins/VirtualProduction/DMX/DMXEngine/Source/DMXEditor/Private/Widgets/FixtureType/SDMXFixtureModeEditor.h
+++ b/Engine/Plugins/VirtualProduction/DMX/DMXEngine/Source/DMXEditor/Private/Widgets/FixtureType/SDMXFixtureModeEditor.h
@@ -63,9 +63,6 @@ private:
 	/** The current transaction or nullptr if there's no transaction ongoing */
 	TUniquePtr<FScopedTransaction> Transaction;
 
-	/** Details view to display the Mode struct */
-	TSharedPtr<IStructureDetailsView> StructDetailsView;
-
 	/** Border that holds the content of the editor */
 	TSharedPtr<SBorder> ContentBorder;
 
