DMX - Fix a crash when clicking the + New Fixture Type in the Fixture Types tab in DMX Library

diff --git a/Engine/Plugins/VirtualProduction/DMX/DMXEngine/Source/DMXEditor/Private/Widgets/FixturePatch/SDMXFixturePatcher.cpp b/Engine/Plugins/VirtualProduction/DMX/DMXEngine/Source/DMXEditor/Private/Widgets/FixturePatch/SDMXFixturePatcher.cpp
--- a/Engine/Plugins/VirtualProduction/DMX/DMXEngine/Source/DMXEditor/Private/Widgets/FixturePatch/SDMXFixturePatcher.cpp
+++ b/Engine/Plugins/VirtualProduction/DMX/DMXEngine/Source/DMXEditor/Private/Widgets/FixturePatch/SDMXFixturePatcher.cpp
@@ -164,21 +164,8 @@ void SDMXFixturePatcher::Construct(const FArguments& InArgs)
 		FGlobalTabmanager::Get()->OnActiveTabChanged_Subscribe(FOnActiveTabChanged::FDelegate::CreateSP(this, &SDMXFixturePatcher::OnActiveTabChanged));
 
 		// Bind to entity changes
-		Library->GetOnEntitiesAdded().AddLambda([this](UDMXLibrary* DMXLibrary, TArray<UDMXEntity*> Entities)
-			{
-				if (DMXLibrary == GetDMXLibrary())
-				{
-					RefreshFromLibrary();
-				}
-			});
-
-		Library->GetOnEntitiesRemoved().AddLambda([this](UDMXLibrary* DMXLibrary, TArray<UDMXEntity*> Entities)
-			{
-				if (DMXLibrary == GetDMXLibrary())
-				{
-					RefreshFromLibrary();
-				}
-			});
+		Library->GetOnEntitiesAdded().AddSP(this, &SDMXFixturePatcher::OnEntitiesAddedOrRemoved);
+		Library->GetOnEntitiesRemoved().AddSP(this, &SDMXFixturePatcher::OnEntitiesAddedOrRemoved);
 
 		GEditor->RegisterForUndo(this);
 
@@ -537,6 +524,14 @@ TSharedPtr<FDMXFixturePatchNode> SDMXFixturePatcher::FindPatchNodeOfType(UDMXEnt
 	return nullptr;
 }
 
+void SDMXFixturePatcher::OnEntitiesAddedOrRemoved(UDMXLibrary* DMXLibrary, TArray<UDMXEntity*> Entities)
+{
+	if (DMXLibrary == GetDMXLibrary())
+	{
+		RefreshFromLibrary();
+	}
+}
+
 void SDMXFixturePatcher::OnFixturePatchChanged(const UDMXEntityFixturePatch* FixturePatch)
 {
 	if (IsValid(FixturePatch))
diff --git a/Engine/Plugins/VirtualProduction/DMX/DMXEngine/Source/DMXEditor/Private/Widgets/FixturePatch/SDMXFixturePatcher.h b/Engine/Plugins/VirtualProduction/DMX/DMXEngine/Source/DMXEditor/Private/Widgets/FixturePatch/SDMXFixturePatcher.h
--- a/Engine/Plugins/VirtualProduction/DMX/DMXEngine/Source/DMXEditor/Private/Widgets/FixturePatch/SDMXFixturePatcher.h
+++ b/Engine/Plugins/VirtualProduction/DMX/DMXEngine/Source/DMXEditor/Private/Widgets/FixturePatch/SDMXFixturePatcher.h
@@ -86,6 +86,9 @@ protected:
 	/** Returns first node with same fixture type */
 	TSharedPtr<FDMXFixturePatchNode> FindPatchNodeOfType(UDMXEntityFixtureType* Type, const TSharedPtr<FDMXFixturePatchNode>& IgoredNode) const;
 
+	/** Called when entities were added to or removed from a DMX Library */
+	void OnEntitiesAddedOrRemoved(UDMXLibrary* DMXLibrary, TArray<UDMXEntity*> Entities);
+
 	/** Called when a fixture patch changed */
 	void OnFixturePatchChanged(const UDMXEntityFixturePatch* FixturePatch);
 
