DMX - Fix a crash on Unix based systems when clicking anywhere in the Editor after renaming a function

diff --git a/Engine/Plugins/VirtualProduction/DMX/DMXEngine/Source/DMXEditor/Private/Widgets/FixtureType/SDMXFixtureFunctionEditor.cpp b/Engine/Plugins/VirtualProduction/DMX/DMXEngine/Source/DMXEditor/Private/Widgets/FixtureType/SDMXFixtureFunctionEditor.cpp
--- a/Engine/Plugins/VirtualProduction/DMX/DMXEngine/Source/DMXEditor/Private/Widgets/FixtureType/SDMXFixtureFunctionEditor.cpp
+++ b/Engine/Plugins/VirtualProduction/DMX/DMXEngine/Source/DMXEditor/Private/Widgets/FixtureType/SDMXFixtureFunctionEditor.cpp
@@ -45,11 +45,15 @@ void SDMXFixtureFunctionEditor::NotifyPreChange(FProperty* PropertyAboutToChange
 	{
 		const FName PropertyName = PropertyAboutToChange->GetFName();
 
-		if (PropertyName == GET_MEMBER_NAME_CHECKED(FDMXFixtureFunction, Channel))
+		if (PropertyName == GET_MEMBER_NAME_CHECKED(FDMXFixtureFunction, FunctionName))
+		{
+			Transaction = MakeUnique<FScopedTransaction>(LOCTEXT("SetFunctionNameTransaction", "Set Fixture Function Name"));
+		}
+		else if (PropertyName == GET_MEMBER_NAME_CHECKED(FDMXFixtureFunction, Channel))
 		{
 			Transaction = MakeUnique<FScopedTransaction>(LOCTEXT("SetChannelTransaction", "Set Fixture Function Starting Channel"));
 		}
-		if (PropertyName == GET_MEMBER_NAME_CHECKED(FDMXFixtureFunction, DataType))
+		else if (PropertyName == GET_MEMBER_NAME_CHECKED(FDMXFixtureFunction, DataType))
 		{
 			Transaction = MakeUnique<FScopedTransaction>(LOCTEXT("ChangeDefaultValueTransaction", "Data Type of Function"));
 		}
@@ -74,7 +78,13 @@ void SDMXFixtureFunctionEditor::NotifyPostChange(const FPropertyChangedEvent& Pr
 
 		if (PropertyChangedEvent.ChangeType != EPropertyChangeType::Interactive)
 		{
-			if (PropertyName == GET_MEMBER_NAME_CHECKED(FDMXFixtureFunction, Channel))
+			if (PropertyName == GET_MEMBER_NAME_CHECKED(FDMXFixtureFunction, FunctionName))
+			{
+				// Make a unique function name
+				FString OutUniqueFunctionName;
+				FixtureType->SetFunctionName(ModeIndex, FunctionIndex, FunctionBeingEditedPtr->FunctionName, OutUniqueFunctionName);
+			}
+			else if (PropertyName == GET_MEMBER_NAME_CHECKED(FDMXFixtureFunction, Channel))
 			{
 				int32 ResultingStartingChannel;
 				FixtureType->SetFunctionStartingChannel(ModeIndex, FunctionIndex, FunctionBeingEditedPtr->Channel, ResultingStartingChannel);
@@ -150,7 +160,7 @@ void SDMXFixtureFunctionEditor::SetFunction(UDMXEntityFixtureType* InFixtureType
 		ModeIndex = InModeIndex;
 		FDMXFixtureMode& Mode = InFixtureType->Modes[ModeIndex];
 
-		if (ensureMsgf(Mode.Functions.IsValidIndex(InFunctionIndex), TEXT("Function Index specified when setting the Function in the Function editor.")))
+		if (ensureMsgf(Mode.Functions.IsValidIndex(InFunctionIndex), TEXT("Invalid Function Index specified when setting the Function in the Function editor.")))
 		{
 			FunctionIndex = InFunctionIndex;
 			FDMXFixtureFunction& Function = Mode.Functions[FunctionIndex];
@@ -177,7 +187,7 @@ void SDMXFixtureFunctionEditor::SetFunction(UDMXEntityFixtureType* InFixtureType
 
 			StructDetailsView->GetDetailsView()->SetIsPropertyVisibleDelegate(FIsPropertyVisible::CreateSP(this, &SDMXFixtureFunctionEditor::IsPropertyVisible));
 			StructDetailsView->GetDetailsView()->ForceRefresh();
-
+			
 			if (TSharedPtr<SWidget> Widget = StructDetailsView->GetWidget())
 			{
 				ContentBorder->SetContent(Widget.ToSharedRef());
