DMX - Refactor remove fixture patches and fixture types to be consistent with how they are added in c++ and blueprints

diff --git a/Engine/Plugins/VirtualProduction/DMX/DMXEngine/Source/DMXEditor/Private/Widgets/FixturePatch/SDMXFixturePatchTree.cpp b/Engine/Plugins/VirtualProduction/DMX/DMXEngine/Source/DMXEditor/Private/Widgets/FixturePatch/SDMXFixturePatchTree.cpp
--- a/Engine/Plugins/VirtualProduction/DMX/DMXEngine/Source/DMXEditor/Private/Widgets/FixturePatch/SDMXFixturePatchTree.cpp
+++ b/Engine/Plugins/VirtualProduction/DMX/DMXEngine/Source/DMXEditor/Private/Widgets/FixturePatch/SDMXFixturePatchTree.cpp
@@ -602,10 +602,30 @@ bool SDMXFixturePatchTree::CanDuplicateNodes() const
 
 void SDMXFixturePatchTree::OnDeleteNodes()
 {
+	// Gather Fixture Patches to remove from the DMX Library
 	const TArray<UDMXEntity*> EntitiesToDelete = GetSelectedEntities();
 
-	// Clears references to the Entities and delete them
-	FDMXEditorUtils::RemoveEntities(GetDMXLibrary(), EntitiesToDelete);
+	TArray<UDMXEntityFixturePatch*> FixturePatchesToDelete;
+	for (UDMXEntity* Entity : EntitiesToDelete)
+	{
+		if (UDMXEntityFixturePatch* FixturePatch = Cast<UDMXEntityFixturePatch>(Entity))
+		{
+			FixturePatchesToDelete.Add(FixturePatch);
+		}
+	}
+
+	// Remove the Fixture Patches from the DMX Library
+	const FScopedTransaction Transaction(EntitiesToDelete.Num() > 1 ? LOCTEXT("RemoveEntities", "Remove Entities") : LOCTEXT("RemoveEntity", "Remove Entity"));
+
+	UDMXLibrary* DMXLibrary = GetDMXLibrary();
+	DMXLibrary->PreEditChange(UDMXLibrary::StaticClass()->FindPropertyByName(UDMXLibrary::GetEntitiesPropertyName()));
+	for (UDMXEntityFixturePatch* FixturePatch : FixturePatchesToDelete)
+	{
+		ensureMsgf(DMXLibrary == FixturePatch->GetParentLibrary(), TEXT("Unexpected DMX Library of Fixture Patch and DMX Library of Editor do not match when removing Fixture Patches."));
+		const FDMXEntityFixturePatchRef FixturePatchRef(FixturePatch);
+		UDMXEntityFixturePatch::RemoveFixturePatchFromLibrary(FixturePatchRef);
+	}
+	DMXLibrary->PostEditChange();
 
 	// Clear selection if no patches remain
 	if (GetSelectedEntities().Num() == 0)
diff --git a/Engine/Plugins/VirtualProduction/DMX/DMXEngine/Source/DMXEditor/Private/Widgets/FixtureType/SDMXFixtureTypeTree.cpp b/Engine/Plugins/VirtualProduction/DMX/DMXEngine/Source/DMXEditor/Private/Widgets/FixtureType/SDMXFixtureTypeTree.cpp
--- a/Engine/Plugins/VirtualProduction/DMX/DMXEngine/Source/DMXEditor/Private/Widgets/FixtureType/SDMXFixtureTypeTree.cpp
+++ b/Engine/Plugins/VirtualProduction/DMX/DMXEngine/Source/DMXEditor/Private/Widgets/FixtureType/SDMXFixtureTypeTree.cpp
@@ -16,9 +16,11 @@
 #include "Widgets/DMXEntityTreeNode.h"
 
 #include "ScopedTransaction.h"
+#include "Dialogs/Dialogs.h"
 #include "Framework/Application/SlateApplication.h"
 #include "Framework/Commands/GenericCommands.h"
 #include "Framework/MultiBox/MultiBoxBuilder.h"
+#include "Misc/MessageDialog.h"
 #include "Widgets/SBoxPanel.h"
 #include "Widgets/Input/SButton.h"
 #include "Widgets/Images/SImage.h"
@@ -371,28 +373,62 @@ bool SDMXFixtureTypeTree::CanDuplicateNodes() const
 
 void SDMXFixtureTypeTree::OnDeleteNodes()
 {
+	// Gather Fixture Types to remove
 	TArray<UDMXEntity*> EntitiesToDelete = GetSelectedEntities();
-	
-	// Check for entities being used by other objects
+
+	TArray<UDMXEntityFixtureType*> FixtureTypesToDelete;
+	for (UDMXEntity* Entity : EntitiesToDelete)
+	{
+		if (UDMXEntityFixtureType* FixtureType = Cast<UDMXEntityFixtureType>(Entity))
+		{
+			FixtureTypesToDelete.Add(FixtureType);
+		}
+	}
+
+	// Find Fixture Types that are used by Fixture Patches
+	UDMXLibrary* DMXLibrary = GetDMXLibrary();
 	TArray<UDMXEntity*> EntitiesInUse;
 	for (UDMXEntity* Entity : EntitiesToDelete)
 	{
-		if (FDMXEditorUtils::IsEntityUsed(GetDMXLibrary(), Entity))
+		if (FDMXEditorUtils::IsEntityUsed(DMXLibrary, Entity))
 		{
 			EntitiesInUse.Add(Entity);
 		}
 	}
 
-	// Clears references to the Entities and delete them
-	FDMXEditorUtils::RemoveEntities(GetDMXLibrary(), MoveTemp(EntitiesToDelete));
-
-	// Clear selection if no fixture types remain
-	if (GetSelectedEntities().Num() == 0)
+	// Confirm deletion of Fixture Types in use
+	if (EntitiesInUse.Num() > 0)
 	{
-		FixtureTypeSharedData->SelectFixtureTypes(TArray<TWeakObjectPtr<UDMXEntityFixtureType>>());
-	}
+		const FText ConfirmDeleteText =
+			FText::Format(LOCTEXT("ConfirmDeleteEntityInUse", "Fixture Patches use {0}|plural(one = this, other = these) Fixture {0}|plural(one = Type, other = Types). Do you really want to delete {0}|plural(one = it, other = them)?"),
+				EntitiesInUse.Num() == 1);
+
+		const FText DialogTitleText = LOCTEXT("DeleteFixtureTypesDialogTitle", "Delete Fixture Types");
 
-	UpdateTree();
+		EAppReturnType::Type DialogResult = FMessageDialog::Open(EAppMsgType::YesNo, EAppReturnType::No, ConfirmDeleteText, &DialogTitleText);
+		if (DialogResult == EAppReturnType::Yes)
+		{
+			// Remove the Fixture Types from the DMX Library
+			const FScopedTransaction Transaction(EntitiesToDelete.Num() > 1 ? LOCTEXT("RemoveEntities", "Remove Entities") : LOCTEXT("RemoveEntity", "Remove Entity"));
+
+			DMXLibrary->PreEditChange(UDMXLibrary::StaticClass()->FindPropertyByName(UDMXLibrary::GetEntitiesPropertyName()));
+			for (UDMXEntityFixtureType* FixtureType : FixtureTypesToDelete)
+			{
+				ensureMsgf(DMXLibrary == FixtureType->GetParentLibrary(), TEXT("Unexpected DMX Library of Fixture Type and DMX Library of Editor do not match when removing Fixture Types."));
+				const FDMXEntityFixtureTypeRef FixtureTypeRef(FixtureType);
+				UDMXEntityFixtureType::RemoveFixtureTypeFromLibrary(FixtureTypeRef);
+			}
+			DMXLibrary->PostEditChange();
+
+			// Clear selection if no fixture types remain
+			if (GetSelectedEntities().Num() == 0)
+			{
+				FixtureTypeSharedData->SelectFixtureTypes(TArray<TWeakObjectPtr<UDMXEntityFixtureType>>());
+			}
+
+			UpdateTree();
+		}
+	}
 }
 
 bool SDMXFixtureTypeTree::CanDeleteNodes() const
diff --git a/Engine/Plugins/VirtualProduction/DMX/DMXEngine/Source/DMXEditor/Public/DMXEditorUtils.h b/Engine/Plugins/VirtualProduction/DMX/DMXEngine/Source/DMXEditor/Public/DMXEditorUtils.h
--- a/Engine/Plugins/VirtualProduction/DMX/DMXEngine/Source/DMXEditor/Public/DMXEditorUtils.h
+++ b/Engine/Plugins/VirtualProduction/DMX/DMXEngine/Source/DMXEditor/Public/DMXEditorUtils.h
@@ -48,7 +48,8 @@ public:
 	/**  Checks if the Entity is being referenced by other objects. */
 	static bool IsEntityUsed(const UDMXLibrary* InLibrary, const UDMXEntity* InEntity);
 
-	/**  Removes the entities and fixes references to it. */
+	/** DEPRECATED 5.0 */
+	UE_DEPRECATED(5.0, "Replaced with UDMXEntityFixtureType::RemoveFixtureTypeFromLibrary and UDMXEntityFixturePatch::RemoveFixturePatchFromLibrary to support blueprint and runtime changes.")
 	static void RemoveEntities(UDMXLibrary* InLibrary, const TArray<UDMXEntity*>& InEntities);
 
 	/**  Copies Entities to the operating system's clipboard. */
diff --git a/Engine/Plugins/VirtualProduction/DMX/DMXEngine/Source/DMXRuntime/Private/Library/DMXEntityFixturePatch.cpp b/Engine/Plugins/VirtualProduction/DMX/DMXEngine/Source/DMXRuntime/Private/Library/DMXEntityFixturePatch.cpp
--- a/Engine/Plugins/VirtualProduction/DMX/DMXEngine/Source/DMXRuntime/Private/Library/DMXEntityFixturePatch.cpp
+++ b/Engine/Plugins/VirtualProduction/DMX/DMXEngine/Source/DMXRuntime/Private/Library/DMXEntityFixturePatch.cpp
@@ -60,6 +60,21 @@ UDMXEntityFixturePatch* UDMXEntityFixturePatch::CreateFixturePatchInLibrary(FDMX
 	return nullptr;
 }
 
+void UDMXEntityFixturePatch::RemoveFixturePatchFromLibrary(FDMXEntityFixturePatchRef FixturePatchRef)
+{
+	if (UDMXEntityFixturePatch* FixturePatch = FixturePatchRef.GetFixturePatch())
+	{
+		if (UDMXLibrary* DMXLibrary = FixturePatch->GetParentLibrary())
+		{
+			FixturePatch->SetFixtureType(nullptr);
+
+			DMXLibrary->Modify();
+			FixturePatch->Modify();
+			FixturePatch->Destroy();
+		}
+	}
+}
+
 #if WITH_EDITOR
 bool UDMXEntityFixturePatch::Modify(bool bAlwaysMarkDirty)
 {
diff --git a/Engine/Plugins/VirtualProduction/DMX/DMXEngine/Source/DMXRuntime/Private/Library/DMXEntityFixtureType.cpp b/Engine/Plugins/VirtualProduction/DMX/DMXEngine/Source/DMXRuntime/Private/Library/DMXEntityFixtureType.cpp
--- a/Engine/Plugins/VirtualProduction/DMX/DMXEngine/Source/DMXRuntime/Private/Library/DMXEntityFixtureType.cpp
+++ b/Engine/Plugins/VirtualProduction/DMX/DMXEngine/Source/DMXRuntime/Private/Library/DMXEntityFixtureType.cpp
@@ -188,6 +188,19 @@ UDMXEntityFixtureType* UDMXEntityFixtureType::CreateFixtureTypeInLibrary(UDMXLib
 	return nullptr;
 }
 
+void UDMXEntityFixtureType::RemoveFixtureTypeFromLibrary(FDMXEntityFixtureTypeRef FixtureTypeRef)
+{
+	if (UDMXEntityFixtureType* FixtureType = FixtureTypeRef.GetFixtureType())
+	{
+		if (UDMXLibrary* DMXLibrary = FixtureType->GetParentLibrary())
+		{
+			DMXLibrary->Modify();
+			FixtureType->Modify();
+			FixtureType->Destroy();
+		}
+	}
+}
+
 void UDMXEntityFixtureType::Serialize(FArchive& Ar)
 {
 	Super::Serialize(Ar);
diff --git a/Engine/Plugins/VirtualProduction/DMX/DMXEngine/Source/DMXRuntime/Private/Library/DMXLibrary.cpp b/Engine/Plugins/VirtualProduction/DMX/DMXEngine/Source/DMXRuntime/Private/Library/DMXLibrary.cpp
--- a/Engine/Plugins/VirtualProduction/DMX/DMXEngine/Source/DMXRuntime/Private/Library/DMXLibrary.cpp
+++ b/Engine/Plugins/VirtualProduction/DMX/DMXEngine/Source/DMXRuntime/Private/Library/DMXLibrary.cpp
@@ -95,7 +95,7 @@ void UDMXLibrary::PostEditChangeProperty(FPropertyChangedEvent& PropertyChangedE
 {
 	Super::PostEditChangeProperty(PropertyChangedEvent);
 
-	FName PropertyName = PropertyChangedEvent.GetPropertyName();
+	const FName PropertyName = PropertyChangedEvent.GetPropertyName();
 	const FProperty* Property = PropertyChangedEvent.Property;
 	const UScriptStruct* InputPortReferenceStruct = FDMXInputPortReference::StaticStruct();
 	const UScriptStruct* OutputPortReferenceStruct = FDMXOutputPortReference::StaticStruct();
@@ -108,6 +108,12 @@ void UDMXLibrary::PostEditChangeProperty(FPropertyChangedEvent& PropertyChangedE
 	{
 		UpdatePorts();
 	}
+
+	// Remove any null entities from the library
+	Entities.RemoveAll([](UDMXEntity* Entity)
+		{
+			return !IsValid(Entity);
+		});
 }
 #endif // WITH_EDITOR
 
diff --git a/Engine/Plugins/VirtualProduction/DMX/DMXEngine/Source/DMXRuntime/Public/Library/DMXEntityFixturePatch.h b/Engine/Plugins/VirtualProduction/DMX/DMXEngine/Source/DMXRuntime/Public/Library/DMXEntityFixturePatch.h
--- a/Engine/Plugins/VirtualProduction/DMX/DMXEngine/Source/DMXRuntime/Public/Library/DMXEntityFixturePatch.h
+++ b/Engine/Plugins/VirtualProduction/DMX/DMXEngine/Source/DMXRuntime/Public/Library/DMXEntityFixturePatch.h
@@ -43,10 +43,14 @@ class DMXRUNTIME_API UDMXEntityFixturePatch
 public:
 	UDMXEntityFixturePatch();
 
-	/** Creates a new fixture patch in the DMX Library using the specified Fixture Type */
+	/** Creates a new Fixture Patch in the DMX Library using the specified Fixture Type */
 	UFUNCTION(BlueprintCallable, Category = "DMX")
 	static UDMXEntityFixturePatch* CreateFixturePatchInLibrary(FDMXEntityFixtureTypeRef FixtureTypeRef, const FString& DesiredName = TEXT(""));
 
+	/** Removes a fixture patch from the DMX Library */
+	UFUNCTION(BlueprintCallable, Category = "DMX")
+	static void RemoveFixturePatchFromLibrary(FDMXEntityFixturePatchRef FixturePatchRef);
+
 	// ~Begin UObject Interface
 #if WITH_EDITOR
 	virtual bool Modify(bool bAlwaysMarkDirty = true) override;
diff --git a/Engine/Plugins/VirtualProduction/DMX/DMXEngine/Source/DMXRuntime/Public/Library/DMXEntityFixtureType.h b/Engine/Plugins/VirtualProduction/DMX/DMXEngine/Source/DMXRuntime/Public/Library/DMXEntityFixtureType.h
--- a/Engine/Plugins/VirtualProduction/DMX/DMXEngine/Source/DMXRuntime/Public/Library/DMXEntityFixtureType.h
+++ b/Engine/Plugins/VirtualProduction/DMX/DMXEngine/Source/DMXRuntime/Public/Library/DMXEntityFixtureType.h
@@ -271,9 +271,13 @@ class DMXRUNTIME_API UDMXEntityFixtureType
 	GENERATED_BODY()
 
 public:
-	/** Creates a new fixture patch in the DMX Library using the specified Fixture Type */
+	/** Creates a new Fixture Type in the DMX Library */
 	UFUNCTION(BlueprintCallable, Category = "DMX")
 	static UDMXEntityFixtureType* CreateFixtureTypeInLibrary(UDMXLibrary* ParentDMXLibrary, const FString& DesiredName = TEXT(""));
+	
+	/** Removes a Fixture Type from a DMX Library */
+	UFUNCTION(BlueprintCallable, Category = "DMX")
+	static void RemoveFixtureTypeFromLibrary(FDMXEntityFixtureTypeRef FixtureTypeRef);
 
 	//~ Begin UObject interface
 protected:
