DMX - Fis CIS NonUnity Compile UnrealGame WinGDK

diff --git a/Engine/Plugins/VirtualProduction/DMX/DMXEngine/Source/DMXRuntime/Public/Library/DMXEntityFixtureType.h b/Engine/Plugins/VirtualProduction/DMX/DMXEngine/Source/DMXRuntime/Public/Library/DMXEntityFixtureType.h
--- a/Engine/Plugins/VirtualProduction/DMX/DMXEngine/Source/DMXRuntime/Public/Library/DMXEntityFixtureType.h
+++ b/Engine/Plugins/VirtualProduction/DMX/DMXEngine/Source/DMXRuntime/Public/Library/DMXEntityFixtureType.h
@@ -5,6 +5,7 @@
 #include "DMXAttribute.h"
 #include "DMXProtocolTypes.h"
 #include "Library/DMXEntity.h"
+#include "Library/DMXEntityReference.h"
 #include "Modulators/DMXModulator.h"
 
 #include "DMXEntityFixtureType.generated.h"
