DMX - Fix a bug where Fixture Types cannot be deleted if they aren't used by any patches

diff --git a/Engine/Plugins/VirtualProduction/DMX/DMXEngine/Source/DMXEditor/Private/Customizations/DMXEntityFixturePatchDetails.cpp b/Engine/Plugins/VirtualProduction/DMX/DMXEngine/Source/DMXEditor/Private/Customizations/DMXEntityFixturePatchDetails.cpp
--- a/Engine/Plugins/VirtualProduction/DMX/DMXEngine/Source/DMXEditor/Private/Customizations/DMXEntityFixturePatchDetails.cpp
+++ b/Engine/Plugins/VirtualProduction/DMX/DMXEngine/Source/DMXEditor/Private/Customizations/DMXEntityFixturePatchDetails.cpp
@@ -9,6 +9,7 @@
 #include "DetailLayoutBuilder.h"
 #include "DetailWidgetRow.h"
 #include "IPropertyUtilities.h"
+#include "ScopedTransaction.h"
 
 
 #define LOCTEXT_NAMESPACE "DMXEntityFixturePatchFixtureSettingsDetails"
@@ -123,18 +124,21 @@ void FDMXEntityFixturePatchDetails::OnParentFixtureTypeChanged(UDMXEntity* NewTe
 
 void FDMXEntityFixturePatchDetails::OnFixtureTypeChanged(const UDMXEntityFixtureType* FixtureType)
 {
-	// Keep the active mode valid
-	int32 ActiveMode;
-	if (ActiveModeHandle->GetValue(ActiveMode) == FPropertyAccess::Success)
+	if (IsValid(FixtureType) && !FixtureType->HasAnyFlags(RF_Transactional))
 	{
-		if (!FixtureType->Modes.IsValidIndex(ActiveMode))
+		// Keep the active mode valid
+		int32 ActiveMode;
+		if (ActiveModeHandle->GetValue(ActiveMode) == FPropertyAccess::Success)
 		{
-			const int32 NewActiveMode = FixtureType->Modes.Num() > 0 ? 0 : INDEX_NONE;
-			ActiveModeHandle->SetValue(NewActiveMode);
+			if (!FixtureType->Modes.IsValidIndex(ActiveMode))
+			{
+				const int32 NewActiveMode = FixtureType->Modes.Num() > 0 ? 0 : INDEX_NONE;
+				ActiveModeHandle->SetValue(NewActiveMode);
+			}
 		}
+	
+		PropertyUtilities->ForceRefresh();
 	}
-
-	PropertyUtilities->ForceRefresh();
 }
 
 void FDMXEntityFixturePatchDetails::OnActiveModeChanged(const TSharedPtr<uint32> InSelectedMode, ESelectInfo::Type SelectInfo)
@@ -266,11 +270,20 @@ void FDMXEntityFixturePatchDetails::SetActiveMode(int32 ModeIndex)
 	TArray<UObject*> OuterObjects;
 	ActiveModeHandle->GetOuterObjects(OuterObjects);
 
-	TArray<UDMXEntityFixturePatch*> FixturePatches;
-	for (UObject* Object : OuterObjects)
+	if (OuterObjects.Num() > 0)
 	{
-		UDMXEntityFixturePatch* Patch = CastChecked<UDMXEntityFixturePatch>(Object);
-		Patch->SetActiveModeIndex(ModeIndex);
+		const FScopedTransaction Transaction(LOCTEXT("SetFixturePatchActiveModeTransaction", "Set DMX Fixture Patch Active Mode"));
+
+		for (UObject* Object : OuterObjects)
+		{
+			UDMXEntityFixturePatch* Patch = CastChecked<UDMXEntityFixturePatch>(Object);
+			Patch->Modify();
+			Patch->PreEditChange(UDMXEntityFixturePatch::StaticClass()->FindPropertyByName(UDMXEntityFixturePatch::GetActiveModePropertyNameChecked()));
+
+			Patch->SetActiveModeIndex(ModeIndex);
+
+			Patch->PostEditChange();
+		}
 	}
 }
 
diff --git a/Engine/Plugins/VirtualProduction/DMX/DMXEngine/Source/DMXEditor/Private/DMXEditorUtils.cpp b/Engine/Plugins/VirtualProduction/DMX/DMXEngine/Source/DMXEditor/Private/DMXEditorUtils.cpp
--- a/Engine/Plugins/VirtualProduction/DMX/DMXEngine/Source/DMXEditor/Private/DMXEditorUtils.cpp
+++ b/Engine/Plugins/VirtualProduction/DMX/DMXEngine/Source/DMXEditor/Private/DMXEditorUtils.cpp
@@ -566,6 +566,8 @@ FText FDMXEditorUtils::GetEntityTypeNameText(TSubclassOf<UDMXEntity> EntityClass
 bool FDMXEditorUtils::TryAutoAssignToUniverses(UDMXEntityFixturePatch* Patch, const TSet<int32>& AllowedUniverses)
 {
 	check(Patch->IsAutoAssignAddress());
+	
+	Patch->Modify();
 	const int32 UniverseToRestore = Patch->GetUniverseID();
 	const int32 AutoAddressToRestore = Patch->GetAutoStartingAddress();
 	
@@ -681,6 +683,7 @@ FDMXEditorUtils::FUnassignedPatchesArray FDMXEditorUtils::AutoAssignedAddresses(
 		
 		static void AssignPatchTo(UDMXEntityFixturePatch* Patch, int32 ToAddress, int32 UniverseToAssignTo)
 		{
+			Patch->Modify();
 			Patch->SetAutoStartingAddress(ToAddress);
 			Patch->SetUniverseID(UniverseToAssignTo);
 		}
diff --git a/Engine/Plugins/VirtualProduction/DMX/DMXEngine/Source/DMXEditor/Private/Widgets/FixturePatch/SDMXFixturePatchTree.cpp b/Engine/Plugins/VirtualProduction/DMX/DMXEngine/Source/DMXEditor/Private/Widgets/FixturePatch/SDMXFixturePatchTree.cpp
--- a/Engine/Plugins/VirtualProduction/DMX/DMXEngine/Source/DMXEditor/Private/Widgets/FixturePatch/SDMXFixturePatchTree.cpp
+++ b/Engine/Plugins/VirtualProduction/DMX/DMXEngine/Source/DMXEditor/Private/Widgets/FixturePatch/SDMXFixturePatchTree.cpp
@@ -844,7 +844,7 @@ void SDMXFixturePatchTree::OnAddNewFixturePatchClicked(UDMXEntity* InSelectedFix
 		if (UDMXEntityFixtureType* FixtureType = Cast<UDMXEntityFixtureType>(InSelectedFixtureType))
 		{
 			// Find the first free channel behind all patches, before adding the new patch
-			const UDMXLibrary* DMXLibrary = FixtureType->GetParentLibrary();
+			UDMXLibrary* DMXLibrary = FixtureType->GetParentLibrary();
 			TArray<UDMXEntityFixturePatch*> OtherFixturePatches = DMXLibrary->GetEntitiesTypeCast<UDMXEntityFixturePatch>();
 			OtherFixturePatches.Sort([](const UDMXEntityFixturePatch& OtherPatchA, const UDMXEntityFixturePatch& OtherPatchB)
 				{
@@ -879,6 +879,8 @@ void SDMXFixturePatchTree::OnAddNewFixturePatchClicked(UDMXEntity* InSelectedFix
 				return 1;
 			}();
 
+			const FScopedTransaction Transaction(LOCTEXT("CreateFixturePatchTransaction", "Create DMX Fixture Patch"));
+
 			const FDMXEntityFixtureTypeRef FixtureTypeRef = FDMXEntityFixtureTypeRef(FixtureType);
 			UDMXEntityFixturePatch* NewFixturePatch = UDMXEntityFixturePatch::CreateFixturePatchInLibrary(FixtureTypeRef, FixtureType->Name);
 			if (NewFixturePatch)
diff --git a/Engine/Plugins/VirtualProduction/DMX/DMXEngine/Source/DMXEditor/Private/Widgets/FixtureType/SDMXFixtureTypeTree.cpp b/Engine/Plugins/VirtualProduction/DMX/DMXEngine/Source/DMXEditor/Private/Widgets/FixtureType/SDMXFixtureTypeTree.cpp
--- a/Engine/Plugins/VirtualProduction/DMX/DMXEngine/Source/DMXEditor/Private/Widgets/FixtureType/SDMXFixtureTypeTree.cpp
+++ b/Engine/Plugins/VirtualProduction/DMX/DMXEngine/Source/DMXEditor/Private/Widgets/FixtureType/SDMXFixtureTypeTree.cpp
@@ -397,37 +397,47 @@ void SDMXFixtureTypeTree::OnDeleteNodes()
 	}
 
 	// Confirm deletion of Fixture Types in use
-	if (EntitiesInUse.Num() > 0)
+	const bool bConfirmDelete = [&EntitiesInUse]()
 	{
-		const FText ConfirmDeleteText =
-			FText::Format(LOCTEXT("ConfirmDeleteEntityInUse", "Fixture Patches use {0}|plural(one = this, other = these) Fixture {0}|plural(one = Type, other = Types). Do you really want to delete {0}|plural(one = it, other = them)?"),
-				EntitiesInUse.Num() == 1);
-
-		const FText DialogTitleText = LOCTEXT("DeleteFixtureTypesDialogTitle", "Delete Fixture Types");
-
-		EAppReturnType::Type DialogResult = FMessageDialog::Open(EAppMsgType::YesNo, EAppReturnType::No, ConfirmDeleteText, &DialogTitleText);
-		if (DialogResult == EAppReturnType::Yes)
+		if (EntitiesInUse.Num() > 0)
 		{
-			// Remove the Fixture Types from the DMX Library
-			const FScopedTransaction Transaction(EntitiesToDelete.Num() > 1 ? LOCTEXT("RemoveEntities", "Remove Entities") : LOCTEXT("RemoveEntity", "Remove Entity"));
+			const FText ConfirmDeleteText =
+				FText::Format(LOCTEXT("ConfirmDeleteEntityInUse", "Fixture Patches use {0}|plural(one = this, other = these) Fixture {0}|plural(one = Type, other = Types). Do you really want to delete {0}|plural(one = it, other = them)?"),
+					EntitiesInUse.Num() == 1);
 
-			DMXLibrary->PreEditChange(UDMXLibrary::StaticClass()->FindPropertyByName(UDMXLibrary::GetEntitiesPropertyName()));
-			for (UDMXEntityFixtureType* FixtureType : FixtureTypesToDelete)
-			{
-				ensureMsgf(DMXLibrary == FixtureType->GetParentLibrary(), TEXT("Unexpected DMX Library of Fixture Type and DMX Library of Editor do not match when removing Fixture Types."));
-				const FDMXEntityFixtureTypeRef FixtureTypeRef(FixtureType);
-				UDMXEntityFixtureType::RemoveFixtureTypeFromLibrary(FixtureTypeRef);
-			}
-			DMXLibrary->PostEditChange();
+			const FText DialogTitleText = LOCTEXT("DeleteFixtureTypesDialogTitle", "Delete Fixture Types");
 
-			// Clear selection if no fixture types remain
-			if (GetSelectedEntities().Num() == 0)
+			EAppReturnType::Type DialogResult = FMessageDialog::Open(EAppMsgType::YesNo, EAppReturnType::No, ConfirmDeleteText, &DialogTitleText);
+			if (DialogResult == EAppReturnType::No)
 			{
-				FixtureTypeSharedData->SelectFixtureTypes(TArray<TWeakObjectPtr<UDMXEntityFixtureType>>());
+				return false;
 			}
+		}
 
-			UpdateTree();
+		return true;
+	}();
+
+	if (bConfirmDelete)
+	{
+		// Remove the Fixture Types from the DMX Library
+		const FScopedTransaction Transaction(EntitiesToDelete.Num() > 1 ? LOCTEXT("RemoveEntities", "Remove Entities") : LOCTEXT("RemoveEntity", "Remove Entity"));
+
+		DMXLibrary->PreEditChange(UDMXLibrary::StaticClass()->FindPropertyByName(UDMXLibrary::GetEntitiesPropertyName()));
+		for (UDMXEntityFixtureType* FixtureType : FixtureTypesToDelete)
+		{
+			ensureMsgf(DMXLibrary == FixtureType->GetParentLibrary(), TEXT("Unexpected DMX Library of Fixture Type and DMX Library of Editor do not match when removing Fixture Types."));
+			const FDMXEntityFixtureTypeRef FixtureTypeRef(FixtureType);
+			UDMXEntityFixtureType::RemoveFixtureTypeFromLibrary(FixtureTypeRef);
+		}
+		DMXLibrary->PostEditChange();
+
+		// Clear selection if no fixture types remain
+		if (DMXLibrary->GetEntitiesTypeCast<UDMXEntityFixtureType>().Num() == 0)
+		{
+			FixtureTypeSharedData->SelectFixtureTypes(TArray<TWeakObjectPtr<UDMXEntityFixtureType>>());
 		}
+
+		UpdateTree();
 	}
 }
 
@@ -573,6 +583,8 @@ FReply SDMXFixtureTypeTree::OnAddNewFixtureTypeClicked()
 {
 	if (TSharedPtr<FDMXEditor> PinnedEditor = DMXEditor.Pin())
 	{		
+		const FScopedTransaction Transaction(LOCTEXT("CreateFixtureTypeTransaction", "Create DMX Fixture Type"));
+
 		UDMXEntityFixtureType* NewFixtureType = UDMXEntityFixtureType::CreateFixtureTypeInLibrary(PinnedEditor->GetDMXLibrary());
 
 		FixtureTypeSharedData->SelectFixtureTypes(TArray<TWeakObjectPtr<UDMXEntityFixtureType>>({ NewFixtureType }));
diff --git a/Engine/Plugins/VirtualProduction/DMX/DMXEngine/Source/DMXRuntime/Private/Library/DMXEntityFixturePatch.cpp b/Engine/Plugins/VirtualProduction/DMX/DMXEngine/Source/DMXRuntime/Private/Library/DMXEntityFixturePatch.cpp
--- a/Engine/Plugins/VirtualProduction/DMX/DMXEngine/Source/DMXRuntime/Private/Library/DMXEntityFixturePatch.cpp
+++ b/Engine/Plugins/VirtualProduction/DMX/DMXEngine/Source/DMXRuntime/Private/Library/DMXEntityFixturePatch.cpp
@@ -38,15 +38,23 @@ UDMXEntityFixturePatch::UDMXEntityFixturePatch()
 #endif // WITH_EDITORONLY_DATA
 {}
 
-UDMXEntityFixturePatch* UDMXEntityFixturePatch::CreateFixturePatchInLibrary(FDMXEntityFixtureTypeRef FixtureTypeRef, const FString& DesiredName)
+UDMXEntityFixturePatch* UDMXEntityFixturePatch::CreateFixturePatchInLibrary(FDMXEntityFixtureTypeRef FixtureTypeRef, const FString& DesiredName, bool bMarkDMXLibraryDirty)
 {
 	UDMXEntityFixtureType* FixtureType = FixtureTypeRef.GetFixtureType();
 
-	if (ensureMsgf(FixtureType, TEXT("Cannot create Fixture Patch when Fixture Type is null.")))
+	if (ensureMsgf(FixtureType, TEXT("Cannot create Fixture Patch when Fixture Type is invalid.")))
 	{
 		UDMXLibrary* DMXLibrary = FixtureType->GetParentLibrary();
-		if (ensureMsgf(DMXLibrary, TEXT("Cannot create Fixture Patch when DMX Library is null.")))
+		if (ensureMsgf(DMXLibrary, TEXT("Cannot create Fixture Patch when Fixture Type's DMX Library is invalid.")))
 		{
+#if WITH_EDITOR
+			if (bMarkDMXLibraryDirty)
+			{
+				DMXLibrary->Modify();
+				DMXLibrary->PreEditChange(UDMXLibrary::StaticClass()->FindPropertyByName(UDMXLibrary::GetEntitiesPropertyName()));
+			}
+#endif
+
 			FString EntityName = FDMXRuntimeUtils::FindUniqueEntityName(DMXLibrary, UDMXEntityFixturePatch::StaticClass(), DesiredName);
 
 			UDMXEntityFixturePatch* NewFixturePatch = NewObject<UDMXEntityFixturePatch>(DMXLibrary, UDMXEntityFixturePatch::StaticClass(), NAME_None, RF_Transactional);
@@ -54,6 +62,13 @@ UDMXEntityFixturePatch* UDMXEntityFixturePatch::CreateFixturePatchInLibrary(FDMX
 			NewFixturePatch->SetFixtureType(FixtureType);
 
 			return NewFixturePatch;
+
+#if WITH_EDITOR
+			if (bMarkDMXLibraryDirty)
+			{
+				DMXLibrary->PostEditChange();
+			}
+#endif
 		}
 	}
 
@@ -378,8 +393,6 @@ const FDMXFixtureMode* UDMXEntityFixturePatch::GetActiveMode() const
 
 void UDMXEntityFixturePatch::SetFixtureType(UDMXEntityFixtureType* NewFixtureType)
 {
-	Modify();
-
 	if (!IsValid(NewFixtureType))
 	{
 		ParentFixtureTypeTemplate = nullptr;
@@ -396,8 +409,6 @@ void UDMXEntityFixturePatch::SetFixtureType(UDMXEntityFixtureType* NewFixtureTyp
 
 void UDMXEntityFixturePatch::SetUniverseID(int32 NewUniverseID)
 {
-	Modify();
-
 	UniverseID = NewUniverseID;
 
 	RebuildCache();
@@ -405,8 +416,6 @@ void UDMXEntityFixturePatch::SetUniverseID(int32 NewUniverseID)
 
 void UDMXEntityFixturePatch::SetAutoStartingAddress(int32 NewAutoStartingAddress)
 {
-	Modify();
-
 	AutoStartingAddress = NewAutoStartingAddress;
 	ManualStartingAddress = NewAutoStartingAddress;
 
@@ -415,8 +424,6 @@ void UDMXEntityFixturePatch::SetAutoStartingAddress(int32 NewAutoStartingAddress
 
 void UDMXEntityFixturePatch::SetManualStartingAddress(int32 NewManualStartingAddress)
 {
-	Modify();
-
 	ManualStartingAddress = NewManualStartingAddress;
 
 	RebuildCache();
@@ -424,8 +431,6 @@ void UDMXEntityFixturePatch::SetManualStartingAddress(int32 NewManualStartingAdd
 
 void UDMXEntityFixturePatch::SetStartingChannel(int32 NewStartingChannel)
 {
-	Modify();
-
 	if (NewStartingChannel == AutoStartingAddress && ManualStartingAddress == NewStartingChannel)
 	{
 		return;
@@ -462,8 +467,6 @@ int32 UDMXEntityFixturePatch::GetEndingChannel() const
 
 bool UDMXEntityFixturePatch::SetActiveModeIndex(int32 NewActiveModeIndex)
 {
-	Modify();
-
 	if (IsValid(ParentFixtureTypeTemplate))
 	{
 		if (ParentFixtureTypeTemplate->Modes.IsValidIndex(NewActiveModeIndex))
diff --git a/Engine/Plugins/VirtualProduction/DMX/DMXEngine/Source/DMXRuntime/Private/Library/DMXEntityFixtureType.cpp b/Engine/Plugins/VirtualProduction/DMX/DMXEngine/Source/DMXRuntime/Private/Library/DMXEntityFixtureType.cpp
--- a/Engine/Plugins/VirtualProduction/DMX/DMXEngine/Source/DMXRuntime/Private/Library/DMXEntityFixtureType.cpp
+++ b/Engine/Plugins/VirtualProduction/DMX/DMXEngine/Source/DMXRuntime/Private/Library/DMXEntityFixtureType.cpp
@@ -173,15 +173,30 @@ int32 FDMXFixtureMode::AddOrInsertFunction(int32 IndexOfFunction, FDMXFixtureFun
 
 FDMXOnFixtureTypeChangedDelegate UDMXEntityFixtureType::OnFixtureTypeChangedDelegate;
 
-UDMXEntityFixtureType* UDMXEntityFixtureType::CreateFixtureTypeInLibrary(UDMXLibrary* ParentDMXLibrary, const FString& DesiredName)
+UDMXEntityFixtureType* UDMXEntityFixtureType::CreateFixtureTypeInLibrary(UDMXLibrary* ParentDMXLibrary, const FString& DesiredName, bool bMarkDMXLibraryDirty)
 {
 	if (ensureMsgf(ParentDMXLibrary, TEXT("Create New Fixture Type cannot create Fixture Type when Parent Library is null.")))
 	{
+#if WITH_EDITOR
+		if (bMarkDMXLibraryDirty)
+		{
+			ParentDMXLibrary->Modify();
+			ParentDMXLibrary->PreEditChange(UDMXLibrary::StaticClass()->FindPropertyByName(UDMXLibrary::GetEntitiesPropertyName()));
+		}
+#endif
+
 		FString EntityName = FDMXRuntimeUtils::FindUniqueEntityName(ParentDMXLibrary, UDMXEntityFixtureType::StaticClass(), DesiredName);
 
 		UDMXEntityFixtureType* NewFixtureType = NewObject<UDMXEntityFixtureType>(ParentDMXLibrary, UDMXEntityFixtureType::StaticClass(), NAME_None, RF_Transactional);
 		NewFixtureType->SetName(EntityName);
 
+#if WITH_EDITOR
+		if (bMarkDMXLibraryDirty)
+		{
+			ParentDMXLibrary->PostEditChange();
+		}
+#endif
+
 		return NewFixtureType;
 	}
 
diff --git a/Engine/Plugins/VirtualProduction/DMX/DMXEngine/Source/DMXRuntime/Public/Library/DMXEntityFixturePatch.h b/Engine/Plugins/VirtualProduction/DMX/DMXEngine/Source/DMXRuntime/Public/Library/DMXEntityFixturePatch.h
--- a/Engine/Plugins/VirtualProduction/DMX/DMXEngine/Source/DMXRuntime/Public/Library/DMXEntityFixturePatch.h
+++ b/Engine/Plugins/VirtualProduction/DMX/DMXEngine/Source/DMXRuntime/Public/Library/DMXEntityFixturePatch.h
@@ -45,7 +45,7 @@ public:
 
 	/** Creates a new Fixture Patch in the DMX Library using the specified Fixture Type */
 	UFUNCTION(BlueprintCallable, Category = "DMX")
-	static UDMXEntityFixturePatch* CreateFixturePatchInLibrary(FDMXEntityFixtureTypeRef FixtureTypeRef, const FString& DesiredName = TEXT(""));
+	static UDMXEntityFixturePatch* CreateFixturePatchInLibrary(FDMXEntityFixtureTypeRef FixtureTypeRef, const FString& DesiredName = TEXT(""), bool bMarkDMXLibraryDirty = true);
 
 	/** Removes a fixture patch from the DMX Library */
 	UFUNCTION(BlueprintCallable, Category = "DMX")
diff --git a/Engine/Plugins/VirtualProduction/DMX/DMXEngine/Source/DMXRuntime/Public/Library/DMXEntityFixtureType.h b/Engine/Plugins/VirtualProduction/DMX/DMXEngine/Source/DMXRuntime/Public/Library/DMXEntityFixtureType.h
--- a/Engine/Plugins/VirtualProduction/DMX/DMXEngine/Source/DMXRuntime/Public/Library/DMXEntityFixtureType.h
+++ b/Engine/Plugins/VirtualProduction/DMX/DMXEngine/Source/DMXRuntime/Public/Library/DMXEntityFixtureType.h
@@ -274,8 +274,8 @@ class DMXRUNTIME_API UDMXEntityFixtureType
 public:
 	/** Creates a new Fixture Type in the DMX Library */
 	UFUNCTION(BlueprintCallable, Category = "DMX")
-	static UDMXEntityFixtureType* CreateFixtureTypeInLibrary(UDMXLibrary* ParentDMXLibrary, const FString& DesiredName = TEXT(""));
-	
+	static UDMXEntityFixtureType* CreateFixtureTypeInLibrary(UDMXLibrary* ParentDMXLibrary, const FString& DesiredName = TEXT(""), bool bMarkDMXLibraryDirty = true);
+
 	/** Removes a Fixture Type from a DMX Library */
 	UFUNCTION(BlueprintCallable, Category = "DMX")
 	static void RemoveFixtureTypeFromLibrary(FDMXEntityFixtureTypeRef FixtureTypeRef);
