DMX -  DMX Library Editor shows function names where no attributes are selected as well

diff --git a/Engine/Plugins/VirtualProduction/DMX/DMXEngine/Source/DMXEditor/Private/Widgets/FixtureType/SDMXFixtureTypeFunctionsEditorFunctionRow.cpp b/Engine/Plugins/VirtualProduction/DMX/DMXEngine/Source/DMXEditor/Private/Widgets/FixtureType/SDMXFixtureTypeFunctionsEditorFunctionRow.cpp
--- a/Engine/Plugins/VirtualProduction/DMX/DMXEngine/Source/DMXEditor/Private/Widgets/FixtureType/SDMXFixtureTypeFunctionsEditorFunctionRow.cpp
+++ b/Engine/Plugins/VirtualProduction/DMX/DMXEngine/Source/DMXEditor/Private/Widgets/FixtureType/SDMXFixtureTypeFunctionsEditorFunctionRow.cpp
@@ -132,16 +132,28 @@ TSharedRef<SWidget> SDMXFixtureTypeFunctionsEditorFunctionRow::GenerateWidgetFor
 				SAssignNew(FunctionNameEditableTextBlock, SInlineEditableTextBlock)
 				.Text_Lambda([this]()
 					{
+						const FText FunctionName = FunctionItem->GetFunctionName();
 						if (FunctionItem->HasValidAttribute())
 						{
-							return FunctionItem->GetFunctionName();
+							return FunctionName;
+						}
+						else if (!FunctionName.IsEmpty())
+						{
+							return FText::Format(LOCTEXT("InvalidAttributeFunctionName", "<{0}>"), FunctionItem->GetFunctionName());
+						}
+						else
+						{
+							return LOCTEXT("InvalidAttributeFunctionName", "<Unnamed>");
 						}
-						return LOCTEXT("InvalidAttributeFunctionName", "<Empty Channel - No Attribute Set>");
 					})
 				.IsEnabled_Lambda([this]()
 					{
 						return FunctionItem->HasValidAttribute();
 					})
+				.ToolTipText_Lambda([this]
+					{
+						return LOCTEXT("InvalidAttributeFunctionNameTooltip", "Empty Channel - No Attribute Set");
+					})
 				.Font(FCoreStyle::GetDefaultFontStyle("Regular", 10))
 				.IsReadOnly(false)
 				.OnVerifyTextChanged(this, &SDMXFixtureTypeFunctionsEditorFunctionRow::OnVerifyFunctionNameChanged)
