DMX - Improve ways Fixture Types and Fixture Patches can be created

diff --git a/Engine/Plugins/VirtualProduction/DMX/DMXEngine/Source/DMXEditor/Private/DMXEditor.cpp b/Engine/Plugins/VirtualProduction/DMX/DMXEngine/Source/DMXEditor/Private/DMXEditor.cpp
--- a/Engine/Plugins/VirtualProduction/DMX/DMXEngine/Source/DMXEditor/Private/DMXEditor.cpp
+++ b/Engine/Plugins/VirtualProduction/DMX/DMXEngine/Source/DMXEditor/Private/DMXEditor.cpp
@@ -152,7 +152,10 @@ void FDMXEditor::OnAddNewEntity(TSubclassOf<UDMXEntity> InEntityClass)
 
 	if (InEntityClass == UDMXEntityFixtureType::StaticClass())
 	{
-		UDMXEntityFixtureType* FixtureType = UDMXEntityFixtureType::CreateFixtureTypeInLibrary(DMXLibrary);
+		FDMXEntityFixtureTypeConstructionParams FixtureTypeConstructionParams;
+		FixtureTypeConstructionParams.ParentDMXLibrary = DMXLibrary;
+
+		UDMXEntityFixtureType* FixtureType = UDMXEntityFixtureType::CreateFixtureTypeInLibrary(FixtureTypeConstructionParams);
 		FixtureTypeSharedData->SelectFixtureTypes(TArray<TWeakObjectPtr<UDMXEntityFixtureType>>({ FixtureType }));
 	}
 	else if (InEntityClass == UDMXEntityFixturePatch::StaticClass())
@@ -174,8 +177,10 @@ void FDMXEditor::OnAddNewEntity(TSubclassOf<UDMXEntity> InEntityClass)
 
 			if (LastAddedFixtureType)
 			{
-				const FDMXEntityFixtureTypeRef FixtureTypeRef = FDMXEntityFixtureTypeRef(LastAddedFixtureType);
-				UDMXEntityFixturePatch* NewFixturePatch = UDMXEntityFixturePatch::CreateFixturePatchInLibrary(FixtureTypeRef);
+				FDMXEntityFixturePatchConstructionParams FixturePatchConstructionParams;
+				FixturePatchConstructionParams.FixtureTypeRef = FDMXEntityFixtureTypeRef(LastAddedFixtureType);
+
+				UDMXEntityFixturePatch* NewFixturePatch = UDMXEntityFixturePatch::CreateFixturePatchInLibrary(FixturePatchConstructionParams);
 				FixturePatchSharedData->SelectFixturePatch(NewFixturePatch);
 
 				return;
@@ -186,8 +191,10 @@ void FDMXEditor::OnAddNewEntity(TSubclassOf<UDMXEntity> InEntityClass)
 			const TArray<UDMXEntityFixtureType*> FixtureTypesInLibrary = DMXLibrary->GetEntitiesTypeCast<UDMXEntityFixtureType>();
 			if (FixtureTypesInLibrary.Num() > 0)
 			{
-				const FDMXEntityFixtureTypeRef FixtureTypeRef = FDMXEntityFixtureTypeRef(FixtureTypesInLibrary[0]);
-				UDMXEntityFixturePatch* NewFixturePatch = UDMXEntityFixturePatch::CreateFixturePatchInLibrary(FixtureTypeRef);
+				FDMXEntityFixturePatchConstructionParams FixturePatchConstructionParams;
+				FixturePatchConstructionParams.FixtureTypeRef = FDMXEntityFixtureTypeRef(FixtureTypesInLibrary[0]);
+
+				UDMXEntityFixturePatch* NewFixturePatch = UDMXEntityFixturePatch::CreateFixturePatchInLibrary(FixturePatchConstructionParams);
 				FixturePatchSharedData->SelectFixturePatch(NewFixturePatch);
 			}
 			else
diff --git a/Engine/Plugins/VirtualProduction/DMX/DMXEngine/Source/DMXEditor/Private/Tests/DMXEditorUtilsTests.cpp b/Engine/Plugins/VirtualProduction/DMX/DMXEngine/Source/DMXEditor/Private/Tests/DMXEditorUtilsTests.cpp
--- a/Engine/Plugins/VirtualProduction/DMX/DMXEngine/Source/DMXEditor/Private/Tests/DMXEditorUtilsTests.cpp
+++ b/Engine/Plugins/VirtualProduction/DMX/DMXEngine/Source/DMXEditor/Private/Tests/DMXEditorUtilsTests.cpp
@@ -21,13 +21,18 @@ namespace
 		FFixtureTestSetup(const int32 FixtureChannelLength)
 		{
 			Library							= NewObject<UDMXLibrary>();
-			DummyFixtureType				= UDMXEntityFixtureType::CreateFixtureTypeInLibrary(Library);
-			FDMXEntityFixtureTypeRef FixtureTypeRef(DummyFixtureType);
 
-			FirstManualFixture				= UDMXEntityFixturePatch::CreateFixturePatchInLibrary(FixtureTypeRef);
-			SecondManualFixture				= UDMXEntityFixturePatch::CreateFixturePatchInLibrary(FixtureTypeRef);
-			FirstAutoFixture				= UDMXEntityFixturePatch::CreateFixturePatchInLibrary(FixtureTypeRef);
-			SecondAutoFixture				= UDMXEntityFixturePatch::CreateFixturePatchInLibrary(FixtureTypeRef);
+			FDMXEntityFixtureTypeConstructionParams FixtureTypeConstructionParams;
+			FixtureTypeConstructionParams.ParentDMXLibrary = Library;
+			DummyFixtureType				= UDMXEntityFixtureType::CreateFixtureTypeInLibrary(FixtureTypeConstructionParams);
+
+			FDMXEntityFixturePatchConstructionParams FixturePatchConstructionParams;
+			FixturePatchConstructionParams.FixtureTypeRef = FDMXEntityFixtureTypeRef(DummyFixtureType);
+
+			FirstManualFixture				= UDMXEntityFixturePatch::CreateFixturePatchInLibrary(FixturePatchConstructionParams);
+			SecondManualFixture				= UDMXEntityFixturePatch::CreateFixturePatchInLibrary(FixturePatchConstructionParams);
+			FirstAutoFixture				= UDMXEntityFixturePatch::CreateFixturePatchInLibrary(FixturePatchConstructionParams);
+			SecondAutoFixture				= UDMXEntityFixturePatch::CreateFixturePatchInLibrary(FixturePatchConstructionParams);
 
 			DummyFixtureType->Modes.Add([FixtureChannelLength]()
 				{
diff --git a/Engine/Plugins/VirtualProduction/DMX/DMXEngine/Source/DMXEditor/Private/Widgets/FixturePatch/SDMXFixturePatchTree.cpp b/Engine/Plugins/VirtualProduction/DMX/DMXEngine/Source/DMXEditor/Private/Widgets/FixturePatch/SDMXFixturePatchTree.cpp
--- a/Engine/Plugins/VirtualProduction/DMX/DMXEngine/Source/DMXEditor/Private/Widgets/FixturePatch/SDMXFixturePatchTree.cpp
+++ b/Engine/Plugins/VirtualProduction/DMX/DMXEngine/Source/DMXEditor/Private/Widgets/FixturePatch/SDMXFixturePatchTree.cpp
@@ -775,7 +775,7 @@ void SDMXFixturePatchTree::OnUniverseSelected()
 	const int32 SelectedUniverse = FixturePatchSharedData->GetSelectedUniverse();
 
 	if (AutoExpandedUniverseCategoryNode.IsValid() &&
-		UserExpandedUniverseCategoryNodes.Contains(AutoExpandedUniverseCategoryNode))
+		!UserExpandedUniverseCategoryNodes.Contains(AutoExpandedUniverseCategoryNode))
 	{
 		SetNodeExpansion(AutoExpandedUniverseCategoryNode, false);
 	}
@@ -881,8 +881,10 @@ void SDMXFixturePatchTree::OnAddNewFixturePatchClicked(UDMXEntity* InSelectedFix
 
 			const FScopedTransaction Transaction(LOCTEXT("CreateFixturePatchTransaction", "Create DMX Fixture Patch"));
 
-			const FDMXEntityFixtureTypeRef FixtureTypeRef = FDMXEntityFixtureTypeRef(FixtureType);
-			UDMXEntityFixturePatch* NewFixturePatch = UDMXEntityFixturePatch::CreateFixturePatchInLibrary(FixtureTypeRef, FixtureType->Name);
+			FDMXEntityFixturePatchConstructionParams FixturePatchConstructionParams;
+			FixturePatchConstructionParams.FixtureTypeRef = FDMXEntityFixtureTypeRef(FixtureType);
+
+			UDMXEntityFixturePatch* NewFixturePatch = UDMXEntityFixturePatch::CreateFixturePatchInLibrary(FixturePatchConstructionParams, FixtureType->Name);
 			if (NewFixturePatch)
 			{
 				FDMXEditorUtils::UpdatePatchColors(NewFixturePatch->GetParentLibrary());
diff --git a/Engine/Plugins/VirtualProduction/DMX/DMXEngine/Source/DMXEditor/Private/Widgets/FixtureType/SDMXFixtureTypeTree.cpp b/Engine/Plugins/VirtualProduction/DMX/DMXEngine/Source/DMXEditor/Private/Widgets/FixtureType/SDMXFixtureTypeTree.cpp
--- a/Engine/Plugins/VirtualProduction/DMX/DMXEngine/Source/DMXEditor/Private/Widgets/FixtureType/SDMXFixtureTypeTree.cpp
+++ b/Engine/Plugins/VirtualProduction/DMX/DMXEngine/Source/DMXEditor/Private/Widgets/FixtureType/SDMXFixtureTypeTree.cpp
@@ -585,7 +585,10 @@ FReply SDMXFixtureTypeTree::OnAddNewFixtureTypeClicked()
 	{		
 		const FScopedTransaction Transaction(LOCTEXT("CreateFixtureTypeTransaction", "Create DMX Fixture Type"));
 
-		UDMXEntityFixtureType* NewFixtureType = UDMXEntityFixtureType::CreateFixtureTypeInLibrary(PinnedEditor->GetDMXLibrary());
+		FDMXEntityFixtureTypeConstructionParams FixtureTypeConstructionParams;
+		FixtureTypeConstructionParams.ParentDMXLibrary = PinnedEditor->GetDMXLibrary();
+
+		UDMXEntityFixtureType* NewFixtureType = UDMXEntityFixtureType::CreateFixtureTypeInLibrary(FixtureTypeConstructionParams);
 
 		FixtureTypeSharedData->SelectFixtureTypes(TArray<TWeakObjectPtr<UDMXEntityFixtureType>>({ NewFixtureType }));
 
diff --git a/Engine/Plugins/VirtualProduction/DMX/DMXEngine/Source/DMXEditor/Public/DMXEditor.h b/Engine/Plugins/VirtualProduction/DMX/DMXEngine/Source/DMXEditor/Public/DMXEditor.h
--- a/Engine/Plugins/VirtualProduction/DMX/DMXEngine/Source/DMXEditor/Public/DMXEditor.h
+++ b/Engine/Plugins/VirtualProduction/DMX/DMXEngine/Source/DMXEditor/Public/DMXEditor.h
@@ -49,7 +49,7 @@ public:
 	FReply OnAddNewEntity_OnClick(TSubclassOf<UDMXEntity> InEntityClass) { OnAddNewEntity(InEntityClass); return FReply::Handled(); }
 
 	UE_DEPRECATED(5.0, "Deprecated, this always yields true.")
-	bool CanAddNewEntity(TSubclassOf<UDMXEntity> InEntityClass) const;
+	bool CanAddNewEntity(TSubclassOf<UDMXEntity> InEntityClass) const { return true; }
 
 	bool NewEntity_IsVisibleForType(TSubclassOf<UDMXEntity> InEntityClass) const;
 
diff --git a/Engine/Plugins/VirtualProduction/DMX/DMXEngine/Source/DMXRuntime/Private/Library/DMXEntityFixturePatch.cpp b/Engine/Plugins/VirtualProduction/DMX/DMXEngine/Source/DMXRuntime/Private/Library/DMXEntityFixturePatch.cpp
--- a/Engine/Plugins/VirtualProduction/DMX/DMXEngine/Source/DMXRuntime/Private/Library/DMXEntityFixturePatch.cpp
+++ b/Engine/Plugins/VirtualProduction/DMX/DMXEngine/Source/DMXRuntime/Private/Library/DMXEntityFixturePatch.cpp
@@ -38,9 +38,9 @@ UDMXEntityFixturePatch::UDMXEntityFixturePatch()
 #endif // WITH_EDITORONLY_DATA
 {}
 
-UDMXEntityFixturePatch* UDMXEntityFixturePatch::CreateFixturePatchInLibrary(FDMXEntityFixtureTypeRef FixtureTypeRef, const FString& DesiredName, bool bMarkDMXLibraryDirty)
+UDMXEntityFixturePatch* UDMXEntityFixturePatch::CreateFixturePatchInLibrary(FDMXEntityFixturePatchConstructionParams ConstructionParams, const FString& DesiredName, bool bMarkDMXLibraryDirty)
 {
-	UDMXEntityFixtureType* FixtureType = FixtureTypeRef.GetFixtureType();
+	UDMXEntityFixtureType* FixtureType = ConstructionParams.FixtureTypeRef.GetFixtureType();
 
 	if (ensureMsgf(FixtureType, TEXT("Cannot create Fixture Patch when Fixture Type is invalid.")))
 	{
@@ -60,6 +60,9 @@ UDMXEntityFixturePatch* UDMXEntityFixturePatch::CreateFixturePatchInLibrary(FDMX
 			UDMXEntityFixturePatch* NewFixturePatch = NewObject<UDMXEntityFixturePatch>(DMXLibrary, UDMXEntityFixturePatch::StaticClass(), NAME_None, RF_Transactional);
 			NewFixturePatch->SetName(EntityName);
 			NewFixturePatch->SetFixtureType(FixtureType);
+			NewFixturePatch->SetUniverseID(ConstructionParams.UniverseID);
+			NewFixturePatch->SetStartingChannel(ConstructionParams.StartingAddress);
+			NewFixturePatch->SetActiveModeIndex(ConstructionParams.ActiveMode);
 
 			return NewFixturePatch;
 
diff --git a/Engine/Plugins/VirtualProduction/DMX/DMXEngine/Source/DMXRuntime/Private/Library/DMXEntityFixtureType.cpp b/Engine/Plugins/VirtualProduction/DMX/DMXEngine/Source/DMXRuntime/Private/Library/DMXEntityFixtureType.cpp
--- a/Engine/Plugins/VirtualProduction/DMX/DMXEngine/Source/DMXRuntime/Private/Library/DMXEntityFixtureType.cpp
+++ b/Engine/Plugins/VirtualProduction/DMX/DMXEngine/Source/DMXRuntime/Private/Library/DMXEntityFixtureType.cpp
@@ -173,9 +173,10 @@ int32 FDMXFixtureMode::AddOrInsertFunction(int32 IndexOfFunction, FDMXFixtureFun
 
 FDMXOnFixtureTypeChangedDelegate UDMXEntityFixtureType::OnFixtureTypeChangedDelegate;
 
-UDMXEntityFixtureType* UDMXEntityFixtureType::CreateFixtureTypeInLibrary(UDMXLibrary* ParentDMXLibrary, const FString& DesiredName, bool bMarkDMXLibraryDirty)
+UDMXEntityFixtureType* UDMXEntityFixtureType::CreateFixtureTypeInLibrary(FDMXEntityFixtureTypeConstructionParams ConstructionParams, const FString& DesiredName, bool bMarkDMXLibraryDirty)
 {
-	if (ensureMsgf(ParentDMXLibrary, TEXT("Create New Fixture Type cannot create Fixture Type when Parent Library is null.")))
+	UDMXLibrary* ParentDMXLibrary = ConstructionParams.ParentDMXLibrary;
+	if (ensureMsgf(IsValid(ConstructionParams.ParentDMXLibrary), TEXT("Create New Fixture Type cannot create Fixture Type when Parent Library is null.")))
 	{
 #if WITH_EDITOR
 		if (bMarkDMXLibraryDirty)
@@ -189,6 +190,8 @@ UDMXEntityFixtureType* UDMXEntityFixtureType::CreateFixtureTypeInLibrary(UDMXLib
 
 		UDMXEntityFixtureType* NewFixtureType = NewObject<UDMXEntityFixtureType>(ParentDMXLibrary, UDMXEntityFixtureType::StaticClass(), NAME_None, RF_Transactional);
 		NewFixtureType->SetName(EntityName);
+		NewFixtureType->DMXCategory = ConstructionParams.DMXCategory;
+		NewFixtureType->Modes = ConstructionParams.Modes;
 
 #if WITH_EDITOR
 		if (bMarkDMXLibraryDirty)
@@ -197,6 +200,8 @@ UDMXEntityFixtureType* UDMXEntityFixtureType::CreateFixtureTypeInLibrary(UDMXLib
 		}
 #endif
 
+		OnFixtureTypeChangedDelegate.Broadcast(NewFixtureType);
+
 		return NewFixtureType;
 	}
 
diff --git a/Engine/Plugins/VirtualProduction/DMX/DMXEngine/Source/DMXRuntime/Public/Library/DMXEntityFixturePatch.h b/Engine/Plugins/VirtualProduction/DMX/DMXEngine/Source/DMXRuntime/Public/Library/DMXEntityFixturePatch.h
--- a/Engine/Plugins/VirtualProduction/DMX/DMXEngine/Source/DMXRuntime/Public/Library/DMXEntityFixturePatch.h
+++ b/Engine/Plugins/VirtualProduction/DMX/DMXEngine/Source/DMXRuntime/Public/Library/DMXEntityFixturePatch.h
@@ -26,6 +26,30 @@ struct FPropertyChangedEvent;
 
 DECLARE_MULTICAST_DELEGATE_OneParam(FDMXOnFixturePatchChangedDelegate, const UDMXEntityFixturePatch* /** ChangedFixturePatch */);
 
+
+/** Parameters to construct a Fixture Patch. */
+USTRUCT(BlueprintType)
+struct DMXRUNTIME_API FDMXEntityFixturePatchConstructionParams
+{
+	GENERATED_BODY()
+	
+	/** Property to point to the template parent fixture for details panel purposes */
+	UPROPERTY(VisibleDefaultsOnly, BlueprintReadWrite, Category = "Fixture Patch", meta = (DisplayName = "Fixture Type"))
+	FDMXEntityFixtureTypeRef FixtureTypeRef;
+
+	/** The Index of the Mode in the Fixture Type the Patch uses */
+	UPROPERTY(EditAnywhere, BlueprintReadOnly, Category = "Fixture Patch")
+	int32 ActiveMode = 0;
+
+	/** The local universe of the patch */
+	UPROPERTY(EditAnywhere, BlueprintReadWrite, Category = "Fixture Patch", meta = (ClampMin = 0, DisplayName = "Universe"))
+	int32 UniverseID = 1;
+
+	/** Starting channel for when auto-assign address is false */
+	UPROPERTY(EditAnywhere, BlueprintReadWrite, Category = "Fixture Patch", meta = (EditCondition = "!bAutoAssignAddress", DisplayName = "Manual Starting Address", UIMin = "1", UIMax = "512", ClampMin = "1", ClampMax = "512"))
+	int32 StartingAddress = 1;
+};
+
 /** 
  * A DMX fixture patch that can be patch to channels in a DMX Universe via the DMX Library Editor. 
  * 
@@ -45,7 +69,7 @@ public:
 
 	/** Creates a new Fixture Patch in the DMX Library using the specified Fixture Type */
 	UFUNCTION(BlueprintCallable, Category = "DMX")
-	static UDMXEntityFixturePatch* CreateFixturePatchInLibrary(FDMXEntityFixtureTypeRef FixtureTypeRef, const FString& DesiredName = TEXT(""), bool bMarkDMXLibraryDirty = true);
+	static UDMXEntityFixturePatch* CreateFixturePatchInLibrary(FDMXEntityFixturePatchConstructionParams ConstructionParams, const FString& DesiredName = TEXT(""), bool bMarkDMXLibraryDirty = true);
 
 	/** Removes a fixture patch from the DMX Library */
 	UFUNCTION(BlueprintCallable, Category = "DMX")
@@ -123,6 +147,7 @@ public:
 	FORCEINLINE UDMXEntityFixtureType* GetFixtureType() const { return ParentFixtureTypeTemplate; }
 
 	/** Sets the fixture type this is using */
+	UFUNCTION(BlueprintCallable, Category = "DMX|Fixture Patch")
 	void SetFixtureType(UDMXEntityFixtureType* NewFixtureType);
 
 	/** Returns the universe ID of the patch*/
@@ -208,7 +233,7 @@ protected:
 	int32 AutoStartingAddress;
 
 	/** Property to point to the template parent fixture for details panel purposes */
-	UPROPERTY(VisibleDefaultsOnly, BlueprintReadWrite, Category = "Fixture Patch", meta = (DisplayName = "Fixture Type"))
+	UPROPERTY(VisibleDefaultsOnly, BlueprintReadWrite, BlueprintSetter = SetFixtureType, Category = "Fixture Patch", meta = (DisplayName = "Fixture Type"))
 	UDMXEntityFixtureType* ParentFixtureTypeTemplate;
 
 	/** The Index of the Mode in the Fixture Type the Patch uses */
diff --git a/Engine/Plugins/VirtualProduction/DMX/DMXEngine/Source/DMXRuntime/Public/Library/DMXEntityFixtureType.h b/Engine/Plugins/VirtualProduction/DMX/DMXEngine/Source/DMXRuntime/Public/Library/DMXEntityFixtureType.h
--- a/Engine/Plugins/VirtualProduction/DMX/DMXEngine/Source/DMXRuntime/Public/Library/DMXEntityFixtureType.h
+++ b/Engine/Plugins/VirtualProduction/DMX/DMXEngine/Source/DMXRuntime/Public/Library/DMXEntityFixtureType.h
@@ -66,21 +66,21 @@ struct DMXRUNTIME_API FDMXFixtureFunction
 	 * The list of Attributes can be edited on
 	 * Project Settings->Plugins->DMX Protocol->Fixture Settings->Fixture Function Attributes
 	 */
-	UPROPERTY(EditAnywhere, BlueprintReadOnly, meta = (DisplayName = "Attribute Mapping", DisplayPriority = "11"), Category = "Function Settings")
+	UPROPERTY(EditAnywhere, BlueprintReadWrite, meta = (DisplayName = "Attribute Mapping", DisplayPriority = "11"), Category = "Function Settings")
 	FDMXAttributeName Attribute;
 
-	UPROPERTY(EditAnywhere, BlueprintReadOnly, meta = (DisplayPriority = "10"), Category = "Function Settings")
+	UPROPERTY(EditAnywhere, BlueprintReadWrite, meta = (DisplayPriority = "10"), Category = "Function Settings")
 	FString FunctionName;
 
-	UPROPERTY(EditAnywhere, BlueprintReadOnly, meta = (DisplayPriority = "20"), Category = "Function Settings")
+	UPROPERTY(EditAnywhere, BlueprintReadWrite, meta = (DisplayPriority = "20"), Category = "Function Settings")
 	FString Description;
 
 	/** The Default Value of the function, imported from GDTF. The plugin doesn't make use of this value, but it can be used in blueprints */
-	UPROPERTY(EditAnywhere, BlueprintReadOnly, meta = (DisplayPriority = "30"), Category = "Function Settings")
+	UPROPERTY(EditAnywhere, BlueprintReadWrite, meta = (DisplayPriority = "30"), Category = "Function Settings")
 	int64 DefaultValue;
 
 	/** This function's starting channel (use editor above to make changes) */
-	UPROPERTY(VisibleAnywhere, BlueprintReadOnly, meta = (DisplayName = "Channel Assignment", DisplayPriority = "2"), Category = "Function Settings")
+	UPROPERTY(VisibleAnywhere, BlueprintReadWrite, meta = (DisplayName = "Channel Assignment", DisplayPriority = "2"), Category = "Function Settings")
 	int32 Channel;
 
 	/** DEPRECATED 5.0. Instead the 'Channel' property is EditAnywhere so any function can be assigned freely */
@@ -88,7 +88,7 @@ struct DMXRUNTIME_API FDMXFixtureFunction
 	int32 ChannelOffset_DEPRECATED;
 
 	/** This function's data type. Defines the used number of channels (bytes) */
-	UPROPERTY(EditAnywhere, BlueprintReadOnly, meta = (DisplayPriority = "5"), Category = "Function Settings")
+	UPROPERTY(EditAnywhere, BlueprintReadWrite, meta = (DisplayPriority = "5"), Category = "Function Settings")
 	EDMXFixtureSignalFormat DataType;
 
 	/**
@@ -103,7 +103,7 @@ struct DMXRUNTIME_API FDMXFixtureFunction
 	 * In MSB mode, the example above would be interpreted in binary as 0x00 0x01, which means 1.
 	 * The first byte (0) became the highest part in binary form and the following byte (1), the lowest.
 	 */
-	UPROPERTY(EditAnywhere, BlueprintReadOnly, meta = (DisplayName = "Use LSB Mode", DisplayPriority = "29"), Category = "Function Settings")
+	UPROPERTY(EditAnywhere, BlueprintReadWrite, meta = (DisplayName = "Use LSB Mode", DisplayPriority = "29"), Category = "Function Settings")
 	bool bUseLSBMode = false;
 };
 
@@ -131,18 +131,18 @@ struct DMXRUNTIME_API FDMXFixtureCellAttribute
 	 * The list of Attributes can be edited on
 	 * Project Settings->Plugins->DMX Protocol->Fixture Settings->Fixture Function Attributes
 	 */
-	UPROPERTY(EditAnywhere, BlueprintReadOnly, meta = (DisplayName = "Attribute Mapping", DisplayPriority = "11"), Category = "DMX")
+	UPROPERTY(EditAnywhere, BlueprintReadWrite, meta = (DisplayName = "Attribute Mapping", DisplayPriority = "11"), Category = "DMX")
 	FDMXAttributeName Attribute;
 
-	UPROPERTY(EditAnywhere, BlueprintReadOnly, meta = (DisplayPriority = "20", DisplayName = "Description"), Category = "DMX")
+	UPROPERTY(EditAnywhere, BlueprintReadWrite, meta = (DisplayPriority = "20", DisplayName = "Description"), Category = "DMX")
 	FString Description;
 
 	/** Initial value for this function when no value is set */
-	UPROPERTY(EditAnywhere, BlueprintReadOnly, meta = (DisplayPriority = "30", DisplayName = "Default Value"), Category = "DMX")
+	UPROPERTY(EditAnywhere, BlueprintReadWrite, meta = (DisplayPriority = "30", DisplayName = "Default Value"), Category = "DMX")
 	int64 DefaultValue;
 
 	/** This function's data type. Defines the used number of channels (bytes) */
-	UPROPERTY(EditAnywhere, BlueprintReadOnly, meta = (DisplayPriority = "5", DisplayName = "Data Type"), Category = "DMX")
+	UPROPERTY(EditAnywhere, BlueprintReadWrite, meta = (DisplayPriority = "5", DisplayName = "Data Type"), Category = "DMX")
 	EDMXFixtureSignalFormat DataType;
 
 	/**
@@ -158,7 +158,7 @@ struct DMXRUNTIME_API FDMXFixtureCellAttribute
 	 * In MSB mode, the example above would be interpreted in binary as 00000000 00000001, which means 1.
 	 * The first byte (0) became the highest part in binary form and the following byte (1), the lowest.
 	 */
-	UPROPERTY(EditAnywhere, BlueprintReadOnly, meta = (DisplayName = "Use LSB Mode", DisplayPriority = "29"), Category = "DMX")
+	UPROPERTY(EditAnywhere, BlueprintReadWrite, meta = (DisplayName = "Use LSB Mode", DisplayPriority = "29"), Category = "DMX")
 	bool bUseLSBMode;
 };
 
@@ -182,19 +182,19 @@ struct DMXRUNTIME_API FDMXFixtureMatrix
 	UE_DEPRECATED(5.0, "Deprecated to reduce redundant code. Instead use UDMXFixturePatch::GetMatrixCellChannelsRelative")
 	bool GetChannelsFromCell(FIntPoint CellCoordinate, FDMXAttributeName Attribute, TArray<int32>& Channels) const;
 
-	UPROPERTY(EditAnywhere, BlueprintReadOnly, meta = (DisplayPriority = "60", DisplayName = "Cell Attributes"), Category = "Mode Settings")
+	UPROPERTY(EditAnywhere, BlueprintReadWrite, meta = (DisplayPriority = "60", DisplayName = "Cell Attributes"), Category = "Mode Settings")
 	TArray<FDMXFixtureCellAttribute> CellAttributes;
 
-	UPROPERTY(VisibleAnywhere, BlueprintReadOnly, meta = (DisplayPriority = "20", DisplayName = "First Cell Channel", ClampMin = "1", ClampMax = "512"), Category = "Mode Settings")
+	UPROPERTY(VisibleAnywhere, BlueprintReadWrite, meta = (DisplayPriority = "20", DisplayName = "First Cell Channel", ClampMin = "1", ClampMax = "512"), Category = "Mode Settings")
 	int32 FirstCellChannel = 1;
 
-	UPROPERTY(EditAnywhere, BlueprintReadOnly, meta = (DisplayPriority = "30", DisplayName = "X Cells", ClampMin = "1"), Category = "Mode Settings")
+	UPROPERTY(EditAnywhere, BlueprintReadWrite, meta = (DisplayPriority = "30", DisplayName = "X Cells", ClampMin = "1"), Category = "Mode Settings")
 	int32 XCells = 1;
 
-	UPROPERTY(EditAnywhere, BlueprintReadOnly, meta = (DisplayPriority = "40", DisplayName = "Y Cells", ClampMin = "1"), Category = "Mode Settings")
+	UPROPERTY(EditAnywhere, BlueprintReadWrite, meta = (DisplayPriority = "40", DisplayName = "Y Cells", ClampMin = "1"), Category = "Mode Settings")
 	int32 YCells = 1;
 
-	UPROPERTY(EditAnywhere, BlueprintReadOnly, meta = (DisplayPriority = "50", DisplayName = "PixelMapping Distribution"), Category = "Mode Settings")
+	UPROPERTY(EditAnywhere, BlueprintReadWrite, meta = (DisplayPriority = "50", DisplayName = "PixelMapping Distribution"), Category = "Mode Settings")
 	EDMXPixelMappingDistribution PixelMappingDistribution = EDMXPixelMappingDistribution::TopLeftToRight;
 };
 
@@ -204,11 +204,11 @@ struct DMXRUNTIME_API FDMXCell
 	GENERATED_BODY()
 
 	/** The cell index in a 1D Array (row order), starting from 0 */
-	UPROPERTY(EditAnywhere, BlueprintReadOnly, meta = (DisplayPriority = "20", DisplayName = "Cell ID", ClampMin = "0"), Category = "DMX")
+	UPROPERTY(EditAnywhere, BlueprintReadWrite, meta = (DisplayPriority = "20", DisplayName = "Cell ID", ClampMin = "0"), Category = "DMX")
 	int32 CellID;
 
 	/** The cell coordinate in a 2D Array, starting from (0, 0) */
-	UPROPERTY(EditAnywhere, BlueprintReadOnly, meta = (DisplayPriority = "30", DisplayName = "Coordinate"), Category = "DMX")
+	UPROPERTY(EditAnywhere, BlueprintReadWrite, meta = (DisplayPriority = "30", DisplayName = "Coordinate"), Category = "DMX")
 	FIntPoint Coordinate;
 
 	FDMXCell()
@@ -222,10 +222,10 @@ struct DMXRUNTIME_API FDMXFixtureMode
 {
 	GENERATED_BODY()
 
-	UPROPERTY(EditAnywhere, BlueprintReadOnly, meta = (DisplayPriority = "10"), Category = "Mode Settings")
+	UPROPERTY(EditAnywhere, BlueprintReadWrite, meta = (DisplayPriority = "10"), Category = "Mode Settings")
 	FString ModeName;
 
-	UPROPERTY(EditAnywhere, BlueprintReadOnly, meta = (DisplayPriority = "20"), Category = "Mode Settings")
+	UPROPERTY(EditAnywhere, BlueprintReadWrite, meta = (DisplayPriority = "20"), Category = "Mode Settings")
 	TArray<FDMXFixtureFunction> Functions;
 
 	/**
@@ -237,10 +237,10 @@ struct DMXRUNTIME_API FDMXFixtureMode
 	bool bAutoChannelSpan = true;
 
 	/** Number of channels (bytes) used by this mode's functions */
-	UPROPERTY(EditAnywhere, BlueprintReadOnly, Category = "Mode Settings", meta = (ClampMin = "4", ClampMax = "512", DisplayPriority = "40", EditCondition = "!bAutoChannelSpan"))
+	UPROPERTY(EditAnywhere, BlueprintReadWrite, Category = "Mode Settings", meta = (ClampMin = "4", ClampMax = "512", DisplayPriority = "40", EditCondition = "!bAutoChannelSpan"))
 	int32 ChannelSpan = 0;
 
-	UPROPERTY(EditAnywhere, BlueprintReadOnly, Category = "Mode Settings", meta = (DisplayPriority = "60"))
+	UPROPERTY(EditAnywhere, BlueprintReadWrite, Category = "Mode Settings", meta = (DisplayPriority = "60"))
 	bool bFixtureMatrixEnabled = false;
 
 	UPROPERTY(EditAnywhere, Category = "Mode Settings", meta = (DisplayPriority = "70"))
@@ -254,6 +254,25 @@ struct DMXRUNTIME_API FDMXFixtureMode
 };
 
 
+/** Parameters to construct a Fixture Type. */
+USTRUCT(BlueprintType)
+struct DMXRUNTIME_API FDMXEntityFixtureTypeConstructionParams
+{
+	GENERATED_BODY()
+
+	/** The DMX Library in which the Fixture Type will be constructed */
+	UPROPERTY(EditAnywhere, BlueprintReadWrite, Category = "Fixture Type")
+	UDMXLibrary* ParentDMXLibrary;
+
+	/** The Category of the Fixture, useful for Filtering */
+	UPROPERTY(EditAnywhere, BlueprintReadWrite, Category = "Fixture Type")
+	FDMXFixtureCategory DMXCategory;
+
+	/** The Modes of the Fixture Type */
+	UPROPERTY(EditAnywhere, BlueprintReadWrite, Category = "Fixture Type")
+	TArray<FDMXFixtureMode> Modes;
+};
+
 #if WITH_EDITOR
 /** Notification when data type changed */
 DECLARE_MULTICAST_DELEGATE_TwoParams(FDataTypeChangeDelegate, const UDMXEntityFixtureType*, const FDMXFixtureMode&);
@@ -274,7 +293,7 @@ class DMXRUNTIME_API UDMXEntityFixtureType
 public:
 	/** Creates a new Fixture Type in the DMX Library */
 	UFUNCTION(BlueprintCallable, Category = "DMX")
-	static UDMXEntityFixtureType* CreateFixtureTypeInLibrary(UDMXLibrary* ParentDMXLibrary, const FString& DesiredName = TEXT(""), bool bMarkDMXLibraryDirty = true);
+	static UDMXEntityFixtureType* CreateFixtureTypeInLibrary(FDMXEntityFixtureTypeConstructionParams ConstructionParams, const FString& DesiredName = TEXT(""), bool bMarkDMXLibraryDirty = true);
 
 	/** Removes a Fixture Type from a DMX Library */
 	UFUNCTION(BlueprintCallable, Category = "DMX")
