DMX - Fix CIS Build DDC Win64

diff --git a/Engine/Plugins/VirtualProduction/DMX/DMXEngine/Source/DMXRuntime/Public/Library/DMXEntityFixtureType.h b/Engine/Plugins/VirtualProduction/DMX/DMXEngine/Source/DMXRuntime/Public/Library/DMXEntityFixtureType.h
--- a/Engine/Plugins/VirtualProduction/DMX/DMXEngine/Source/DMXRuntime/Public/Library/DMXEntityFixtureType.h
+++ b/Engine/Plugins/VirtualProduction/DMX/DMXEngine/Source/DMXRuntime/Public/Library/DMXEntityFixtureType.h
@@ -262,7 +262,7 @@ struct DMXRUNTIME_API FDMXEntityFixtureTypeConstructionParams
 
 	/** The DMX Library in which the Fixture Type will be constructed */
 	UPROPERTY(EditAnywhere, BlueprintReadWrite, Category = "Fixture Type")
-	UDMXLibrary* ParentDMXLibrary;
+	UDMXLibrary* ParentDMXLibrary = nullptr;
 
 	/** The Category of the Fixture, useful for Filtering */
 	UPROPERTY(EditAnywhere, BlueprintReadWrite, Category = "Fixture Type")
