DMX - Improve log message and lower verbosity for UDMXEntityFixturePatch::GetMatrixProperties

diff --git a/Engine/Plugins/VirtualProduction/DMX/DMXEngine/Source/DMXRuntime/Private/Library/DMXEntityFixturePatch.cpp b/Engine/Plugins/VirtualProduction/DMX/DMXEngine/Source/DMXRuntime/Private/Library/DMXEntityFixturePatch.cpp
--- a/Engine/Plugins/VirtualProduction/DMX/DMXEngine/Source/DMXRuntime/Private/Library/DMXEntityFixturePatch.cpp
+++ b/Engine/Plugins/VirtualProduction/DMX/DMXEngine/Source/DMXRuntime/Private/Library/DMXEntityFixturePatch.cpp
@@ -1006,20 +1006,20 @@ bool UDMXEntityFixturePatch::GetMatrixProperties(FDMXFixtureMatrix& MatrixProper
 {
 	if (!ParentFixtureTypeTemplate)
 	{
-		UE_LOG(DMXEntityFixturePatchLog, Error, TEXT("Fixture Patch %s has no parent fixture type assigned"), *Name);
+		UE_LOG(DMXEntityFixturePatchLog, Warning, TEXT("'Get Matrix Properties' failed, Fixture Patch %s has no parent fixture type assigned"), *Name);
 		return false;
 	}
 
 	const FDMXFixtureMode* ModePtr = GetActiveMode();
 	if (!ModePtr)
 	{
-		UE_LOG(DMXEntityFixturePatchLog, Error, TEXT("Invalid active Mode in Fixture Patch %s"), *Name);
+		UE_LOG(DMXEntityFixturePatchLog, Warning, TEXT("'Get Matrix Properties' failed, Invalid active Mode in Fixture Patch %s"), *Name);
 		return false;
 	}
 
 	if (!ModePtr->bFixtureMatrixEnabled)
 	{
-		UE_LOG(DMXEntityFixturePatchLog, Error, TEXT("Fixture Patch %s is not a Matrix Fixture"), *Name);
+		UE_LOG(DMXEntityFixturePatchLog, Warning, TEXT("'Get Matrix Properties' failed, Fixture Patch %s is not a Matrix Fixture"), *Name);
 		return false;
 	}
 
