DMX - Fix various bugs in DMX Monitors and Output Console

diff --git a/Engine/Plugins/VirtualProduction/DMX/DMXEngine/Source/DMXEditor/Private/Widgets/Monitors/SDMXActivityMonitor.cpp b/Engine/Plugins/VirtualProduction/DMX/DMXEngine/Source/DMXEditor/Private/Widgets/Monitors/SDMXActivityMonitor.cpp
--- a/Engine/Plugins/VirtualProduction/DMX/DMXEngine/Source/DMXEditor/Private/Widgets/Monitors/SDMXActivityMonitor.cpp
+++ b/Engine/Plugins/VirtualProduction/DMX/DMXEngine/Source/DMXEditor/Private/Widgets/Monitors/SDMXActivityMonitor.cpp
@@ -35,7 +35,7 @@ SDMXActivityMonitor::~SDMXActivityMonitor()
 {
 	check(SourceSelector.IsValid());
 
-	for (const TSharedRef<FDMXRawListener>& Input : DMXInputs)
+	for (const TSharedRef<FDMXRawListener>& Input : DMXListeners)
 	{
 		Input->Stop();
 	}
@@ -203,7 +203,7 @@ void SDMXActivityMonitor::ClearAllDMXBuffers()
 	FDMXEditorUtils::ClearAllDMXPortBuffers();
 	UniverseToDataMap.Reset();
 
-	for (const TSharedRef<FDMXRawListener>& Input : DMXInputs)
+	for (const TSharedRef<FDMXRawListener>& Input : DMXListeners)
 	{
 		Input->ClearBuffer();
 	}
@@ -228,7 +228,7 @@ void SDMXActivityMonitor::Tick(const FGeometry& AllottedGeometry, const double I
 	// Tick universe listeners with latest signals on tick
 
 	FDMXSignalSharedPtr DMXSignal;
-	for (const TSharedRef<FDMXRawListener>& Input : DMXInputs)
+	for (const TSharedRef<FDMXRawListener>& Input : DMXListeners)
 	{
 		int32 LocalUniverseID;
 		while (Input->DequeueSignal(this, DMXSignal, LocalUniverseID))
@@ -439,28 +439,33 @@ void SDMXActivityMonitor::UpdateListenerRegistration()
 	check(SourceSelector.IsValid());
 
 	// Stop listening to previous ports
-	for (const TSharedRef<FDMXRawListener>& Input : DMXInputs)
+	for (const TSharedRef<FDMXRawListener>& Input : DMXListeners)
 	{
 		Input->Stop();
 	}
-	DMXInputs.Reset();
+	DMXListeners.Reset();
 
-	// Listen to selected ports
-	for (const FDMXInputPortSharedRef& InputPort : SourceSelector->GetSelectedInputPorts())
+	if (SourceSelector->IsMonitorInputPorts())
 	{
-		TSharedRef<FDMXRawListener> NewInput = MakeShared<FDMXRawListener>(InputPort);
-		DMXInputs.Add(NewInput);
+		// Listen to selected ports
+		for (const FDMXInputPortSharedRef& InputPort : SourceSelector->GetSelectedInputPorts())
+		{
+			TSharedRef<FDMXRawListener> NewListener = MakeShared<FDMXRawListener>(InputPort);
+			DMXListeners.Add(NewListener);
 
-		NewInput->Start();
+			NewListener->Start();
+		}
 	}
-
-	for (const FDMXOutputPortSharedRef& OutputPort : SourceSelector->GetSelectedOutputPorts())
+	else
 	{
-		// Monitor outputs 
-		TSharedRef<FDMXRawListener> NewInput = MakeShared<FDMXRawListener>(OutputPort);
-		DMXInputs.Add(NewInput);
+		for (const FDMXOutputPortSharedRef& OutputPort : SourceSelector->GetSelectedOutputPorts())
+		{
+			// Monitor outputs 
+			TSharedRef<FDMXRawListener> NewListener = MakeShared<FDMXRawListener>(OutputPort);
+			DMXListeners.Add(NewListener);
 
-		NewInput->Start();
+			NewListener->Start();
+		}
 	}
 }
 
diff --git a/Engine/Plugins/VirtualProduction/DMX/DMXEngine/Source/DMXEditor/Private/Widgets/Monitors/SDMXActivityMonitor.h b/Engine/Plugins/VirtualProduction/DMX/DMXEngine/Source/DMXEditor/Private/Widgets/Monitors/SDMXActivityMonitor.h
--- a/Engine/Plugins/VirtualProduction/DMX/DMXEngine/Source/DMXEditor/Private/Widgets/Monitors/SDMXActivityMonitor.h
+++ b/Engine/Plugins/VirtualProduction/DMX/DMXEngine/Source/DMXEditor/Private/Widgets/Monitors/SDMXActivityMonitor.h
@@ -93,7 +93,7 @@ private:
 	TSharedRef<ITableRow> OnGenerateUniverseRow(TSharedPtr<SDMXActivityInUniverse> Item, const TSharedRef<STableViewBase>& OwnerTable);
 
 	/** The DMX inputs of the activity monitor */
-	TSet<TSharedRef<FDMXRawListener>> DMXInputs;
+	TSet<TSharedRef<FDMXRawListener>> DMXListeners;
 
 	/** ListView for universe counts */
 	TSharedPtr<SListView<TSharedPtr<SDMXActivityInUniverse>>> UniverseList;
diff --git a/Engine/Plugins/VirtualProduction/DMX/DMXEngine/Source/DMXEditor/Private/Widgets/Monitors/SDMXChannelsMonitor.cpp b/Engine/Plugins/VirtualProduction/DMX/DMXEngine/Source/DMXEditor/Private/Widgets/Monitors/SDMXChannelsMonitor.cpp
--- a/Engine/Plugins/VirtualProduction/DMX/DMXEngine/Source/DMXEditor/Private/Widgets/Monitors/SDMXChannelsMonitor.cpp
+++ b/Engine/Plugins/VirtualProduction/DMX/DMXEngine/Source/DMXEditor/Private/Widgets/Monitors/SDMXChannelsMonitor.cpp
@@ -139,24 +139,26 @@ void SDMXChannelsMonitor::Tick(const FGeometry& AllottedGeometry, const double I
 {
 	FDMXSignalSharedPtr Signal;
 
-	const UDMXProtocolSettings* ProtocolSettings = GetDefault<UDMXProtocolSettings>();
-	const bool bSendOrReceiveDisabled = !ProtocolSettings->IsSendDMXEnabled() || !ProtocolSettings->IsReceiveDMXEnabled();
-
-	// The widget returns the right arrays, also output ports that loopback when input ports are selected.
-	for (const FDMXInputPortSharedRef& InputPort : SourceSelector->GetSelectedInputPorts())
+	if (SourceSelector->IsMonitorInputPorts())
 	{
-		if (InputPort->GameThreadGetDMXSignal(UniverseID, Signal))
+		// The widget returns the right arrays, also output ports that loopback when input ports are selected.
+		for (const FDMXInputPortSharedRef& InputPort : SourceSelector->GetSelectedInputPorts())
 		{
-			break;
+			if (InputPort->GameThreadGetDMXSignal(UniverseID, Signal))
+			{
+				break;
+			}
 		}
 	}
-
-	for (const FDMXOutputPortSharedRef& OutputPort : SourceSelector->GetSelectedOutputPorts())
+	else
 	{
-		constexpr bool bEvenIfNotLoopbackToEngine = true;
-		if (OutputPort->GameThreadGetDMXSignal(UniverseID, Signal, bEvenIfNotLoopbackToEngine))
+		for (const FDMXOutputPortSharedRef& OutputPort : SourceSelector->GetSelectedOutputPorts())
 		{
-			break;
+			constexpr bool bEvenIfNotLoopbackToEngine = true;
+			if (OutputPort->GameThreadGetDMXSignal(UniverseID, Signal, bEvenIfNotLoopbackToEngine))
+			{
+				break;
+			}
 		}
 	}
 
diff --git a/Engine/Plugins/VirtualProduction/DMX/DMXEngine/Source/DMXEditor/Private/Widgets/Monitors/SDMXMonitorSourceSelector.cpp b/Engine/Plugins/VirtualProduction/DMX/DMXEngine/Source/DMXEditor/Private/Widgets/Monitors/SDMXMonitorSourceSelector.cpp
--- a/Engine/Plugins/VirtualProduction/DMX/DMXEngine/Source/DMXEditor/Private/Widgets/Monitors/SDMXMonitorSourceSelector.cpp
+++ b/Engine/Plugins/VirtualProduction/DMX/DMXEngine/Source/DMXEditor/Private/Widgets/Monitors/SDMXMonitorSourceSelector.cpp
@@ -160,10 +160,17 @@ bool SDMXMonitorSourceSelector::IsMonitorInputPorts() const
 	check(IODirectionNames.Num() == 2);
 	check(MonitoredIODirectionComboBox.IsValid());
 
-	TSharedPtr<FText> SelectedDirectionItem = MonitoredIODirectionComboBox->GetSelectedItem();
-	check(SelectedDirectionItem.IsValid());
+	if (IsMonitorAllPorts())
+	{
+		TSharedPtr<FText> SelectedDirectionItem = MonitoredIODirectionComboBox->GetSelectedItem();
+		return SelectedDirectionItem == IODirectionNames[0];
+	}
+	else if(PortSelector->GetSelectedInputPort().IsValid())
+	{
+		return true;
+	}
 
-	return SelectedDirectionItem == IODirectionNames[0];
+	return false;
 }
 
 void SDMXMonitorSourceSelector::SetMonitorInputPorts(bool bMonitorInputPorts)
@@ -226,16 +233,13 @@ void SDMXMonitorSourceSelector::UpdateMonitoredPorts()
 		MonitoredIODirectionWrapper->SetVisibility(EVisibility::Visible);
 		PortSelectorWrapper->SetVisibility(EVisibility::Collapsed);
 
-		const bool bInputSelected = MonitoredIODirectionComboBox->GetSelectedItem() == IODirectionNames[0];
-		if (bInputSelected)
+		if (IsMonitorInputPorts())
 		{
 			SelectedInputPorts = FDMXPortManager::Get().GetInputPorts();
 		}
-		
-		// Show output ports as inputs while they loopback to engine
-		for (FDMXOutputPortSharedPtr OutputPort : FDMXPortManager::Get().GetOutputPorts())
+		else
 		{
-			SelectedOutputPorts.Add(OutputPort.ToSharedRef());
+			SelectedOutputPorts = FDMXPortManager::Get().GetOutputPorts();
 		}
 	}
 	else
@@ -243,18 +247,38 @@ void SDMXMonitorSourceSelector::UpdateMonitoredPorts()
 		MonitoredIODirectionWrapper->SetVisibility(EVisibility::Collapsed);
 		PortSelectorWrapper->SetVisibility(EVisibility::Visible);
 
-		if (FDMXInputPortSharedPtr SharedInputPort = PortSelector->GetSelectedInputPort())
+		if (IsMonitorInputPorts())
 		{
-			if (SharedInputPort.IsValid())
+			if (FDMXInputPortSharedPtr InputPort = PortSelector->GetSelectedInputPort())
 			{
-				SelectedInputPorts.Add(SharedInputPort.ToSharedRef());
+				if (InputPort.IsValid())
+				{
+					SelectedInputPorts.Add(InputPort.ToSharedRef());
+				}
+			}
+
+			// Make a selection if nothing was already selected
+			if (SelectedInputPorts.IsEmpty() && FDMXPortManager::Get().GetInputPorts().Num() > 0)
+			{
+				const FDMXInputPortSharedRef& NewSelection = FDMXPortManager::Get().GetInputPorts()[0];
+				PortSelector->SelectPort(NewSelection->GetPortGuid());
 			}
 		}
-		else if (FDMXOutputPortSharedPtr SharedOutputPort = PortSelector->GetSelectedOutputPort())
+		else
 		{
-			if (SharedOutputPort.IsValid())
+			if (FDMXOutputPortSharedPtr OutputPort = PortSelector->GetSelectedOutputPort())
+			{
+				if (OutputPort.IsValid())
+				{
+					SelectedOutputPorts.Add(OutputPort.ToSharedRef());
+				}
+			}
+
+			// Make a selection if nothing was already selected
+			if (SelectedOutputPorts.IsEmpty() && FDMXPortManager::Get().GetInputPorts().Num() > 0)
 			{
-				SelectedOutputPorts.Add(SharedOutputPort.ToSharedRef());
+				const FDMXOutputPortSharedRef& NewSelection = FDMXPortManager::Get().GetOutputPorts()[0];
+				PortSelector->SelectPort(NewSelection->GetPortGuid());
 			}
 		}
 	}
diff --git a/Engine/Plugins/VirtualProduction/DMX/DMXEngine/Source/DMXEditor/Private/Widgets/OutputConsole/SDMXFader.cpp b/Engine/Plugins/VirtualProduction/DMX/DMXEngine/Source/DMXEditor/Private/Widgets/OutputConsole/SDMXFader.cpp
--- a/Engine/Plugins/VirtualProduction/DMX/DMXEngine/Source/DMXEditor/Private/Widgets/OutputConsole/SDMXFader.cpp
+++ b/Engine/Plugins/VirtualProduction/DMX/DMXEngine/Source/DMXEditor/Private/Widgets/OutputConsole/SDMXFader.cpp
@@ -266,7 +266,7 @@ FReply SDMXFader::OnMouseButtonDown(const FGeometry& MyGeometry, const FPointerE
 
 TSharedRef<SWidget> SDMXFader::GenerateAdressEditWidget()
 {
-	FMargin LabelPadding = FMargin(4.0f, 0.0f, 20.0f, 1.0f);
+	FMargin LabelPadding = FMargin(4.0f, 0.0f, 8.0f, 1.0f);
 	FMargin ValuePadding = FMargin(0.0f, 0.0f, 4.0f, 1.0f);
 
 	return 
diff --git a/Engine/Plugins/VirtualProduction/DMX/DMXEngine/Source/DMXEditor/Private/Widgets/OutputConsole/SDMXOutputConsolePortSelector.cpp b/Engine/Plugins/VirtualProduction/DMX/DMXEngine/Source/DMXEditor/Private/Widgets/OutputConsole/SDMXOutputConsolePortSelector.cpp
--- a/Engine/Plugins/VirtualProduction/DMX/DMXEngine/Source/DMXEditor/Private/Widgets/OutputConsole/SDMXOutputConsolePortSelector.cpp
+++ b/Engine/Plugins/VirtualProduction/DMX/DMXEngine/Source/DMXEditor/Private/Widgets/OutputConsole/SDMXOutputConsolePortSelector.cpp
@@ -130,6 +130,11 @@ void SDMXOutputConsolePortSelector::UpdateSelectedPorts()
 				SelectedOutputPorts.Add(SharedOutputPort.ToSharedRef());
 			}
 		}
+
+		if (SelectedOutputPorts.Num() == 0 && FDMXPortManager::Get().GetOutputPorts().Num() > 0)
+		{
+			PortSelector->SelectPort(FDMXPortManager::Get().GetOutputPorts()[0]->GetPortGuid());
+		}
 	}
 }
 
diff --git a/Engine/Plugins/VirtualProduction/DMX/DMXProtocol/Source/DMXProtocol/Private/IO/DMXOutputPort.cpp b/Engine/Plugins/VirtualProduction/DMX/DMXProtocol/Source/DMXProtocol/Private/IO/DMXOutputPort.cpp
--- a/Engine/Plugins/VirtualProduction/DMX/DMXProtocol/Source/DMXProtocol/Private/IO/DMXOutputPort.cpp
+++ b/Engine/Plugins/VirtualProduction/DMX/DMXProtocol/Source/DMXProtocol/Private/IO/DMXOutputPort.cpp
@@ -617,21 +617,18 @@ void FDMXOutputPort::SendDMX(int32 LocalUniverseID, const TMap<int32, uint8>& Ch
 			const TSharedPtr<FDMXSignalFragment> Fragment = MakeShared<FDMXSignalFragment>(ExternUniverseID, ChannelToValueMap, SendTime);
 			SignalFragments.Enqueue(Fragment);
 
-			if (bNeedsLoopbackToEngine)
+			// Write the fragment to the game thread's buffer
+			const FDMXSignalSharedPtr& Signal = ExternUniverseToLatestSignalMap_GameThread.FindOrAdd(ExternUniverseID, MakeShared<FDMXSignal, ESPMode::ThreadSafe>());
+
+			for (const TTuple<int32, uint8>& ChannelValueKvp : ChannelToValueMap)
 			{
-				// Write the fragment to the game thread's buffer
-				const FDMXSignalSharedPtr& Signal = ExternUniverseToLatestSignalMap_GameThread.FindOrAdd(ExternUniverseID, MakeShared<FDMXSignal, ESPMode::ThreadSafe>());
+				int32 ChannelIndex = ChannelValueKvp.Key - 1;
 
-				for (const TTuple<int32, uint8>& ChannelValueKvp : ChannelToValueMap)
+				// Filter invalid indicies so we can send bp calls here without testing them first.
+				if (Signal->ChannelData.IsValidIndex(ChannelIndex))
 				{
-					int32 ChannelIndex = ChannelValueKvp.Key - 1;
-
-					// Filter invalid indicies so we can send bp calls here without testing them first.
-					if (Signal->ChannelData.IsValidIndex(ChannelIndex))
-					{
-						Signal->Timestamp = SendTime;
-						Signal->ChannelData[ChannelIndex] = ChannelValueKvp.Value;
-					}
+					Signal->Timestamp = SendTime;
+					Signal->ChannelData[ChannelIndex] = ChannelValueKvp.Value;
 				}
 			}
 		}
