DMX - Fix a bug where the fixture patch editor didn't reflect fixture type name changes

diff --git a/Engine/Plugins/VirtualProduction/DMX/DMXEngine/Source/DMXEditor/Private/Widgets/FixtureType/SDMXFixtureTypeTreeFixtureTypeRow.cpp b/Engine/Plugins/VirtualProduction/DMX/DMXEngine/Source/DMXEditor/Private/Widgets/FixtureType/SDMXFixtureTypeTreeFixtureTypeRow.cpp
--- a/Engine/Plugins/VirtualProduction/DMX/DMXEngine/Source/DMXEditor/Private/Widgets/FixtureType/SDMXFixtureTypeTreeFixtureTypeRow.cpp
+++ b/Engine/Plugins/VirtualProduction/DMX/DMXEngine/Source/DMXEditor/Private/Widgets/FixtureType/SDMXFixtureTypeTreeFixtureTypeRow.cpp
@@ -263,10 +263,12 @@ void SDMXFixtureTypeTreeFixtureTypeRow::OnNameTextCommit(const FText& InNewName,
 		}
 
 		const FScopedTransaction Transaction(LOCTEXT("RenameEntity", "Rename Entity"));
-		DMXLibrary->Modify();
+		Entity->PreEditChange(UDMXEntityFixtureType::StaticClass()->FindPropertyByName(GET_MEMBER_NAME_CHECKED(UDMXEntityFixtureType, Name)));
 
 		FDMXEditorUtils::RenameEntity(DMXLibrary, Entity, NewNameString);
 
+		Entity->PostEditChange();
+
 		FixtureTypeTree->SelectItemByName(NewNameString, ESelectInfo::OnMouseClick);
 	}
 }
