DMX - Fix a bad range check in DMXFixtureTypeCache

diff --git a/Engine/Plugins/VirtualProduction/DMX/DMXEngine/Source/DMXRuntime/Private/Library/DMXEntityFixturePatchCache.cpp b/Engine/Plugins/VirtualProduction/DMX/DMXEngine/Source/DMXRuntime/Private/Library/DMXEntityFixturePatchCache.cpp
--- a/Engine/Plugins/VirtualProduction/DMX/DMXEngine/Source/DMXRuntime/Private/Library/DMXEntityFixturePatchCache.cpp
+++ b/Engine/Plugins/VirtualProduction/DMX/DMXEngine/Source/DMXRuntime/Private/Library/DMXEntityFixturePatchCache.cpp
@@ -341,7 +341,7 @@ const float* FDMXEntityFixturePatchCache::GetNormalizedMatrixAttributeValueFromC
 
 const FDMXNormalizedAttributeValueMap* FDMXEntityFixturePatchCache::GetAllNormalizedMatrixAttributeValuesFromCell(int32 CellIndex) const
 {
-	if (bValid && CachedRawValuesPerMatrixCell.IsValidIndex(CellIndex))
+	if (bValid && CachedNormalizedValuesPerMatrixCell.IsValidIndex(CellIndex))
 	{
 		return &CachedNormalizedValuesPerMatrixCell[CellIndex];
 	}
