DMX - Add UE5 icons

diff --git a/Engine/Plugins/VirtualProduction/DMX/DMXEngine/Source/DMXEditor/Private/Customizations/DMXLibraryDetails.cpp b/Engine/Plugins/VirtualProduction/DMX/DMXEngine/Source/DMXEditor/Private/Customizations/DMXLibraryDetails.cpp
new file mode 100644
--- /dev/null
+++ b/Engine/Plugins/VirtualProduction/DMX/DMXEngine/Source/DMXEditor/Private/Customizations/DMXLibraryDetails.cpp
@@ -0,0 +1,72 @@
+// Copyright Epic Games, Inc. All Rights Reserved.
+
+#include "Customizations/DMXLibraryDetails.h"
+
+#include "DMXEditorStyle.h"
+#include "Library/DMXLibrary.h"
+
+#include "DetailCategoryBuilder.h"
+#include "DetailLayoutBuilder.h"
+#include "DetailWidgetRow.h"
+#include "EditorFontGlyphs.h"
+#include "ISettingsModule.h"
+#include "Modules/ModuleManager.h"
+#include "Widgets/Images/SImage.h"
+#include "Widgets/Input/SButton.h"
+#include "Widgets/Text/STextBlock.h"
+
+
+#define LOCTEXT_NAMESPACE "DMXLibraryDetails"
+
+TSharedRef<IDetailCustomization> FDMXLibraryDetails::MakeInstance()
+{
+	return MakeShared<FDMXLibraryDetails>();
+}
+
+void FDMXLibraryDetails::CustomizeDetails(IDetailLayoutBuilder& DetailBuilder)
+{
+	DetailBuilder.HideCategory("DMX");
+
+	IDetailCategoryBuilder& DetailCategoryBuilder = DetailBuilder.EditCategory("DMX Ports");
+
+	const TSharedPtr<IPropertyHandle> PortReferencesHandle = DetailBuilder.GetProperty(UDMXLibrary::GetPortReferencesPropertyName());
+	
+	DetailCategoryBuilder.AddCustomRow(LOCTEXT("OpenSettingsFilterString", "Settings"))
+		.WholeRowContent()
+		.HAlign(HAlign_Left)
+		.VAlign(VAlign_Top)
+		[
+			SNew(SButton)
+			.ButtonStyle(FEditorStyle::Get(), "FlatButton.Success")
+			.ForegroundColor(FLinearColor::White)
+			.ToolTipText(LOCTEXT("OpenSettingsTooltip", "Opens the DMX Project Settings"))
+			.ContentPadding(FMargin(4.0f, 1.0f))
+			.OnClicked_Lambda([]() 
+				{ 
+					FModuleManager::GetModulePtr<ISettingsModule>("Settings")->ShowViewer("Project", "Plugins", "DMX"); return FReply::Handled(); 
+				})
+			.Content()
+			[
+				SNew(SHorizontalBox)
+				+SHorizontalBox::Slot()
+				.AutoWidth()
+				.Padding(FMargin(0.f, 0.f, 12.f, 0.f))
+				[
+					SNew(SImage)
+					.Image(FDMXEditorStyle::Get().GetBrush("Icons.DMXLibrarySettings"))
+				]
+				+ SHorizontalBox::Slot()
+				.VAlign(VAlign_Center)
+				.AutoWidth()
+				.Padding(FMargin(2.f, 0.f, 2.f, 0.f))
+				[
+					SNew(STextBlock)
+					.Text(LOCTEXT("OpenSettingsLabel", "Open DMX Project Settings"))
+				]
+			]
+		];
+
+	DetailCategoryBuilder.AddProperty(PortReferencesHandle.ToSharedRef());
+}
+ 
+#undef LOCTEXT_NAMESPACE
diff --git a/Engine/Plugins/VirtualProduction/DMX/DMXEngine/Source/DMXEditor/Private/Customizations/DMXLibraryDetails.h b/Engine/Plugins/VirtualProduction/DMX/DMXEngine/Source/DMXEditor/Private/Customizations/DMXLibraryDetails.h
new file mode 100644
--- /dev/null
+++ b/Engine/Plugins/VirtualProduction/DMX/DMXEngine/Source/DMXEditor/Private/Customizations/DMXLibraryDetails.h
@@ -0,0 +1,24 @@
+// Copyright Epic Games, Inc. All Rights Reserved.
+
+#pragma once
+
+#include "CoreMinimal.h"
+#include "IDetailCustomization.h"
+
+class IPropertyHandle;
+class IPropertyUtilities;
+
+
+/** Details customization for the 'FixtureType FunctionProperties' details view */
+class FDMXLibraryDetails
+	: public IDetailCustomization
+{
+public:
+	/** Creates an instance of this details customization */
+	static TSharedRef<IDetailCustomization> MakeInstance();
+
+protected:
+	//~ Begin IDetailCustomization interface
+	virtual void CustomizeDetails(IDetailLayoutBuilder& DetailBuilder) override;
+	//~ End IDetailCustomization interface
+};
diff --git a/Engine/Plugins/VirtualProduction/DMX/DMXEngine/Source/DMXEditor/Private/Customizations/DMXLibraryPortDetails.cpp b/Engine/Plugins/VirtualProduction/DMX/DMXEngine/Source/DMXEditor/Private/Customizations/DMXLibraryPortDetails.cpp
deleted file mode 100644
diff --git a/Engine/Plugins/VirtualProduction/DMX/DMXEngine/Source/DMXEditor/Private/Customizations/DMXLibraryPortDetails.h b/Engine/Plugins/VirtualProduction/DMX/DMXEngine/Source/DMXEditor/Private/Customizations/DMXLibraryPortDetails.h
deleted file mode 100644
diff --git a/Engine/Plugins/VirtualProduction/DMX/DMXEngine/Source/DMXEditor/Private/DMXEditor.cpp b/Engine/Plugins/VirtualProduction/DMX/DMXEngine/Source/DMXEditor/Private/DMXEditor.cpp
--- a/Engine/Plugins/VirtualProduction/DMX/DMXEngine/Source/DMXEditor/Private/DMXEditor.cpp
+++ b/Engine/Plugins/VirtualProduction/DMX/DMXEngine/Source/DMXEditor/Private/DMXEditor.cpp
@@ -13,7 +13,6 @@
 #include "Library/DMXEntityFixtureType.h"
 #include "Library/DMXEntityFixturePatch.h"
 #include "Library/DMXEntityReference.h"
-#include "LibraryEditorTab/SDMXLibraryEditorTab.h"
 #include "Modes/DMXEditorApplicationMode.h"
 #include "Toolbars/DMXEditorToolbar.h"
 #include "Commands/DMXEditorCommands.h"
@@ -24,6 +23,7 @@
 #include "Widgets/SDMXEntityEditor.h"
 #include "Widgets/FixturePatch/SDMXFixturePatchEditor.h"
 #include "Widgets/FixtureType/SDMXFixtureTypeEditor.h"
+#include "Widgets/LibrarySettings/SDMXLibraryEditorTab.h"
 #include "Widgets/OutputConsole/SDMXOutputConsole.h"
 
 #include "Modules/ModuleManager.h"
diff --git a/Engine/Plugins/VirtualProduction/DMX/DMXEngine/Source/DMXEditor/Private/DMXEditorModule.cpp b/Engine/Plugins/VirtualProduction/DMX/DMXEngine/Source/DMXEditor/Private/DMXEditorModule.cpp
--- a/Engine/Plugins/VirtualProduction/DMX/DMXEngine/Source/DMXEditor/Private/DMXEditorModule.cpp
+++ b/Engine/Plugins/VirtualProduction/DMX/DMXEngine/Source/DMXEditor/Private/DMXEditorModule.cpp
@@ -55,8 +55,6 @@ EAssetTypeCategories::Type FDMXEditorModule::DMXEditorAssetCategory;
 
 void FDMXEditorModule::StartupModule()
 {
-	FDMXEditorStyle::Initialize();
-
 	FDMXEditorCommands::Register();
 	BindDMXEditorCommands();
 
@@ -77,8 +75,6 @@ void FDMXEditorModule::StartupModule()
 
 void FDMXEditorModule::ShutdownModule()
 {
-	FDMXEditorStyle::Shutdown();
-
 	if(UObjectInitialized())
 	{
 		UToolMenus::Get()->RemoveSection("LevelEditor.LevelEditorToolBar.User", "DMX");
@@ -148,7 +144,7 @@ void FDMXEditorModule::BindDMXEditorCommands()
 
 void FDMXEditorModule::ExtendLevelEditorToolbar()
 {
-	UToolMenu* Menu = UToolMenus::Get()->ExtendMenu("LevelEditor.LevelEditorToolBar.User");
+	UToolMenu* const Menu = UToolMenus::Get()->ExtendMenu("LevelEditor.LevelEditorToolBar.User");
 
 	FToolMenuSection& Section = Menu->FindOrAddSection("DMX");
 
@@ -165,18 +161,24 @@ void FDMXEditorModule::ExtendLevelEditorToolbar()
 
 				if (bSendEnabled && bReceiveEnabled)
 				{
-					return FSlateIcon(FDMXEditorStyle::GetStyleSetName(), "DMXEditor.LevelEditor.MenuIcon_snd-rcv");
+					static const FSlateIcon IconSendReceiveDMXEnabled = FSlateIcon(FDMXEditorStyle::Get().GetStyleSetName(), "Icons.SendReceiveDMXEnabled");
+					return IconSendReceiveDMXEnabled;
 				}
 				else if (bSendEnabled)
 				{
-					return FSlateIcon(FDMXEditorStyle::GetStyleSetName(), "DMXEditor.LevelEditor.MenuIcon_snd");
+					static const FSlateIcon IconSendDMXEnabled = FSlateIcon(FDMXEditorStyle::Get().GetStyleSetName(), "Icons.SendDMXEnabled");
+					return IconSendDMXEnabled;
 				}
 				else if (bReceiveEnabled)
 				{
-					return FSlateIcon(FDMXEditorStyle::GetStyleSetName(), "DMXEditor.LevelEditor.MenuIcon_rcv");
+					static const FSlateIcon IconReceiveDMXEnabled = FSlateIcon(FDMXEditorStyle::Get().GetStyleSetName(), "Icons.ReceiveDMXEnabled");
+					return IconReceiveDMXEnabled;
+				}
+				else
+				{
+					static const FSlateIcon IconSendReceiveDMXDisabled = FSlateIcon(FDMXEditorStyle::Get().GetStyleSetName(), "Icons.SendReceiveDMXDisabled");
+					return IconSendReceiveDMXDisabled;
 				}
-
-				return FSlateIcon(FDMXEditorStyle::GetStyleSetName(), "DMXEditor.LevelEditor.MenuIcon_none");
 			})
 		);
 
@@ -191,24 +193,51 @@ TSharedRef<SWidget> FDMXEditorModule::GenerateDMXLevelEditorToolbarMenu()
 
 	FMenuBuilder MenuBuilder(true, CommandBindings);
 
-	MenuBuilder.BeginSection("CustomMenu", TAttribute<FText>(FText::FromString("DMX")));
-
-	MenuBuilder.AddMenuEntry(FDMXEditorCommands::Get().OpenChannelsMonitor);
-	MenuBuilder.AddMenuEntry(FDMXEditorCommands::Get().OpenActivityMonitor);
-	MenuBuilder.AddMenuEntry(FDMXEditorCommands::Get().OpenOutputConsole);
-	MenuBuilder.AddMenuEntry(FDMXEditorCommands::Get().OpenPatchTool);
-	MenuBuilder.AddMenuEntry(FDMXEditorCommands::Get().ToggleReceiveDMX);
-	MenuBuilder.AddMenuEntry(FDMXEditorCommands::Get().ToggleSendDMX);
-
-	MenuBuilder.EndSection();
-
+	static const FName NoExtensionHook = NAME_None;
+	MenuBuilder.AddMenuEntry(FDMXEditorCommands::Get().OpenChannelsMonitor,
+		NoExtensionHook,
+		LOCTEXT("ChannelsMonitorLabel", "Open Channel Monitor"),
+		LOCTEXT("ChannelsMonitorTooltip", "Opens the Monitor for all DMX Channels in a Universe"),
+		FSlateIcon(FDMXEditorStyle::Get().GetStyleSetName(), "Icons.ChannelsMonitor")
+	);
+	MenuBuilder.AddMenuEntry(FDMXEditorCommands::Get().OpenActivityMonitor,
+		NoExtensionHook,
+		LOCTEXT("ActivityMonitorLabel", "Open Activity Monitor"),
+		LOCTEXT("ActivityMonitorTooltip", "Open the Monitor for all DMX activity in a range of Universes"),
+		FSlateIcon(FDMXEditorStyle::Get().GetStyleSetName(), "Icons.ActivityMonitor")
+	);
+	MenuBuilder.AddMenuEntry(FDMXEditorCommands::Get().OpenOutputConsole,
+		NoExtensionHook,
+		LOCTEXT("OutputConsoleLabel", "Open Output Console"),
+		LOCTEXT("OutputConsoleTooltip", "Opens a Console to generate and output DMX Signals"),
+		FSlateIcon(FDMXEditorStyle::Get().GetStyleSetName(), "Icons.OutputConsole")
+	);
+	MenuBuilder.AddMenuEntry(FDMXEditorCommands::Get().OpenPatchTool,
+		NoExtensionHook,
+		LOCTEXT("PatchToolLabel", "Open Patch Tool"),
+		LOCTEXT("PatchToolTooltip", "Open the patch tool - Useful to patch many fixtures at once."),
+		FSlateIcon(FDMXEditorStyle::Get().GetStyleSetName(), "Icons.PatchTool")
+	);
+	MenuBuilder.AddMenuEntry(FDMXEditorCommands::Get().ToggleReceiveDMX,
+		NoExtensionHook,
+		LOCTEXT("ReceiveDMXLabel", "Receive DMX"),
+		LOCTEXT("ReceiveDMXTooltip", "Sets whether DMX is received in from the network"),
+		FSlateIcon(FDMXEditorStyle::Get().GetStyleSetName(), "Icons.ReceiveDMX")
+	);
+	MenuBuilder.AddMenuEntry(FDMXEditorCommands::Get().ToggleSendDMX,
+		NoExtensionHook,
+		LOCTEXT("SendDMXLabel", "Send DMX"),
+		LOCTEXT("SendDMXTooltip", "Sets whether DMX is sent to the network"),
+		FSlateIcon(FDMXEditorStyle::Get().GetStyleSetName(), "Icons.SendDMX")
+	);
+	
 	return MenuBuilder.MakeWidget();
 }
 
 void FDMXEditorModule::RegisterAssetTypeCategories()
 {
 	IAssetTools& AssetTools = FModuleManager::LoadModuleChecked<FAssetToolsModule>("AssetTools").Get();
-	DMXEditorAssetCategory = AssetTools.RegisterAdvancedAssetCategory(FName(TEXT("DMX")), LOCTEXT("DmxCategory", "DMX"));
+	DMXEditorAssetCategory = AssetTools.RegisterAdvancedAssetCategory(FName(TEXT("DMX")), LOCTEXT("DmxCategory", "DMX")); 
 }
 
 void FDMXEditorModule::RegisterAssetTypeActions()
@@ -282,22 +311,26 @@ void FDMXEditorModule::RegisterNomadTabSpawners()
 	RegisterNomadTabSpawner(FDMXEditorTabNames::ChannelsMonitor,
 		FOnSpawnTab::CreateStatic(&FDMXEditorModule::OnSpawnChannelsMonitorTab))
 		.SetDisplayName(LOCTEXT("ChannelsMonitorTabTitle", "DMX Channel Monitor"))
-		.SetMenuType(ETabSpawnerMenuType::Hidden);
+		.SetMenuType(ETabSpawnerMenuType::Hidden)
+		.SetIcon(FSlateIcon(FDMXEditorStyle::Get().GetStyleSetName(), "Icons.ChannelsMonitor"));
 
 	RegisterNomadTabSpawner(FDMXEditorTabNames::ActivityMonitor,
 		FOnSpawnTab::CreateStatic(&FDMXEditorModule::OnSpawnActivityMonitorTab))
 		.SetDisplayName(LOCTEXT("ActivityMonitorTabTitle", "DMX Activity Monitor"))
-		.SetMenuType(ETabSpawnerMenuType::Hidden);
+		.SetMenuType(ETabSpawnerMenuType::Hidden)
+		.SetIcon(FSlateIcon(FDMXEditorStyle::Get().GetStyleSetName(), "Icons.ActivityMonitor"));
 
 	RegisterNomadTabSpawner(FDMXEditorTabNames::OutputConsole,
 		FOnSpawnTab::CreateStatic(&FDMXEditorModule::OnSpawnOutputConsoleTab))
 		.SetDisplayName(LOCTEXT("OutputConsoleTabTitle", "DMX Output Console"))
-		.SetMenuType(ETabSpawnerMenuType::Hidden);
+		.SetMenuType(ETabSpawnerMenuType::Hidden)
+		.SetIcon(FSlateIcon(FDMXEditorStyle::Get().GetStyleSetName(), "Icons.OutputConsole"));
 
 	RegisterNomadTabSpawner(FDMXEditorTabNames::PatchTool,
 		FOnSpawnTab::CreateStatic(&FDMXEditorModule::OnSpawnPatchToolTab))
 		.SetDisplayName(LOCTEXT("PatchToolTabTitle", "DMX Patch Tool"))
-		.SetMenuType(ETabSpawnerMenuType::Hidden);
+		.SetMenuType(ETabSpawnerMenuType::Hidden)
+		.SetIcon(FSlateIcon(FDMXEditorStyle::Get().GetStyleSetName(), "Icons.PatchTool"));
 }
 
 void FDMXEditorModule::StartupPIEManager()
@@ -384,30 +417,6 @@ bool FDMXEditorModule::IsSendDMXEnabled()
 	return UDMXProtocolBlueprintLibrary::IsSendDMXEnabled();
 }
 
-FText FDMXEditorModule::GetToggleSendDMXText() const
-{
-	if (UDMXProtocolBlueprintLibrary::IsSendDMXEnabled())
-	{
-		return LOCTEXT("MenuButtonText_PauseSendDMX", "Pause Send DMX");
-	}
-	else
-	{
-		return LOCTEXT("MenuButtonText_ResumeSendDMX", "Resume Send DMX");
-	}
-}
-
-FText FDMXEditorModule::GetToggleSendDMXTooltip() const
-{
-	if (UDMXProtocolBlueprintLibrary::IsSendDMXEnabled())
-	{
-		return LOCTEXT("MenuButtonText_DisableOutboundDMXPackets", "Disables outbound DMX packets in editor.");
-	}
-	else
-	{
-		return LOCTEXT("MenuButtonText_EnableOutboundDMXPackets", "Enables outbound DMX packets in editor.");
-	}
-}
-
 void FDMXEditorModule::OnToggleReceiveDMX()
 {
 	bool bAffectEditor = true;
@@ -427,30 +436,6 @@ bool FDMXEditorModule::IsReceiveDMXEnabled()
 	return UDMXProtocolBlueprintLibrary::IsReceiveDMXEnabled();
 }
 
-FText FDMXEditorModule::GetToggleReceiveDMXText() const
-{
-	if (UDMXProtocolBlueprintLibrary::IsReceiveDMXEnabled())
-	{
-		return LOCTEXT("MenuButtonText_PauseReceiveDMX", "Pause Receive DMX");
-	}
-	else
-	{
-		return LOCTEXT("MenuButtonText_ResumeReceiveDMX", "Resume Receive DMX");
-	}
-}
-
-FText FDMXEditorModule::GetToggleReceiveDMXTooltip() const
-{
-	if (UDMXProtocolBlueprintLibrary::IsReceiveDMXEnabled())
-	{
-		return LOCTEXT("MenuButtonText_DisableInboundDMXPackets", "Disables inbound DMX packets in editor.");
-	}
-	else
-	{
-		return LOCTEXT("MenuButtonText_EnableInboundDMXPackets", "Enables inbound DMX packets editor.");
-	}
-}
-
 void FDMXEditorModule::RegisterAssetTypeAction(TSharedRef<IAssetTypeActions> Action)
 {
 	IAssetTools& AssetTools = FModuleManager::GetModuleChecked<FAssetToolsModule>("AssetTools").Get();
diff --git a/Engine/Plugins/VirtualProduction/DMX/DMXEngine/Source/DMXEditor/Private/DMXEditorStyle.cpp b/Engine/Plugins/VirtualProduction/DMX/DMXEngine/Source/DMXEditor/Private/DMXEditorStyle.cpp
new file mode 100644
--- /dev/null
+++ b/Engine/Plugins/VirtualProduction/DMX/DMXEngine/Source/DMXEditor/Private/DMXEditorStyle.cpp
@@ -0,0 +1,182 @@
+// Copyright Epic Games, Inc. All Rights Reserved.
+
+#include "DMXEditorStyle.h"
+
+#include "EditorStyleSet.h"
+#include "Engine/Font.h"
+#include "Framework/Application/SlateApplication.h"
+#include "Interfaces/IPluginManager.h"
+#include "Styling/SlateStyleRegistry.h"
+#include "Styling/SlateStyleMacros.h"
+#include "Styling/SlateTypes.h"
+#include "Styling/StyleColors.h"
+
+
+FDMXEditorStyle::FDMXEditorStyle()
+	: FSlateStyleSet("DMXEditorStyle")
+{
+	static const FVector2D Icon8x8(8.f, 8.f);
+	static const FVector2D Icon16x16(16.f, 16.f);
+	static const FVector2D Icon64x64(64.f, 64.f);
+
+	static const FVector2D Icon20x20(20.f, 20.f);
+	static const FVector2D Icon40x40(40.f, 40.f);
+	static const FVector2D Icon50x40(50.f, 40.f);
+	static const FVector2D Icon34x29(34.f, 29.f);
+	static const FVector2D Icon51x30(51.f, 30.f);
+	static const FVector2D Icon51x31(51.f, 31.f);
+
+	const TSharedPtr<FSlateStyleSet> Style = MakeShared<FSlateStyleSet>(GetStyleSetName());
+
+	static const TCHAR* DMXEnginePluginName = TEXT("DMXEngine");
+	TSharedPtr<IPlugin> Plugin = IPluginManager::Get().FindPlugin(DMXEnginePluginName);
+	if (ensureMsgf(Plugin.IsValid(), TEXT("Cannot find Plugin 'DMXEngine' hence cannot set the DMXEditorStyle root directory.")))
+	{
+		SetContentRoot(Plugin->GetBaseDir() / TEXT("Content/Slate"));
+	}
+	SetCoreContentRoot(FPaths::EngineContentDir() / TEXT("Slate"));
+
+	// Default color brushes
+	{
+		Set("DMXEditor.WhiteBrush", new FSlateColorBrush(FLinearColor(1.f, 1.f, 1.f, 1.f)));
+		Set("DMXEditor.BlackBrush", new FSlateColorBrush(FLinearColor(0.f, 0.f, 0.f, 1.f)));
+	}
+
+	// Default Fonts
+	{
+		static const TCHAR* FontPathRoboto = TEXT("Font'/Engine/EngineFonts/Roboto.Roboto'");
+		static const UFont* FontRoboto = Cast<UFont>(StaticLoadObject(UFont::StaticClass(), nullptr, FontPathRoboto));
+		check(FontRoboto != nullptr);
+
+		Set("DMXEditor.Font.InputChannelID", FSlateFontInfo(FontRoboto, 8, FName(TEXT("Light"))));
+		Set("DMXEditor.Font.InputChannelValue", FSlateFontInfo(FontRoboto, 10, FName(TEXT("Regular"))));
+
+		Set("DMXEditor.Font.InputUniverseHeader", FSlateFontInfo(FontRoboto, 10, FName(TEXT("Bold"))));
+		Set("DMXEditor.Font.InputUniverseID", FSlateFontInfo(FontRoboto, 10, FName(TEXT("Regular"))));
+		Set("DMXEditor.Font.InputUniverseChannelID", FSlateFontInfo(FontRoboto, 10, FName(TEXT("Regular"))));
+		Set("DMXEditor.Font.InputUniverseChannelValue", FSlateFontInfo(FontRoboto, 10, FName(TEXT("Light"))));
+	}
+
+	// Asset icons
+	{
+		Set("ClassIcon.DMXLibrary", new IMAGE_BRUSH_SVG("DMXLibrary_16", Icon16x16));
+		Set("ClassThumbnail.DMXLibrary", new IMAGE_BRUSH_SVG("DMXLibrary_64", Icon64x64));
+	}
+
+	// Icons
+	{ 
+		Set("Icons.DMXLibrary", new IMAGE_BRUSH_SVG("DMXLibrary_20", Icon20x20));
+
+		{
+			Set("Icons.DMXLibrarySettings", new CORE_IMAGE_BRUSH_SVG("Starship/Common/settings", Icon16x16));
+		}
+		Set("Icons.FixtureType", new IMAGE_BRUSH_SVG("FixtureType_16", Icon16x16));
+		Set("Icons.FixturePatch", new IMAGE_BRUSH_SVG("FixturePatch_16", Icon16x16));
+
+		Set("Icons.SendReceiveDMXEnabled", new IMAGE_BRUSH_SVG("SendReceiveDMXEnabled_20", Icon20x20));
+		Set("Icons.ReceiveDMXEnabled", new IMAGE_BRUSH_SVG("ReceiveDMXEnabled_20", Icon20x20));
+		Set("Icons.SendDMXEnabled", new IMAGE_BRUSH_SVG("SendDMXEnabled_20", Icon20x20));
+		Set("Icons.SendReceiveDMXDisabled", new IMAGE_BRUSH_SVG("SendReceiveDMXDisabled_20", Icon20x20));
+	
+		Set("Icons.ChannelsMonitor", new IMAGE_BRUSH_SVG("ChannelsMonitor_16", Icon16x16));
+		Set("Icons.ActivityMonitor", new IMAGE_BRUSH_SVG("ActivityMonitor_16", Icon16x16));
+		Set("Icons.OutputConsole", new IMAGE_BRUSH_SVG("OutputConsole_16", Icon16x16));
+		Set("Icons.PatchTool", new IMAGE_BRUSH_SVG("PatchTool_16", Icon16x16));
+		Set("Icons.ReceiveDMX", new IMAGE_BRUSH_SVG("ToggleReceiveDMX_16", Icon16x16));
+		Set("Icons.SendDMX", new IMAGE_BRUSH_SVG("ToggleSendDMX_16", Icon16x16));
+	}
+	
+	// Distribution Grid buttons
+	{
+		Set("DMXEditor.PixelMapping.DistributionGrid.0.0", new IMAGE_BRUSH("icon_PixelDirection_0_0_161x160", Icon34x29));
+		Set("DMXEditor.PixelMapping.DistributionGrid.0.1", new IMAGE_BRUSH("icon_PixelDirection_0_1_161x160", Icon34x29));
+		Set("DMXEditor.PixelMapping.DistributionGrid.0.2", new IMAGE_BRUSH("icon_PixelDirection_0_2_161x160", Icon34x29));
+		Set("DMXEditor.PixelMapping.DistributionGrid.0.3", new IMAGE_BRUSH("icon_PixelDirection_0_3_161x160", Icon34x29));
+
+		Set("DMXEditor.PixelMapping.DistributionGrid.1.0", new IMAGE_BRUSH("icon_PixelDirection_1_0_161x160", Icon34x29));
+		Set("DMXEditor.PixelMapping.DistributionGrid.1.1", new IMAGE_BRUSH("icon_PixelDirection_1_1_161x160", Icon34x29));
+		Set("DMXEditor.PixelMapping.DistributionGrid.1.2", new IMAGE_BRUSH("icon_PixelDirection_1_2_161x160", Icon34x29));
+		Set("DMXEditor.PixelMapping.DistributionGrid.1.3", new IMAGE_BRUSH("icon_PixelDirection_1_3_161x160", Icon34x29));
+
+		Set("DMXEditor.PixelMapping.DistributionGrid.2.0", new IMAGE_BRUSH("icon_PixelDirection_2_0_161x160", Icon34x29));
+		Set("DMXEditor.PixelMapping.DistributionGrid.2.1", new IMAGE_BRUSH("icon_PixelDirection_2_1_161x160", Icon34x29));
+		Set("DMXEditor.PixelMapping.DistributionGrid.2.2", new IMAGE_BRUSH("icon_PixelDirection_2_2_161x160", Icon34x29));
+		Set("DMXEditor.PixelMapping.DistributionGrid.2.3", new IMAGE_BRUSH("icon_PixelDirection_2_3_161x160", Icon34x29));
+
+		Set("DMXEditor.PixelMapping.DistributionGrid.3.0", new IMAGE_BRUSH("icon_PixelDirection_3_0_161x160", Icon34x29));
+		Set("DMXEditor.PixelMapping.DistributionGrid.3.1", new IMAGE_BRUSH("icon_PixelDirection_3_1_161x160", Icon34x29));
+		Set("DMXEditor.PixelMapping.DistributionGrid.3.2", new IMAGE_BRUSH("icon_PixelDirection_3_2_161x160", Icon34x29));
+		Set("DMXEditor.PixelMapping.DistributionGrid.3.3", new IMAGE_BRUSH("icon_PixelDirection_3_3_161x160", Icon34x29));
+	}
+	
+	// Output Console
+	{
+		Set("DMXEditor.OutputConsole.Fader", new FSlateColorBrush(FLinearColor(.2f, .2f, .2f, .2f)));
+		Set("DMXEditor.OutputConsole.Fader_Highlighted", new FSlateColorBrush(FLinearColor(.4f, .4f, .4f, .4f)));
+
+		Set("DMXEditor.OutputConsole.MacroSineWave", new IMAGE_BRUSH("icon_MacroSineWave_51x31", Icon51x31));
+		Set("DMXEditor.OutputConsole.MacroMin", new IMAGE_BRUSH("icon_MacroMin_51x30", Icon51x30));
+		Set("DMXEditor.OutputConsole.MacroMax", new IMAGE_BRUSH("icon_MacroMax_51x30", Icon51x30));
+
+		static const FLinearColor DefaultFaderBackColor = FLinearColor::FromSRGBColor(FColor::FromHex("191919"));
+		static const FLinearColor DefaultFaderFillColor = FLinearColor::FromSRGBColor(FColor::FromHex("00aeef"));
+		static const FLinearColor DefeaultFaderForeColor = FLinearColor::FromSRGBColor(FColor::FromHex("ffffff"));
+
+		Set("DMXEditor.OutputConsole.Fader", FSpinBoxStyle(FEditorStyle::GetWidgetStyle<FSpinBoxStyle>("SpinBox"))
+			.SetBackgroundBrush(CORE_BOX_BRUSH("Slate/Common/Spinbox", FMargin(4.f / 16.f), DefaultFaderBackColor))
+			.SetHoveredBackgroundBrush(CORE_BOX_BRUSH("Slate/Common/Spinbox", FMargin(4.f / 16.f), DefaultFaderBackColor))
+			.SetActiveFillBrush(CORE_BOX_BRUSH("Slate/Common/Spinbox_Fill", FMargin(4.f / 16.f), DefaultFaderFillColor))
+			.SetInactiveFillBrush(CORE_BOX_BRUSH("Slate/Common/Spinbox_Fill", FMargin(4.f / 16.f), DefaultFaderFillColor))
+			.SetForegroundColor(DefeaultFaderForeColor)
+			.SetArrowsImage(FSlateNoResource()));
+	}
+
+	FSlateStyleRegistry::RegisterSlateStyle(*this);
+}
+
+FDMXEditorStyle::~FDMXEditorStyle()
+{
+	FSlateStyleRegistry::UnRegisterSlateStyle(*this);
+}
+
+const FDMXEditorStyle& FDMXEditorStyle::Get()
+{
+	static const FDMXEditorStyle Inst;
+	return Inst;
+}
+
+#undef IMAGE_BRUSH
+#undef BOX_BRUSH
+#undef CORE_BOX_BRUSH
+
+
+// DEPRECATED 5.0
+TSharedPtr<FSlateStyleSet> FDMXEditorStyle::StyleInstance_DEPRECATED = nullptr;
+
+void FDMXEditorStyle::Initialize()
+{
+	// DEPRECATED 5.0
+	if (!StyleInstance_DEPRECATED.IsValid())
+	{
+		StyleInstance_DEPRECATED = MakeShared<FDMXEditorStyle>();
+		FSlateStyleRegistry::RegisterSlateStyle(*StyleInstance_DEPRECATED);
+	}
+}
+
+void FDMXEditorStyle::Shutdown()
+{
+	// DEPRECATED 5.0
+	ensureMsgf(StyleInstance_DEPRECATED.IsValid(), TEXT("%S called, but StyleInstance wasn't initialized"));
+	FSlateStyleRegistry::UnRegisterSlateStyle(*StyleInstance_DEPRECATED);
+	ensure(StyleInstance_DEPRECATED.IsUnique());
+	StyleInstance_DEPRECATED.Reset();
+}
+
+void FDMXEditorStyle::ReloadTextures()
+{
+	// DEPRECATED 5.0
+	if (FSlateApplication::IsInitialized())
+	{
+		FSlateApplication::Get().GetRenderer()->ReloadTextureResources();
+	}
+}
diff --git a/Engine/Plugins/VirtualProduction/DMX/DMXEngine/Source/DMXEditor/Private/LibraryEditorTab/SDMXLibraryEditorTab.cpp b/Engine/Plugins/VirtualProduction/DMX/DMXEngine/Source/DMXEditor/Private/LibraryEditorTab/SDMXLibraryEditorTab.cpp
deleted file mode 100644
diff --git a/Engine/Plugins/VirtualProduction/DMX/DMXEngine/Source/DMXEditor/Private/LibraryEditorTab/SDMXLibraryEditorTab.h b/Engine/Plugins/VirtualProduction/DMX/DMXEngine/Source/DMXEditor/Private/LibraryEditorTab/SDMXLibraryEditorTab.h
deleted file mode 100644
diff --git a/Engine/Plugins/VirtualProduction/DMX/DMXEngine/Source/DMXEditor/Private/Style/DMXEditorStyle.cpp b/Engine/Plugins/VirtualProduction/DMX/DMXEngine/Source/DMXEditor/Private/Style/DMXEditorStyle.cpp
deleted file mode 100644
diff --git a/Engine/Plugins/VirtualProduction/DMX/DMXEngine/Source/DMXEditor/Private/Tabs/DMXEditorTabFactories.cpp b/Engine/Plugins/VirtualProduction/DMX/DMXEngine/Source/DMXEditor/Private/Tabs/DMXEditorTabFactories.cpp
--- a/Engine/Plugins/VirtualProduction/DMX/DMXEngine/Source/DMXEditor/Private/Tabs/DMXEditorTabFactories.cpp
+++ b/Engine/Plugins/VirtualProduction/DMX/DMXEngine/Source/DMXEditor/Private/Tabs/DMXEditorTabFactories.cpp
@@ -4,13 +4,14 @@
 
 #include "DMXEditor.h"
 #include "DMXEditorTabNames.h"
-#include "Library/DMXEntityFixtureType.h"
-#include "LibraryEditorTab/SDMXLibraryEditorTab.h"
+#include "DMXEditorStyle.h"
+#include "Widgets/LibrarySettings/SDMXLibraryEditorTab.h"
 #include "Widgets/FixtureType/SDMXFixtureTypeEditor.h"
 #include "Widgets/FixturePatch/SDMXFixturePatchEditor.h"
 
-#include "Widgets/Docking/SDockTab.h"
+#include "Textures/SlateIcon.h"
 #include "Widgets/Layout/SBox.h"
+#include "Widgets/Docking/SDockTab.h"
 
 
 #define LOCTEXT_NAMESPACE "DMXWorkflowTabFactory"
@@ -29,7 +30,7 @@ FDMXLibraryEditorTabSummoner::FDMXLibraryEditorTabSummoner(TSharedPtr<FAssetEdit
 	: FDMXEditorPropertyTabSummoner(FDMXEditorTabNames::DMXLibraryEditor, InHostingApp)
 {
 	TabLabel = LOCTEXT("DMXLibraryEditorTabName", "Library Settings");
-	TabIcon = FSlateIcon(FEditorStyle::GetStyleSetName(), "Kismet.Tabs.BlueprintDefaults");
+	TabIcon = FSlateIcon(FDMXEditorStyle::Get().GetStyleSetName(), "Icons.DMXLibrarySettings");
 
 	bIsSingleton = true;
 
@@ -58,8 +59,8 @@ FDMXEditorFixtureTypesSummoner::FDMXEditorFixtureTypesSummoner(TSharedPtr<FAsset
 	: FDMXEditorPropertyTabSummoner(FDMXEditorTabNames::DMXFixtureTypesEditor, InHostingApp)
 {
 	TabLabel = LOCTEXT("DMXFixtureTypesTabLabel", "Fixture Types");
-	TabIcon = FSlateIcon(FEditorStyle::GetStyleSetName(), "Kismet.Tabs.BlueprintDefaults");
-
+	TabIcon = FSlateIcon(FDMXEditorStyle::Get().GetStyleSetName(), "Icons.FixtureType");
+	
 	bIsSingleton = true;
 
 	ViewMenuDescription = LOCTEXT("DMXFixtureTypesView", "Fixture Types");
@@ -82,7 +83,7 @@ FDMXEditorFixturePatchSummoner::FDMXEditorFixturePatchSummoner(TSharedPtr<FAsset
 	: FDMXEditorPropertyTabSummoner(FDMXEditorTabNames::DMXFixturePatchEditor, InHostingApp)
 {
 	TabLabel = LOCTEXT("DMXFixturePatchTabLabel", "Fixture Patch");
-	TabIcon = FSlateIcon(FEditorStyle::GetStyleSetName(), "Kismet.Tabs.BlueprintDefaults");
+	TabIcon = FSlateIcon(FDMXEditorStyle::Get().GetStyleSetName(), "Icons.FixturePatch");
 
 	bIsSingleton = true;
 
diff --git a/Engine/Plugins/VirtualProduction/DMX/DMXEngine/Source/DMXEditor/Private/Widgets/LibrarySettings/SDMXLibraryEditorTab.cpp b/Engine/Plugins/VirtualProduction/DMX/DMXEngine/Source/DMXEditor/Private/Widgets/LibrarySettings/SDMXLibraryEditorTab.cpp
new file mode 100644
--- /dev/null
+++ b/Engine/Plugins/VirtualProduction/DMX/DMXEngine/Source/DMXEditor/Private/Widgets/LibrarySettings/SDMXLibraryEditorTab.cpp
@@ -0,0 +1,37 @@
+// Copyright Epic Games, Inc. All Rights Reserved.
+
+#include "SDMXLibraryEditorTab.h"
+
+#include "Customizations/DMXLibraryDetails.h"
+#include "Library/DMXLibrary.h"
+
+#include "IDetailsView.h"
+#include "PropertyEditorModule.h"
+#include "Modules/ModuleManager.h"
+
+
+void SDMXLibraryEditorTab::Construct(const FArguments& InArgs)
+{
+	DMXEditor = InArgs._DMXEditor;
+
+	// Initialize property view widget
+	FDetailsViewArgs DetailsViewArgs;
+	DetailsViewArgs.bAllowSearch = false;
+	DetailsViewArgs.bHideSelectionTip = true;
+	DetailsViewArgs.bCustomNameAreaLocation = false;
+
+	FPropertyEditorModule& PropertyEditorModule = FModuleManager::GetModuleChecked<FPropertyEditorModule>("PropertyEditor");
+	const TSharedRef<IDetailsView> DMXLibraryDetailsView = PropertyEditorModule.CreateDetailView(DetailsViewArgs);
+
+	DMXLibraryDetailsView->RegisterInstancedCustomPropertyLayout(UDMXLibrary::StaticClass(), FOnGetDetailCustomizationInstance::CreateStatic(FDMXLibraryDetails::MakeInstance));
+
+	if (IsValid(InArgs._DMXLibrary))
+	{
+		DMXLibraryDetailsView->SetObject(InArgs._DMXLibrary, true);
+
+		ChildSlot
+			[
+				DMXLibraryDetailsView
+			];
+	}
+}
diff --git a/Engine/Plugins/VirtualProduction/DMX/DMXEngine/Source/DMXEditor/Private/Widgets/LibrarySettings/SDMXLibraryEditorTab.h b/Engine/Plugins/VirtualProduction/DMX/DMXEngine/Source/DMXEditor/Private/Widgets/LibrarySettings/SDMXLibraryEditorTab.h
new file mode 100644
--- /dev/null
+++ b/Engine/Plugins/VirtualProduction/DMX/DMXEngine/Source/DMXEditor/Private/Widgets/LibrarySettings/SDMXLibraryEditorTab.h
@@ -0,0 +1,37 @@
+// Copyright Epic Games, Inc. All Rights Reserved.
+
+#pragma once
+
+#include "DMXEditor.h"
+
+#include "CoreMinimal.h"
+#include "Widgets/SCompoundWidget.h"
+
+class FDMXEditor;
+class UDMXLibrary;
+
+class IDetailsView;
+
+
+/** Editor for Fixture Patches */
+class SDMXLibraryEditorTab
+	: public SCompoundWidget
+{
+public:
+	SLATE_BEGIN_ARGS(SDMXLibraryEditorTab)
+	{}
+
+		SLATE_ARGUMENT(UDMXLibrary*, DMXLibrary)
+	
+		SLATE_ARGUMENT(TWeakPtr<FDMXEditor>, DMXEditor)
+
+	SLATE_END_ARGS()
+
+public:
+	/** Constructs this widget */
+	void Construct(const FArguments& InArgs);
+
+private:
+	/** The DMXEditor that owns this tab */
+	TWeakPtr<FDMXEditor> DMXEditor;
+};
\ No newline at end of file
diff --git a/Engine/Plugins/VirtualProduction/DMX/DMXEngine/Source/DMXEditor/Public/DMXEditorModule.h b/Engine/Plugins/VirtualProduction/DMX/DMXEngine/Source/DMXEditor/Public/DMXEditorModule.h
--- a/Engine/Plugins/VirtualProduction/DMX/DMXEngine/Source/DMXEditor/Public/DMXEditorModule.h
+++ b/Engine/Plugins/VirtualProduction/DMX/DMXEngine/Source/DMXEditor/Public/DMXEditorModule.h
@@ -133,24 +133,12 @@ private:
 	/** Returns true if send dmx is enabled */
 	static bool IsSendDMXEnabled();
 
-	/** Returns text for the toggle receive DMX button in the menu */
-	FText GetToggleSendDMXText() const;
-
-	/** Returns text for the toggle receive DMX tooltip in the menu */
-	FText GetToggleSendDMXTooltip() const;
-
 	/** Called when the Toggle Receive DMX menu command is selected */
 	static void OnToggleReceiveDMX();
 
 	/** Returns true if receive dmx is enabled */
 	static bool IsReceiveDMXEnabled();
 
-	/** Returns text for the toggle receive DMX button in the menu */
-	FText GetToggleReceiveDMXText() const;
-
-	/** Returns text for the toggle receive DMX tooltip in the menu */
-	FText GetToggleReceiveDMXTooltip() const;
-
 	/** Command list for the DMX level editor menu */
 	TSharedPtr<class FUICommandList> DMXLevelEditorMenuCommands;
 
diff --git a/Engine/Plugins/VirtualProduction/DMX/DMXEngine/Source/DMXEditor/Public/DMXEditorStyle.h b/Engine/Plugins/VirtualProduction/DMX/DMXEngine/Source/DMXEditor/Public/DMXEditorStyle.h
--- a/Engine/Plugins/VirtualProduction/DMX/DMXEngine/Source/DMXEditor/Public/DMXEditorStyle.h
+++ b/Engine/Plugins/VirtualProduction/DMX/DMXEngine/Source/DMXEditor/Public/DMXEditorStyle.h
@@ -5,31 +5,39 @@
 #include "CoreMinimal.h"
 #include "Styling/SlateStyle.h"
 
-/**  DMX editor style, for loading and creation of UI resources */
+
+/**  DMX editor style */
 class DMXEDITOR_API FDMXEditorStyle
+	: public FSlateStyleSet
 {
 public:
+	/** Constructor */
+	FDMXEditorStyle();
 
-	/** Initializes this style's singleton and register the Slate Style */
-	static void Initialize();
+	/** Desonstructor */
+	virtual ~FDMXEditorStyle();
 
-	/** Unregister the Slate Style and release the Singleton */
-	static void Shutdown();
+	/**  Returns the singleton ISlateStyle instance for DMX editor style */
+	static const FDMXEditorStyle& Get();
 
-	/**  Reloads textures used by slate renderer */
-	static void ReloadTextures();
+/// <summary>
+/// DEPRECATED Members
+/// </summary>
 
-	/**  Returns the singleton ISlateStyle instance for DMX editor style */
-	static const ISlateStyle& Get();
+public:	
+	/** DEPRECATED 5.0 */
+	UE_DEPRECATED(5.0, "Deprecated since the public API leaves it unclear if the style is initialized or shut down. The style does not need to be Initialized explicitly anymore. Note, the StyleSetName now needs be accessed via FDMXEditorStyle::Get().GetStyleSetName().")
+	static void Initialize();
 
-	/**  Returns the DMX editor style name */
-	static FName GetStyleSetName();
+	/** DEPRECATED 5.0 */
+	UE_DEPRECATED(5.0, "Deprecated since the public API leaves it unclear if the style is initialized or shut down. The style does not need to be shut down explicitly anymore. Note, the StyleSetName now needs be accessed via FDMXEditorStyle::Get().GetStyleSetName().")
+	static void Shutdown();
 
-private:
-	/** Instantiate and returns a Slate Style Set object */
-	static TSharedRef< class FSlateStyleSet > Create();
+	/** DEPRECATED 5.0 */
+	UE_DEPRECATED(5.0, "Deprecated since no clear use (call FSlateApplication::Get().GetRenderer()->ReloadTextureResources() directly where needed).")
+	static void ReloadTextures();
 
 private:
 	/** Singleton style instance */
-	static TSharedPtr< class FSlateStyleSet > StyleInstance;
-};
\ No newline at end of file
+	static TSharedPtr<FSlateStyleSet> StyleInstance_DEPRECATED;
+};
diff --git a/Engine/Plugins/VirtualProduction/DMX/DMXPixelMapping/Source/DMXPixelMappingEditor/Private/DMXPixelMappingEditorStyle.cpp b/Engine/Plugins/VirtualProduction/DMX/DMXPixelMapping/Source/DMXPixelMappingEditor/Private/DMXPixelMappingEditorStyle.cpp
--- a/Engine/Plugins/VirtualProduction/DMX/DMXPixelMapping/Source/DMXPixelMappingEditor/Private/DMXPixelMappingEditorStyle.cpp
+++ b/Engine/Plugins/VirtualProduction/DMX/DMXPixelMapping/Source/DMXPixelMappingEditor/Private/DMXPixelMappingEditorStyle.cpp
@@ -40,6 +40,7 @@ FString RelativePathToPluginPath(const FString& RelativePath, const ANSICHAR* Ex
 }
 
 #define IMAGE_BRUSH( RelativePath, ... ) FSlateImageBrush( Style->RootToContentDir( RelativePath, TEXT(".png") ), __VA_ARGS__ )
+#define IMAGE_BRUSH_SVG( RelativePath, ... ) FSlateVectorImageBrush( Style->RootToContentDir(RelativePath, TEXT(".svg")), __VA_ARGS__ )
 #define IMAGE_CORE_BRUSH( RelativePath, ... ) FSlateImageBrush( FPaths::EngineContentDir() / "Editor/Slate" / RelativePath + TEXT(".png"), __VA_ARGS__ )
 #define IMAGE_PLUGIN_BRUSH( RelativePath, ... ) FSlateImageBrush( RelativePathToPluginPath( RelativePath, ".png" ), __VA_ARGS__ )
 #define BOX_BRUSH( RelativePath, ... ) FSlateBoxBrush( Style->RootToContentDir( RelativePath, TEXT(".png") ), __VA_ARGS__ )
@@ -54,35 +55,41 @@ const FVector2D Icon20x20(20.0f, 20.0f);
 const FVector2D Icon40x40(40.0f, 40.0f);
 const FVector2D Icon64x64(64.0f, 64.0f);
 
-TSharedRef< FSlateStyleSet > FDMXPixelMappingEditorStyle::Create()
+TSharedRef<FSlateStyleSet> FDMXPixelMappingEditorStyle::Create()
 {
-	TSharedRef< FSlateStyleSet > Style = MakeShareable(new FSlateStyleSet("DMXPixelMappingEditorStyle"));
+	TSharedRef<FSlateStyleSet> Style = MakeShareable(new FSlateStyleSet("DMXPixelMappingEditorStyle"));
 
 	TSharedPtr<IPlugin> Plugin = IPluginManager::Get().FindPlugin(TEXT("DMXPixelMapping"));
 	if (Plugin.IsValid())
 	{
-		Style->SetContentRoot(FPaths::Combine(Plugin->GetBaseDir(), TEXT("Resources")));
+		Style->SetContentRoot(FPaths::Combine(Plugin->GetBaseDir(), TEXT("Content/Slate")));
 	}
 
-	// Toolbar Icons
-	Style->Set("DMXPixelMappingEditor.AddMapping", new IMAGE_BRUSH("Icons/icon_DMXPixelMappingEditor_AddMapping_40x", Icon40x40));
-	Style->Set("DMXPixelMappingEditor.AddMapping.Small", new IMAGE_BRUSH("Icons/icon_DMXPixelMappingEditor_AddMapping_40x", Icon20x20));
+	// Class styles
+	Style->Set("ClassIcon.DMXPixelMapping", new IMAGE_BRUSH_SVG("DMXPixelMapping_16", Icon16x16));
+	Style->Set("ClassThumbnail.DMXPixelMapping", new IMAGE_BRUSH_SVG("DMXPixelMapping_64", Icon64x64));
 
-	Style->Set("ClassThumbnail.DMXPixelMapping", new IMAGE_BRUSH("Icons/icon_DMXPixelMappingEditor_Thumbnail_40x", Icon40x40));
-	Style->Set("ClassIcon.DMXPixelMapping", new IMAGE_BRUSH("Icons/icon_DMXPixelMappingEditor_Thumbnail_40x", Icon16x16));
+	// Icons
+	Style->Set("Icons.Preview", new IMAGE_BRUSH_SVG("Preview", Icon16x16));
+	Style->Set("Icons.ZoomToFit", new IMAGE_BRUSH_SVG("ZoomToFit", Icon16x16));
 
-	Style->Set("DMXPixelMappingEditor.PlayDMX", new IMAGE_BRUSH("Icons/icon_DMXPixelMappingEditor_PlayDMX_40x", Icon40x40));
-	Style->Set("DMXPixelMappingEditor.PlayDMX.Small", new IMAGE_BRUSH("Icons/icon_DMXPixelMappingEditor_PlayDMX_40x", Icon20x20));
-	Style->Set("DMXPixelMappingEditor.StopPlayingDMX", new IMAGE_BRUSH("Icons/icon_DMXPixelMappingEditor_StopPlayingDMX_40x", Icon40x40));
-	Style->Set("DMXPixelMappingEditor.StopPlayingDMX.Small", new IMAGE_BRUSH("Icons/icon_DMXPixelMappingEditor_StopPlayingDMX_40x", Icon20x20));
+	// Toolbar Icons
+	Style->Set("DMXPixelMappingEditor.AddMapping", new IMAGE_BRUSH("icon_DMXPixelMappingEditor_AddMapping_40x", Icon40x40));
+	Style->Set("DMXPixelMappingEditor.AddMapping.Small", new IMAGE_BRUSH("icon_DMXPixelMappingEditor_AddMapping_40x", Icon20x20));
 
+	Style->Set("DMXPixelMappingEditor.PlayDMX", new IMAGE_BRUSH("icon_DMXPixelMappingEditor_PlayDMX_40x", Icon40x40));
+	Style->Set("DMXPixelMappingEditor.PlayDMX.Small", new IMAGE_BRUSH("icon_DMXPixelMappingEditor_PlayDMX_40x", Icon20x20));
+	Style->Set("DMXPixelMappingEditor.StopPlayingDMX", new IMAGE_BRUSH("icon_DMXPixelMappingEditor_StopPlayingDMX_40x", Icon40x40));
+	Style->Set("DMXPixelMappingEditor.StopPlayingDMX.Small", new IMAGE_BRUSH("icon_DMXPixelMappingEditor_StopPlayingDMX_40x", Icon20x20));
+	
 
 	return Style;
 }
 
-#undef IMAGE_PLUGIN_BRUSH
-#undef IMAGE_CORE_BRUSH
 #undef IMAGE_BRUSH
+#undef IMAGE_BRUSH_SVG
+#undef IMAGE_CORE_BRUSH
+#undef IMAGE_PLUGIN_BRUSH
 #undef BOX_BRUSH
 #undef BORDER_BRUSH
 #undef BOX_CORE_BRUSH
diff --git a/Engine/Plugins/VirtualProduction/DMX/DMXPixelMapping/Source/DMXPixelMappingEditor/Private/Toolkits/DMXPixelMappingToolkit.cpp b/Engine/Plugins/VirtualProduction/DMX/DMXPixelMapping/Source/DMXPixelMappingEditor/Private/Toolkits/DMXPixelMappingToolkit.cpp
--- a/Engine/Plugins/VirtualProduction/DMX/DMXPixelMapping/Source/DMXPixelMappingEditor/Private/Toolkits/DMXPixelMappingToolkit.cpp
+++ b/Engine/Plugins/VirtualProduction/DMX/DMXPixelMapping/Source/DMXPixelMappingEditor/Private/Toolkits/DMXPixelMappingToolkit.cpp
@@ -6,6 +6,7 @@
 #include "DMXPixelMapping.h"
 #include "DMXPixelMappingEditorCommands.h"
 #include "DMXPixelMappingEditorModule.h"
+#include "DMXPixelMappingEditorStyle.h"
 #include "DMXPixelMappingEditorUtils.h"
 #include "DMXPixelMappingToolbar.h"
 #include "K2Node_PixelMappingBaseComponent.h"
@@ -34,6 +35,7 @@
 #include "Framework/Commands/GenericCommands.h"
 #include "Kismet2/BlueprintEditorUtils.h"
 #include "Misc/ScopedSlowTask.h"
+#include "Styling/AppStyle.h"
 #include "Widgets/SWidget.h"
 #include "Widgets/Docking/SDockTab.h"
 
@@ -85,12 +87,12 @@ void FDMXPixelMappingToolkit::RegisterTabSpawners(const TSharedRef<class FTabMan
 	InTabManager->RegisterTabSpawner(PaletteViewTabID, FOnSpawnTab::CreateSP(this, &FDMXPixelMappingToolkit::SpawnTab_PaletteView))
 		.SetDisplayName(LOCTEXT("Tab_PaletteView", "Palette"))
 		.SetGroup(WorkspaceMenuCategoryRef)
-		.SetIcon(FSlateIcon(FEditorStyle::GetStyleSetName(), "LevelEditor.Tabs.Viewports"));
+		.SetIcon(FSlateIcon(FDMXPixelMappingEditorStyle::GetStyleSetName(), "ClassIcon.DMXPixelMapping"));
 
 	InTabManager->RegisterTabSpawner(HierarchyViewTabID, FOnSpawnTab::CreateSP(this, &FDMXPixelMappingToolkit::SpawnTab_HierarchyView))
 		.SetDisplayName(LOCTEXT("Tab_HierarchyView", "Hierarchy"))
 		.SetGroup(WorkspaceMenuCategoryRef)
-		.SetIcon(FSlateIcon(FEditorStyle::GetStyleSetName(), "LevelEditor.Tabs.Viewports"));
+		.SetIcon(FSlateIcon(FAppStyle::Get().GetStyleSetName(), "Kismet.Tabs.Components"));
 
 	InTabManager->RegisterTabSpawner(DesignerViewTabID, FOnSpawnTab::CreateSP(this, &FDMXPixelMappingToolkit::SpawnTab_DesignerView))
 		.SetDisplayName(LOCTEXT("Tab_DesignerView", "Designer"))
@@ -100,12 +102,12 @@ void FDMXPixelMappingToolkit::RegisterTabSpawners(const TSharedRef<class FTabMan
 	InTabManager->RegisterTabSpawner(PreviewViewTabID, FOnSpawnTab::CreateSP(this, &FDMXPixelMappingToolkit::SpawnTab_PreviewView))
 		.SetDisplayName(LOCTEXT("Tab_PreviewView", "Preview"))
 		.SetGroup(WorkspaceMenuCategoryRef)
-		.SetIcon(FSlateIcon(FEditorStyle::GetStyleSetName(), "LevelEditor.Tabs.Viewports"));
+		.SetIcon(FSlateIcon(FDMXPixelMappingEditorStyle::GetStyleSetName(), "Icons.Preview"));
 
 	InTabManager->RegisterTabSpawner(DetailsViewTabID, FOnSpawnTab::CreateSP(this, &FDMXPixelMappingToolkit::SpawnTab_DetailsView))
 		.SetDisplayName(LOCTEXT("Tab_DetailsView", "Details"))
 		.SetGroup(WorkspaceMenuCategoryRef)
-		.SetIcon(FSlateIcon(FEditorStyle::GetStyleSetName(), "LevelEditor.Tabs.Viewports"));
+		.SetIcon(FSlateIcon(FAppStyle::Get().GetStyleSetName(), "Icons.Details"));
 }
 
 void FDMXPixelMappingToolkit::UnregisterTabSpawners(const TSharedRef<class FTabManager>& InTabManager)
diff --git a/Engine/Plugins/VirtualProduction/DMX/DMXPixelMapping/Source/DMXPixelMappingEditor/Private/Views/SDMXPixelMappingDesignerView.cpp b/Engine/Plugins/VirtualProduction/DMX/DMXPixelMapping/Source/DMXPixelMappingEditor/Private/Views/SDMXPixelMappingDesignerView.cpp
--- a/Engine/Plugins/VirtualProduction/DMX/DMXPixelMapping/Source/DMXPixelMappingEditor/Private/Views/SDMXPixelMappingDesignerView.cpp
+++ b/Engine/Plugins/VirtualProduction/DMX/DMXPixelMapping/Source/DMXPixelMappingEditor/Private/Views/SDMXPixelMappingDesignerView.cpp
@@ -5,6 +5,7 @@
 #include "DMXPixelMapping.h"
 #include "DMXPixelMappingComponentWidget.h"
 #include "DMXPixelMappingEditorUtils.h"
+#include "DMXPixelMappingEditorStyle.h"
 #include "SDMXPixelMappingComponentBox.h"
 #include "Components/DMXPixelMappingOutputComponent.h"
 #include "Components/DMXPixelMappingRendererComponent.h"
@@ -224,7 +225,7 @@ void SDMXPixelMappingDesignerView::Construct(const FArguments& InArgs, const TSh
 								.ContentPadding(FEditorStyle::Get().GetMargin("ViewportMenu.SToolBarButtonBlock.Button.Padding"))
 								[
 									SNew(SImage)
-									.Image(FEditorStyle::GetBrush("UMGEditor.ZoomToFit"))
+									.Image(FDMXPixelMappingEditorStyle::Get().GetBrush("Icons.ZoomToFit"))
 								]
 							]
 						]
diff --git a/Engine/Plugins/VirtualProduction/DMX/DMXPixelMapping/Source/DMXPixelMappingEditor/Private/Views/SDMXPixelMappingPreviewView.cpp b/Engine/Plugins/VirtualProduction/DMX/DMXPixelMapping/Source/DMXPixelMappingEditor/Private/Views/SDMXPixelMappingPreviewView.cpp
--- a/Engine/Plugins/VirtualProduction/DMX/DMXPixelMapping/Source/DMXPixelMappingEditor/Private/Views/SDMXPixelMappingPreviewView.cpp
+++ b/Engine/Plugins/VirtualProduction/DMX/DMXPixelMapping/Source/DMXPixelMappingEditor/Private/Views/SDMXPixelMappingPreviewView.cpp
@@ -1,6 +1,8 @@
 // Copyright Epic Games, Inc. All Rights Reserved.
 
 #include "Views/SDMXPixelMappingPreviewView.h"
+
+#include "DMXPixelMappingEditorStyle.h"
 #include "Toolkits/DMXPixelMappingToolkit.h"
 #include "Widgets/SDMXPixelMappingPreviewViewport.h"
 #include "Widgets/SDMXPixelMappingZoomPan.h"
@@ -287,7 +289,7 @@ TSharedRef<SWidget> SDMXPixelMappingPreviewView::CreateOverlayUI()
 				.ContentPadding(FEditorStyle::Get().GetMargin("ViewportMenu.SToolBarButtonBlock.Button.Padding"))
 				[
 					SNew(SImage)
-					.Image(FEditorStyle::GetBrush("UMGEditor.ZoomToFit"))
+					.Image(FDMXPixelMappingEditorStyle::Get().GetBrush("Icons.ZoomToFit"))
 				]
 			]
 		];
