DMX - Fix unchecking monitor all ports in DMX channel monitor causes a crash

diff --git a/Engine/Plugins/VirtualProduction/DMX/DMXEngine/Source/DMXEditor/Private/Widgets/Monitors/SDMXMonitorSourceSelector.cpp b/Engine/Plugins/VirtualProduction/DMX/DMXEngine/Source/DMXEditor/Private/Widgets/Monitors/SDMXMonitorSourceSelector.cpp
--- a/Engine/Plugins/VirtualProduction/DMX/DMXEngine/Source/DMXEditor/Private/Widgets/Monitors/SDMXMonitorSourceSelector.cpp
+++ b/Engine/Plugins/VirtualProduction/DMX/DMXEngine/Source/DMXEditor/Private/Widgets/Monitors/SDMXMonitorSourceSelector.cpp
@@ -275,7 +275,7 @@ void SDMXMonitorSourceSelector::UpdateMonitoredPorts()
 			}
 
 			// Make a selection if nothing was already selected
-			if (SelectedOutputPorts.IsEmpty() && FDMXPortManager::Get().GetInputPorts().Num() > 0)
+			if (SelectedOutputPorts.IsEmpty() && FDMXPortManager::Get().GetOutputPorts().Num() > 0)
 			{
 				const FDMXOutputPortSharedRef& NewSelection = FDMXPortManager::Get().GetOutputPorts()[0];
 				PortSelector->SelectPort(NewSelection->GetPortGuid());
