DMX -  Add fixture patch menu now correctly refreshes when fixture type names change

diff --git a/Engine/Plugins/VirtualProduction/DMX/DMXEngine/Source/DMXEditor/Private/Widgets/FixturePatch/SDMXFixturePatchTree.cpp b/Engine/Plugins/VirtualProduction/DMX/DMXEngine/Source/DMXEditor/Private/Widgets/FixturePatch/SDMXFixturePatchTree.cpp
--- a/Engine/Plugins/VirtualProduction/DMX/DMXEngine/Source/DMXEditor/Private/Widgets/FixturePatch/SDMXFixturePatchTree.cpp
+++ b/Engine/Plugins/VirtualProduction/DMX/DMXEngine/Source/DMXEditor/Private/Widgets/FixturePatch/SDMXFixturePatchTree.cpp
@@ -56,6 +56,9 @@ void SDMXFixturePatchTree::Construct(const FArguments& InArgs)
 		// Bind to fixture patch changes
 		UDMXEntityFixturePatch::GetOnFixturePatchChanged().AddSP(this, &SDMXFixturePatchTree::OnFixturePatchChanged);
 
+		// Bind to fixture type changes
+		UDMXEntityFixtureType::GetOnFixtureTypeChanged().AddSP(this, &SDMXFixturePatchTree::OnFixtureTypeChanged);
+
 		// Make an initial selection
 		const TArray<UDMXEntityFixturePatch*> FixturePatches = GetDMXLibrary()->GetEntitiesTypeCast<UDMXEntityFixturePatch>();
 		if (FixturePatches.Num() > 0)
@@ -741,6 +744,17 @@ void SDMXFixturePatchTree::OnFixturePatchChanged(const UDMXEntityFixturePatch* F
 	}
 }
 
+void SDMXFixturePatchTree::OnFixtureTypeChanged(const UDMXEntityFixtureType* FixtureType)
+{
+	if (TSharedPtr<FDMXEditor> PinnedDMXEditor = DMXEditor.Pin())
+	{
+		if (PinnedDMXEditor->GetDMXLibrary() == FixtureType->GetParentLibrary())
+		{
+			UpdateTree();
+		}
+	}
+}
+
 void SDMXFixturePatchTree::OnFixturePatchesSelected()
 {
 	if (!bChangingFixturePatchSelection)
diff --git a/Engine/Plugins/VirtualProduction/DMX/DMXEngine/Source/DMXEditor/Private/Widgets/FixturePatch/SDMXFixturePatchTree.h b/Engine/Plugins/VirtualProduction/DMX/DMXEngine/Source/DMXEditor/Private/Widgets/FixturePatch/SDMXFixturePatchTree.h
--- a/Engine/Plugins/VirtualProduction/DMX/DMXEngine/Source/DMXEditor/Private/Widgets/FixturePatch/SDMXFixturePatchTree.h
+++ b/Engine/Plugins/VirtualProduction/DMX/DMXEngine/Source/DMXEditor/Private/Widgets/FixturePatch/SDMXFixturePatchTree.h
@@ -74,6 +74,9 @@ private:
 	/** Called when a Fixture Patch changed */
 	void OnFixturePatchChanged(const UDMXEntityFixturePatch* FixturePatch);
 
+	/** Called when a Fixture Type changed */
+	void OnFixtureTypeChanged(const UDMXEntityFixtureType* FixtureType);
+
 	/** Called when Fixture Patches were selected in Fixture Patch Shared Data */
 	void OnFixturePatchesSelected();
 
