DMX - Fix DMX Library Track Recorder spins infinitely fast, normal and matrix fixture types are recorded correctly again

diff --git a/Engine/Plugins/VirtualProduction/DMX/DMXEngine/Source/DMXEditor/Private/Sequencer/DMXAsyncDMXRecorder.cpp b/Engine/Plugins/VirtualProduction/DMX/DMXEngine/Source/DMXEditor/Private/Sequencer/DMXAsyncDMXRecorder.cpp
--- a/Engine/Plugins/VirtualProduction/DMX/DMXEngine/Source/DMXEditor/Private/Sequencer/DMXAsyncDMXRecorder.cpp
+++ b/Engine/Plugins/VirtualProduction/DMX/DMXEngine/Source/DMXEditor/Private/Sequencer/DMXAsyncDMXRecorder.cpp
@@ -268,6 +268,7 @@ uint32 FDMXAsyncDMXRecorder::Run()
 	while (!bThreadStopping)
 	{
 		Update();
+		FPlatformProcess::Sleep(0.01f);
 	}
 
 	return 0;
diff --git a/Engine/Plugins/VirtualProduction/DMX/DMXEngine/Source/DMXRuntime/Private/Sequencer/MovieSceneDMXLibrarySection.cpp b/Engine/Plugins/VirtualProduction/DMX/DMXEngine/Source/DMXRuntime/Private/Sequencer/MovieSceneDMXLibrarySection.cpp
--- a/Engine/Plugins/VirtualProduction/DMX/DMXEngine/Source/DMXRuntime/Private/Sequencer/MovieSceneDMXLibrarySection.cpp
+++ b/Engine/Plugins/VirtualProduction/DMX/DMXEngine/Source/DMXRuntime/Private/Sequencer/MovieSceneDMXLibrarySection.cpp
@@ -96,15 +96,15 @@ void FDMXFixturePatchChannel::UpdateNumberOfChannels(bool bResetDefaultValues /*
 		}
 	}
 
-	if (Mode.FixtureMatrixConfig.GetLastChannel() <= Patch->GetChannelSpan())
+	if (Mode.bFixtureMatrixEnabled)
 	{
 		int32 NumXCells = Mode.FixtureMatrixConfig.XCells;
 		int32 NumYCells = Mode.FixtureMatrixConfig.YCells;
 
 		for (int32 IdxCellY = 0; IdxCellY < NumYCells; IdxCellY++)
 		{
-		for (int32 IdxCellX = 0; IdxCellX < NumXCells; IdxCellX++)
-		{
+			for (int32 IdxCellX = 0; IdxCellX < NumXCells; IdxCellX++)
+			{
 				FIntPoint CellCoordinates = FIntPoint(IdxCellX, IdxCellY);
 
 				for (const FDMXFixtureCellAttribute& CellAttribute : Mode.FixtureMatrixConfig.CellAttributes)
