DMX - Deleting a fixture type or fixture patch now updates selection correctly again, preventing from various issues

diff --git a/Engine/Plugins/VirtualProduction/DMX/DMXEngine/Source/DMXEditor/Private/Widgets/FixturePatch/SDMXFixturePatchTree.cpp b/Engine/Plugins/VirtualProduction/DMX/DMXEngine/Source/DMXEditor/Private/Widgets/FixturePatch/SDMXFixturePatchTree.cpp
--- a/Engine/Plugins/VirtualProduction/DMX/DMXEngine/Source/DMXEditor/Private/Widgets/FixturePatch/SDMXFixturePatchTree.cpp
+++ b/Engine/Plugins/VirtualProduction/DMX/DMXEngine/Source/DMXEditor/Private/Widgets/FixturePatch/SDMXFixturePatchTree.cpp
@@ -523,16 +523,14 @@ void SDMXFixturePatchTree::OnPasteNodes()
 
 				// Move the Fixture Type template from the transient package into the Library package
 				FixturePatch->Rename(*MakeUniqueObjectName(Library, UDMXEntityFixtureType::StaticClass()).ToString(), Library, REN_DoNotDirty | REN_DontCreateRedirectors);
-				FixturePatch->RefreshID();
-
-				// Add to the Library
-				AutoAssignCopiedPatch(FixturePatch);
-
-				// Move the Entity from the transient package into the Library package
-				FixturePatch->Rename(*MakeUniqueObjectName(Library, FixturePatch->GetClass()).ToString(), Library, REN_DoNotDirty | REN_DontCreateRedirectors);
 
 				// Make sure the Entity's name won't collide with existing ones
 				FixturePatch->SetName(FDMXRuntimeUtils::FindUniqueEntityName(Library, FixturePatch->GetClass(), FixturePatch->GetDisplayName()));
+
+				// Update the library and ID
+				FixturePatch->RefreshID();
+				FixturePatch->SetParentLibrary(Library);
+				AutoAssignCopiedPatch(FixturePatch);
 			}
 		}
 
@@ -617,6 +615,36 @@ void SDMXFixturePatchTree::OnDeleteNodes()
 		}
 	}
 
+	// Find a new selection
+	const TSharedPtr<FDMXEntityTreeEntityNode> EntityNodeToSelect = [&FixturePatchesToDelete, this]() -> TSharedPtr<FDMXEntityTreeEntityNode>
+	{
+		const TArray<TSharedPtr<FDMXEntityTreeEntityNode>> EntityNodes = GetEntityNodes();
+
+		for (const TSharedPtr<FDMXEntityTreeEntityNode>& EntityNode : EntityNodes)
+		{
+			if (EntityNode.IsValid() && !FixturePatchesToDelete.Contains(EntityNode->GetEntity()))
+			{
+				return EntityNode;
+			}
+		}
+		return nullptr;
+	}();
+
+	// Apply the new selectio
+	if (EntityNodeToSelect.IsValid())
+	{
+		if (UDMXEntityFixturePatch* FixturePatchToSelect = Cast<UDMXEntityFixturePatch>(EntityNodeToSelect->GetEntity()))
+		{
+			const TArray<TWeakObjectPtr<UDMXEntityFixturePatch>> FixturePatchsToSelect = { FixturePatchToSelect };
+			FixturePatchSharedData->SelectFixturePatches(FixturePatchsToSelect);
+		}
+	}
+	else
+	{
+		// Clear selection if no patches remain
+		FixturePatchSharedData->SelectFixturePatches(TArray<TWeakObjectPtr<UDMXEntityFixturePatch>>());
+	}
+
 	// Remove the Fixture Patches from the DMX Library
 	const FScopedTransaction Transaction(EntitiesToDelete.Num() > 1 ? LOCTEXT("RemoveEntities", "Remove Entities") : LOCTEXT("RemoveEntity", "Remove Entity"));
 
@@ -630,12 +658,6 @@ void SDMXFixturePatchTree::OnDeleteNodes()
 	}
 	DMXLibrary->PostEditChange();
 
-	// Clear selection if no patches remain
-	if (GetSelectedEntities().Num() == 0)
-	{
-		FixturePatchSharedData->SelectFixturePatches(TArray<TWeakObjectPtr<UDMXEntityFixturePatch>>());
-	}
-
 	UpdateTree();
 }
 
diff --git a/Engine/Plugins/VirtualProduction/DMX/DMXEngine/Source/DMXEditor/Private/Widgets/FixtureType/SDMXFixtureTypeTree.cpp b/Engine/Plugins/VirtualProduction/DMX/DMXEngine/Source/DMXEditor/Private/Widgets/FixtureType/SDMXFixtureTypeTree.cpp
--- a/Engine/Plugins/VirtualProduction/DMX/DMXEngine/Source/DMXEditor/Private/Widgets/FixtureType/SDMXFixtureTypeTree.cpp
+++ b/Engine/Plugins/VirtualProduction/DMX/DMXEngine/Source/DMXEditor/Private/Widgets/FixtureType/SDMXFixtureTypeTree.cpp
@@ -313,6 +313,10 @@ void SDMXFixtureTypeTree::OnPasteNodes()
 			// Make sure the Entity's name won't collide with existing ones
 			NewEntity->SetName(FDMXRuntimeUtils::FindUniqueEntityName(Library, NewEntity->GetClass(), NewEntity->GetDisplayName()));
 
+			// Update the library and ID
+			NewEntity->SetParentLibrary(Library);
+			NewEntity->RefreshID();
+
 			NewFixtureTypes.Add(CastChecked<UDMXEntityFixtureType>(NewEntity));
 		}
 
@@ -385,6 +389,11 @@ void SDMXFixtureTypeTree::OnDeleteNodes()
 		}
 	}
 
+	if (FixtureTypesToDelete.Num() == 0)
+	{
+		return;
+	}
+
 	// Find Fixture Types that are used by Fixture Patches
 	UDMXLibrary* DMXLibrary = GetDMXLibrary();
 	TArray<UDMXEntity*> EntitiesInUse;
@@ -419,6 +428,38 @@ void SDMXFixtureTypeTree::OnDeleteNodes()
 
 	if (bConfirmDelete)
 	{
+		// Find a new selection
+		const TSharedPtr<FDMXEntityTreeEntityNode> EntityNodeToSelect = [&FixtureTypesToDelete, this]() -> TSharedPtr<FDMXEntityTreeEntityNode>
+		{
+			const TArray<TSharedPtr<FDMXEntityTreeEntityNode>> EntityNodes = GetEntityNodes();
+
+			for (const TSharedPtr<FDMXEntityTreeEntityNode>& EntityNode : EntityNodes)
+			{
+				if (EntityNode.IsValid() && !FixtureTypesToDelete.Contains(EntityNode->GetEntity()))
+				{
+					return EntityNode;
+				}
+			}
+			return nullptr;
+		}();
+
+		// Apply the new selection
+		if (EntityNodeToSelect.IsValid())
+		{
+			if (UDMXEntityFixtureType* FixtureTypeToSelect = Cast<UDMXEntityFixtureType>(EntityNodeToSelect->GetEntity()))
+			{
+				SelectItemByNode(EntityNodeToSelect.ToSharedRef());
+
+				const TArray<TWeakObjectPtr<UDMXEntityFixtureType>> FixtureTypesToSelect = { FixtureTypeToSelect };
+				FixtureTypeSharedData->SelectFixtureTypes(FixtureTypesToSelect);
+			}
+		}
+		else
+		{		
+			// Clear selection if no fixture types remain
+			FixtureTypeSharedData->SelectFixtureTypes(TArray<TWeakObjectPtr<UDMXEntityFixtureType>>());
+		}
+
 		// Remove the Fixture Types from the DMX Library
 		const FScopedTransaction Transaction(EntitiesToDelete.Num() > 1 ? LOCTEXT("RemoveEntities", "Remove Entities") : LOCTEXT("RemoveEntity", "Remove Entity"));
 
@@ -431,11 +472,6 @@ void SDMXFixtureTypeTree::OnDeleteNodes()
 		}
 		DMXLibrary->PostEditChange();
 
-		// Clear selection if no fixture types remain
-		if (DMXLibrary->GetEntitiesTypeCast<UDMXEntityFixtureType>().Num() == 0)
-		{
-			FixtureTypeSharedData->SelectFixtureTypes(TArray<TWeakObjectPtr<UDMXEntityFixtureType>>());
-		}
 
 		UpdateTree();
 	}
diff --git a/Engine/Plugins/VirtualProduction/DMX/DMXEngine/Source/DMXEditor/Private/Widgets/SDMXEntityTreeViewBase.cpp b/Engine/Plugins/VirtualProduction/DMX/DMXEngine/Source/DMXEditor/Private/Widgets/SDMXEntityTreeViewBase.cpp
--- a/Engine/Plugins/VirtualProduction/DMX/DMXEngine/Source/DMXEditor/Private/Widgets/SDMXEntityTreeViewBase.cpp
+++ b/Engine/Plugins/VirtualProduction/DMX/DMXEngine/Source/DMXEditor/Private/Widgets/SDMXEntityTreeViewBase.cpp
@@ -213,6 +213,28 @@ void SDMXEntityTreeViewBase::UpdateTree(bool bRegenerateTreeNodes)
 	}
 }
 
+TArray<TSharedPtr<FDMXEntityTreeEntityNode>> SDMXEntityTreeViewBase::GetEntityNodes(TSharedPtr<FDMXEntityTreeNodeBase> ParentNode) const
+{
+	TArray<TSharedPtr<FDMXEntityTreeEntityNode>> Result;
+
+	ParentNode = ParentNode.IsValid() ? ParentNode : GetRootNode();
+	for (const TSharedPtr<FDMXEntityTreeNodeBase>& Child : ParentNode->GetChildren())
+	{
+		if (Child.IsValid())
+		{
+			if (Child->GetNodeType() == FDMXEntityTreeNodeBase::ENodeType::EntityNode)
+			{
+				Result.Add(StaticCastSharedPtr<FDMXEntityTreeEntityNode>(Child));
+			}
+
+			// Recursive for all childs
+			Result.Append(GetEntityNodes(Child));
+		}
+	}
+
+	return Result;
+}
+
 TSharedPtr<FDMXEntityTreeCategoryNode> SDMXEntityTreeViewBase::FindCategoryNodeOfEntity(UDMXEntity* Entity) const
 {
 	TSharedPtr<FDMXEntityTreeEntityNode> EntityNode = FindNodeByEntity(Entity);
diff --git a/Engine/Plugins/VirtualProduction/DMX/DMXEngine/Source/DMXEditor/Private/Widgets/SDMXEntityTreeViewBase.h b/Engine/Plugins/VirtualProduction/DMX/DMXEngine/Source/DMXEditor/Private/Widgets/SDMXEntityTreeViewBase.h
--- a/Engine/Plugins/VirtualProduction/DMX/DMXEngine/Source/DMXEditor/Private/Widgets/SDMXEntityTreeViewBase.h
+++ b/Engine/Plugins/VirtualProduction/DMX/DMXEngine/Source/DMXEditor/Private/Widgets/SDMXEntityTreeViewBase.h
@@ -96,6 +96,9 @@ public:
 		return NewNode;
 	}
 
+	/** Returns the Entity Nodes in the Tree */
+	TArray<TSharedPtr<FDMXEntityTreeEntityNode>> GetEntityNodes(TSharedPtr<FDMXEntityTreeNodeBase> ParentNode = nullptr) const;
+
 	/** Returns the category node of the entity, or null if is not in the list */
 	TSharedPtr<FDMXEntityTreeCategoryNode> FindCategoryNodeOfEntity(UDMXEntity* Entity) const;
 
diff --git a/Engine/Plugins/VirtualProduction/DMX/DMXEngine/Source/DMXRuntime/Public/Library/DMXEntity.h b/Engine/Plugins/VirtualProduction/DMX/DMXEngine/Source/DMXRuntime/Public/Library/DMXEntity.h
--- a/Engine/Plugins/VirtualProduction/DMX/DMXEngine/Source/DMXRuntime/Public/Library/DMXEntity.h
+++ b/Engine/Plugins/VirtualProduction/DMX/DMXEngine/Source/DMXRuntime/Public/Library/DMXEntity.h
@@ -60,7 +60,7 @@ public:
 		return IsValidEntity(OutReason);
 	}
 
-	UE_DEPRECATED(5.0, "The library should not be set explicitly anymore. Enties register themselves with the DMXLibrary on creation.")
+	/** Sets the DMX Library of the entity */
 	void SetParentLibrary(UDMXLibrary* InParent);
 
 	/** Gets the library the Entity resides in */
