DMX - Matrix Cell Data Type now correctly affects the channel span of the fixture type

diff --git a/Engine/Plugins/VirtualProduction/DMX/DMXEngine/Source/DMXEditor/Private/Widgets/FixtureType/SDMXFixtureModeEditor.cpp b/Engine/Plugins/VirtualProduction/DMX/DMXEngine/Source/DMXEditor/Private/Widgets/FixtureType/SDMXFixtureModeEditor.cpp
--- a/Engine/Plugins/VirtualProduction/DMX/DMXEngine/Source/DMXEditor/Private/Widgets/FixtureType/SDMXFixtureModeEditor.cpp
+++ b/Engine/Plugins/VirtualProduction/DMX/DMXEngine/Source/DMXEditor/Private/Widgets/FixtureType/SDMXFixtureModeEditor.cpp
@@ -44,7 +44,6 @@ void SDMXFixtureModeEditor::NotifyPreChange(FProperty* PropertyAboutToChange)
 	FDMXFixtureMode* ModeBeingEditedPtr = GetModeBeingEdited();
 	if (PropertyAboutToChange && FixtureType && ModeBeingEditedPtr)
 	{
-		FixtureType->Modify();
 		FixtureType->PreEditChange(PropertyAboutToChange);
 
 		const FName PropertyName = PropertyAboutToChange->GetFName();
@@ -75,6 +74,10 @@ void SDMXFixtureModeEditor::NotifyPreChange(FProperty* PropertyAboutToChange)
 		{
 			Transaction = MakeUnique<FScopedTransaction>(LOCTEXT("SetChannelSpanTransaction", "Channel Span of Mode"));
 		}
+		else if (PropertyName == GET_MEMBER_NAME_CHECKED(FDMXFixtureCellAttribute, DataType))
+		{
+			Transaction = MakeUnique<FScopedTransaction>(LOCTEXT("SetCellAttributeDataType", "CellAttributeDatatType"));
+		}
 	}
 }
 
@@ -119,7 +122,15 @@ void SDMXFixtureModeEditor::NotifyPostChange(const FPropertyChangedEvent& Proper
 			}
 			Refresh();
 		}
+		else if (PropertyName == GET_MEMBER_NAME_CHECKED(FDMXFixtureCellAttribute, DataType))
+		{
+			// Disable the matrix and enable it again so surrounding functions align
+			FixtureType->SetFixtureMatrixEnabled(ModeIndex, false);
+			FixtureType->SetFixtureMatrixEnabled(ModeIndex, true);
 
+			FixtureType->UpdateChannelSpan(ModeIndex);
+			Refresh();
+		}
 		FPropertyChangedEvent ObjectPropertyChangedEvent(PropertyChangedEvent);
 		FixtureType->PostEditChangeProperty(ObjectPropertyChangedEvent);
 	}
