DMX - Fixture type menu in fixture type tab now correctly adopts name changes of fixture types

diff --git a/Engine/Plugins/VirtualProduction/DMX/DMXEngine/Source/DMXEditor/Private/Widgets/FixturePatch/SDMXFixturePatchTree.cpp b/Engine/Plugins/VirtualProduction/DMX/DMXEngine/Source/DMXEditor/Private/Widgets/FixturePatch/SDMXFixturePatchTree.cpp
--- a/Engine/Plugins/VirtualProduction/DMX/DMXEngine/Source/DMXEditor/Private/Widgets/FixturePatch/SDMXFixturePatchTree.cpp
+++ b/Engine/Plugins/VirtualProduction/DMX/DMXEngine/Source/DMXEditor/Private/Widgets/FixturePatch/SDMXFixturePatchTree.cpp
@@ -108,7 +108,10 @@ TSharedRef<SWidget> SDMXFixturePatchTree::GenerateAddNewEntityButton()
 			.ButtonStyle(FEditorStyle::Get(), "FlatButton.Success")
 			.ForegroundColor(FLinearColor::White)
 			.ToolTipText(AddButtonToolTip)
-			.OnComboBoxOpened(FOnComboBoxOpened::CreateLambda([this]() { AddButtonDropdownList->ClearSelection(); }));
+			.OnComboBoxOpened(FOnComboBoxOpened::CreateLambda([this]() 
+				{ 
+					AddButtonDropdownList->RefreshEntitiesList();
+				}));
 
 	AddButtonDropdownList->SetComboButton(AddComboButton);
 
