DMX - GDTF Functions without an offset are no longer interpreted as valid channels with offset 0

diff --git a/Engine/Plugins/VirtualProduction/DMX/DMXEngine/Source/DMXEditor/Private/Factories/DMXGDTFImporter.cpp b/Engine/Plugins/VirtualProduction/DMX/DMXEngine/Source/DMXEditor/Private/Factories/DMXGDTFImporter.cpp
--- a/Engine/Plugins/VirtualProduction/DMX/DMXEngine/Source/DMXEditor/Private/Factories/DMXGDTFImporter.cpp
+++ b/Engine/Plugins/VirtualProduction/DMX/DMXEngine/Source/DMXEditor/Private/Factories/DMXGDTFImporter.cpp
@@ -418,9 +418,16 @@ UDMXImportGDTF* FDMXGDTFImporter::CreateGDTFDesctription()
                         {
                             for (const FXmlNode* DMXChannelNode : DMXChannelsNode->GetChildrenNodes())
                             {
-                                FDMXImportGDTFDMXChannel ImportDMXChannel;
+								FDMXImportGDTFDMXChannel ImportDMXChannel;
+
+								// Ignore Channels that do not specify a valid offset.
+								// E.g. 'Robe Lighting s.r.o.@Robin 800X LEDWash.gdtf' specifies virtual Dimmer channels that have no offset and cannot be accessed.
+								if (!ImportDMXChannel.ParseOffset(DMXChannelNode->GetAttribute("Offset")))
+								{
+									continue;
+								}
+
                                 LexTryParseString(ImportDMXChannel.DMXBreak, *DMXChannelNode->GetAttribute("DMXBreak"));
-                                ImportDMXChannel.ParseOffset(DMXChannelNode->GetAttribute("Offset"));
                                 ImportDMXChannel.Default = FDMXImportGDTFDMXValue(DMXChannelNode->GetAttribute("Default"));
                                 ImportDMXChannel.Highlight = FDMXImportGDTFDMXValue(DMXChannelNode->GetAttribute("Highlight"));
                                 ImportDMXChannel.Geometry = FName(*DMXChannelNode->GetAttribute("Geometry"));
diff --git a/Engine/Plugins/VirtualProduction/DMX/DMXEngine/Source/DMXRuntime/Private/Library/DMXImportGDTF.cpp b/Engine/Plugins/VirtualProduction/DMX/DMXEngine/Source/DMXRuntime/Private/Library/DMXImportGDTF.cpp
--- a/Engine/Plugins/VirtualProduction/DMX/DMXEngine/Source/DMXRuntime/Private/Library/DMXImportGDTF.cpp
+++ b/Engine/Plugins/VirtualProduction/DMX/DMXEngine/Source/DMXRuntime/Private/Library/DMXImportGDTF.cpp
@@ -112,21 +112,25 @@ FDMXImportGDTFDMXValue::FDMXImportGDTFDMXValue(const FString& InDMXValueStr)
     }
 };
 
-void FDMXImportGDTFDMXChannel::ParseOffset(const FString& InOffsetStr)
+bool FDMXImportGDTFDMXChannel::ParseOffset(const FString& InOffsetStr)
 {
-    if (!InOffsetStr.IsEmpty())
-    {
-        Offset.Empty();
-        TArray<FString> OffsetStrArray;
-        InOffsetStr.ParseIntoArray(OffsetStrArray, TEXT(","));
+	if (InOffsetStr.IsEmpty())
+	{
+		return false;
+	}
 
-        for (int32 OffsetIndex = 0; OffsetIndex < OffsetStrArray.Num(); ++OffsetIndex)
-        {
-            int32 OffsetValue;
-            LexTryParseString(OffsetValue, *OffsetStrArray[OffsetIndex]);
-            Offset.Add(OffsetValue);
-        }
+    Offset.Empty();
+    TArray<FString> OffsetStrArray;
+    InOffsetStr.ParseIntoArray(OffsetStrArray, TEXT(","));
+
+    for (int32 OffsetIndex = 0; OffsetIndex < OffsetStrArray.Num(); ++OffsetIndex)
+    {
+        int32 OffsetValue;
+        LexTryParseString(OffsetValue, *OffsetStrArray[OffsetIndex]);
+        Offset.Add(OffsetValue);
     }
+    
+	return true;
 }
 
 TArray<FDMXImportGDTFChannelFunction> UDMXImportGDTFDMXModes::GetDMXChannelFunctions(const FDMXImportGDTFDMXMode& InMode)
diff --git a/Engine/Plugins/VirtualProduction/DMX/DMXEngine/Source/DMXRuntime/Public/Library/DMXImportGDTF.h b/Engine/Plugins/VirtualProduction/DMX/DMXEngine/Source/DMXRuntime/Public/Library/DMXImportGDTF.h
--- a/Engine/Plugins/VirtualProduction/DMX/DMXEngine/Source/DMXRuntime/Public/Library/DMXImportGDTF.h
+++ b/Engine/Plugins/VirtualProduction/DMX/DMXEngine/Source/DMXRuntime/Public/Library/DMXImportGDTF.h
@@ -613,7 +613,8 @@ struct DMXRUNTIME_API FDMXImportGDTFDMXChannel
 {
     GENERATED_BODY()
 
-    void ParseOffset(const FString& InOffsetStr);
+	/** Parses the offset of the channel. Returns false if no valid offset is specified */
+    bool ParseOffset(const FString& InOffsetStr);
 
     UPROPERTY(EditAnywhere, BlueprintReadOnly, BlueprintReadOnly, Category = "DMX")
     int32 DMXBreak = 0;
