DMX - DMX component now optionally raises OnFixturePatchReceived events at editor time

diff --git a/Engine/Plugins/VirtualProduction/DMX/DMXEngine/Source/DMXRuntime/Private/Game/DMXComponent.cpp b/Engine/Plugins/VirtualProduction/DMX/DMXEngine/Source/DMXRuntime/Private/Game/DMXComponent.cpp
--- a/Engine/Plugins/VirtualProduction/DMX/DMXEngine/Source/DMXRuntime/Private/Game/DMXComponent.cpp
+++ b/Engine/Plugins/VirtualProduction/DMX/DMXEngine/Source/DMXRuntime/Private/Game/DMXComponent.cpp
@@ -2,6 +2,7 @@
 
 #include "Game/DMXComponent.h"
 
+#include "DMXProtocolSettings.h"
 #include "DMXRuntimeLog.h"
 #include "DMXStats.h"
 #include "Library/DMXEntityFixturePatch.h"
@@ -11,7 +12,18 @@ UDMXComponent::UDMXComponent()
 	: bReceiveDMXFromPatch(true)
 {
 	PrimaryComponentTick.bCanEverTick = true;
-	bTickInEditor = false;
+	PrimaryComponentTick.bStartWithTickEnabled = false;
+	PrimaryComponentTick.TickGroup = ETickingGroup::TG_PrePhysics;
+
+#if WITH_EDITOR
+	bTickInEditor = true;
+	if (!HasAnyFlags(RF_ClassDefaultObject | RF_ArchetypeObject))
+	{
+		UDMXProtocolSettings* DMXProtocolSettings = GetMutableDefault<UDMXProtocolSettings>();
+		DMXProtocolSettings->GetOnAllFixturePatchesReceiveDMXInEditorEnabled().AddUObject(this, &UDMXComponent::OnAllFixturePatchesReceiveDMXInEditorEnabled);
+		UDMXEntityFixturePatch::GetOnFixturePatchChanged().AddUObject(this, &UDMXComponent::OnFixturePatchPropertiesChanged);
+	}
+#endif
 }
 
 UDMXEntityFixturePatch* UDMXComponent::GetFixturePatch() const
@@ -31,8 +43,10 @@ void UDMXComponent::SetFixturePatch(UDMXEntityFixturePatch* InFixturePatch)
 			PreviousFixturePatch->OnFixturePatchReceivedDMX.RemoveAll(this);
 		}
 
+		// Set the new patch and setup the new binding
 		FixturePatchRef.SetEntity(InFixturePatch);
 		SetupReceiveDMXBinding();
+		UpdateTickEnabled();
 	}
 }
 
@@ -41,10 +55,15 @@ void UDMXComponent::SetReceiveDMXFromPatch(bool bReceive)
 	bReceiveDMXFromPatch = bReceive;
 
 	SetupReceiveDMXBinding();
+	UpdateTickEnabled();
 }
 
 void UDMXComponent::OnFixturePatchReceivedDMX(UDMXEntityFixturePatch* FixturePatch, const FDMXNormalizedAttributeValueMap& NormalizedValuePerAttribute)
 {
+#if WITH_EDITOR
+	FEditorScriptExecutionGuard ScriptGuard;
+#endif
+	
 	OnFixturePatchReceived.Broadcast(FixturePatch, NormalizedValuePerAttribute);
 }
 
@@ -53,19 +72,90 @@ void UDMXComponent::SetupReceiveDMXBinding()
 	UDMXEntityFixturePatch* FixturePatch = FixturePatchRef.GetFixturePatch();
 	if (IsValid(FixturePatch))
 	{
-		if (bReceiveDMXFromPatch && !FixturePatch->OnFixturePatchReceivedDMX.Contains(this, GET_FUNCTION_NAME_CHECKED(UDMXComponent, OnFixturePatchReceivedDMX)))
+		const bool bReceiveDMX = [this, FixturePatch]()
+		{
+#if WITH_EDITOR
+			if (bReceiveDMXFromPatch && GIsEditor && !GIsPlayInEditorWorld)
+			{
+				const UDMXProtocolSettings* ProtocolSettings = GetDefault<UDMXProtocolSettings>();
+				const bool bFixturePatchReceivesDMXInEditor = FixturePatch->bReceiveDMXInEditor;
+				const bool bAllFixturePatchesReceiveDMXInEditor = ProtocolSettings->ShouldAllFixturePatchesReceiveDMXInEditor();
+
+				return bFixturePatchReceivesDMXInEditor || bAllFixturePatchesReceiveDMXInEditor;
+			}
+#endif
+			return bReceiveDMXFromPatch;
+		}();
+		
+		if (bReceiveDMX && !FixturePatch->OnFixturePatchReceivedDMX.Contains(this, GET_FUNCTION_NAME_CHECKED(UDMXComponent, OnFixturePatchReceivedDMX)))
 		{
 			// Enable receive DMX
 			FixturePatch->OnFixturePatchReceivedDMX.AddDynamic(this, &UDMXComponent::OnFixturePatchReceivedDMX);
+
+			FDMXNormalizedAttributeValueMap NormalizeAttributeValues;
+			FixturePatch->GetNormalizedAttributesValues(NormalizeAttributeValues);
+
+			if (NormalizeAttributeValues.Map.Num() > 0)
+			{
+				OnFixturePatchReceived.Broadcast(FixturePatch, NormalizeAttributeValues);
+			}
 		}
-		else if (!bReceiveDMXFromPatch && FixturePatch->OnFixturePatchReceivedDMX.Contains(this, GET_FUNCTION_NAME_CHECKED(UDMXComponent, OnFixturePatchReceivedDMX)))
-		{
+		else if (!bReceiveDMX)
+		{			
 			// Disable receive DMX
 			FixturePatch->OnFixturePatchReceivedDMX.RemoveAll(this);
 		}
 	}
 }
 
+void UDMXComponent::UpdateTickEnabled()
+{
+#if WITH_EDITOR
+	FEditorScriptExecutionGuard ScriptGuard;
+#endif
+
+	if (!OnDMXComponentTick.IsBound())
+	{
+		SetComponentTickEnabled(false);
+		return;
+	}
+
+	if (!bReceiveDMXFromPatch)
+	{
+		SetComponentTickEnabled(false);
+		return;
+	}
+
+	UDMXEntityFixturePatch* FixturePatch = FixturePatchRef.GetFixturePatch();
+	if (!FixturePatch)
+	{
+		SetComponentTickEnabled(false);
+		return;
+	}
+
+#if WITH_EDITOR
+	if (GIsEditor && !GIsPlayInEditorWorld)
+	{
+		const UDMXProtocolSettings* ProtocolSettings = GetDefault<UDMXProtocolSettings>();
+		if (!FixturePatch->bReceiveDMXInEditor && !ProtocolSettings->ShouldAllFixturePatchesReceiveDMXInEditor())
+		{
+			SetComponentTickEnabled(false);
+			return;
+		} 
+	}
+#endif 
+
+	SetComponentTickEnabled(true);
+}
+
+void UDMXComponent::OnRegister()
+{
+	Super::OnRegister();
+
+	SetupReceiveDMXBinding();
+	UpdateTickEnabled();
+}
+
 void UDMXComponent::BeginPlay()
 {
 	Super::BeginPlay();
@@ -82,6 +172,14 @@ void UDMXComponent::BeginPlay()
 	}
 
 	SetupReceiveDMXBinding();
+	UpdateTickEnabled();
+}
+
+void UDMXComponent::TickComponent(float DeltaTime, enum ELevelTick TickType, FActorComponentTickFunction* ThisTickFunction)
+{
+	Super::TickComponent(DeltaTime, TickType, ThisTickFunction);
+
+	OnDMXComponentTick.Broadcast(DeltaTime);
 }
 
 void UDMXComponent::DestroyComponent(bool bPromoteChildren)
@@ -109,7 +207,25 @@ void UDMXComponent::PostEditChangeProperty(FPropertyChangedEvent& PropertyChange
 	else if (PropertyName == GET_MEMBER_NAME_CHECKED(UDMXComponent, bReceiveDMXFromPatch))
 	{
 		SetupReceiveDMXBinding();
+		UpdateTickEnabled();
 	}
 }
 #endif // WITH_EDITOR
 
+#if WITH_EDITOR
+void UDMXComponent::OnAllFixturePatchesReceiveDMXInEditorEnabled(bool bEnabled)
+{
+	SetupReceiveDMXBinding();
+	UpdateTickEnabled();
+}
+#endif // WITH_EDITOR
+
+#if WITH_EDITOR
+void UDMXComponent::OnFixturePatchPropertiesChanged(const UDMXEntityFixturePatch* FixturePatch)
+{
+	if (FixturePatch == GetFixturePatch())
+	{
+		UpdateTickEnabled();
+	}
+}
+#endif // WITH_EDITOR
diff --git a/Engine/Plugins/VirtualProduction/DMX/DMXEngine/Source/DMXRuntime/Private/Library/DMXEntityFixturePatch.cpp b/Engine/Plugins/VirtualProduction/DMX/DMXEngine/Source/DMXRuntime/Private/Library/DMXEntityFixturePatch.cpp
--- a/Engine/Plugins/VirtualProduction/DMX/DMXEngine/Source/DMXRuntime/Private/Library/DMXEntityFixturePatch.cpp
+++ b/Engine/Plugins/VirtualProduction/DMX/DMXEngine/Source/DMXRuntime/Private/Library/DMXEntityFixturePatch.cpp
@@ -171,11 +171,26 @@ bool UDMXEntityFixturePatch::IsTickable() const
 bool UDMXEntityFixturePatch::IsTickableInEditor() const
 {
 	const bool bHasListener = OnFixturePatchReceivedDMX.IsBound();
+	if (!bHasListener)
+	{
+		return false;
+	}
 
-#if WITH_EDITORONLY_DATA
-	return bHasListener && bReceiveDMXInEditor;
+#if WITH_EDITOR
+	if (GIsEditor && !GIsPlayInEditorWorld)
+	{
+		if (bReceiveDMXInEditor)
+		{
+			return true;
+		}
+
+		const UDMXProtocolSettings* ProtocolSettings = GetDefault<UDMXProtocolSettings>();
+		const bool bAllFixturePatchesReceiveDMXInEditor = ProtocolSettings->ShouldAllFixturePatchesReceiveDMXInEditor();
+		return bAllFixturePatchesReceiveDMXInEditor;
+	}
 #endif
-	return bHasListener; 
+
+	return true; 
 }
 
 ETickableTickType UDMXEntityFixturePatch::GetTickableTickType() const
diff --git a/Engine/Plugins/VirtualProduction/DMX/DMXEngine/Source/DMXRuntime/Public/Game/DMXComponent.h b/Engine/Plugins/VirtualProduction/DMX/DMXEngine/Source/DMXRuntime/Public/Game/DMXComponent.h
--- a/Engine/Plugins/VirtualProduction/DMX/DMXEngine/Source/DMXRuntime/Public/Game/DMXComponent.h
+++ b/Engine/Plugins/VirtualProduction/DMX/DMXEngine/Source/DMXRuntime/Public/Game/DMXComponent.h
@@ -18,23 +18,31 @@ class UDMXEntityFixturePatch;
 
 /** 
  * Component that receives DMX input each Tick from a fixture patch.  
- * NOTE: Does not support receive in Editor! Use the 'Get DMX Fixture Patch' and bind 'On Fixture Patch Received DMX' instead (requires the patch to be set to 'Receive DMX in Editor' in the library). 
  */
-UCLASS( ClassGroup=(DMX), meta=(BlueprintSpawnableComponent), HideCategories = ("Variable", "Sockets", "Tags", "Activation", "Cooking", "ComponentReplication", "AssetUserData", "Collision", "Events"))
+UCLASS( ClassGroup=(DMX), meta=(BlueprintSpawnableComponent), HideCategories = ("Variable", "Sockets", "Activation", "Cooking", "ComponentReplication", "Collision", "ComponentTick"))
 class DMXRUNTIME_API UDMXComponent
 	: public UActorComponent
 {
 	GENERATED_BODY()
 
 	DECLARE_DYNAMIC_MULTICAST_DELEGATE_TwoParams(FDMXComponentFixturePatchReceivedSignature, UDMXEntityFixturePatch*, FixturePatch, const FDMXNormalizedAttributeValueMap&, ValuePerAttribute);
+	DECLARE_DYNAMIC_MULTICAST_DELEGATE_OneParam(FDMXOnDMXComponentTickSignature, float, DeltaTime);
 
 public:
 	UDMXComponent();
 
 	/** Broadcast when the component's fixture patch received DMX */
-	UPROPERTY(BlueprintAssignable, Category = "Components|DMX");
+	UPROPERTY(BlueprintAssignable, Category = "DMX");
 	FDMXComponentFixturePatchReceivedSignature OnFixturePatchReceived;
 
+	/**
+	 * Event raised each tick, when the component has a Fixture Patch set and either:
+	 * - The assigned fixture patch is set to 'receives DMX in editor'.
+	 * - Project Settings -> Plugins -> DMX -> 'Fixture Patches receive DMX in Editor' is set to true.
+	 */
+	UPROPERTY(BlueprintAssignable, Category = "DMX", Meta = (DisplayName = "On DMX Component Tick"))
+	FDMXOnDMXComponentTickSignature OnDMXComponentTick;
+
 	/** Gets the fixture patch used in the component */
 	UFUNCTION(BlueprintPure, Category = "DMX")
 	UDMXEntityFixturePatch* GetFixturePatch() const;
@@ -58,16 +66,28 @@ protected:
 	/** Sets up binding for receiving depending on the patch's and the component's properties */
 	void SetupReceiveDMXBinding();
 
+	/** Enables or disables the component tick depending of it being required for the OnDXMComponentTick function */
+	void UpdateTickEnabled();
+
 	/** If true, the component will receive DMX from the patch */
 	UPROPERTY(BlueprintReadOnly, EditAnywhere, AdvancedDisplay, Category = "DMX")
 	bool bReceiveDMXFromPatch;
 
 	// ~Begin UActorComponent interface
+	virtual void OnRegister() override;
 	virtual void BeginPlay() override;
+	virtual void TickComponent(float DeltaTime, enum ELevelTick TickType, FActorComponentTickFunction* ThisTickFunction) override;
 	virtual void DestroyComponent(bool bPromoteChildren) override;
 
 #if WITH_EDITOR
 	virtual void PostEditChangeProperty(FPropertyChangedEvent& PropertyChangedEvent) override;
-#endif // WITH_EDITOR
 	// ~End UActorComponent interface
+
+private:
+	/** Called when receive DMX in Editor was enabled or disabled for DMX Components (via DMXProtocolSettings) */
+	virtual void OnAllFixturePatchesReceiveDMXInEditorEnabled(bool bEnabled);
+
+	/** Called when the fixture patch changed. Can be editor only as we only want to update the component tick */
+	void OnFixturePatchPropertiesChanged(const UDMXEntityFixturePatch* FixturePatch);
+#endif // WITH_EDITOR
 };
diff --git a/Engine/Plugins/VirtualProduction/DMX/DMXEngine/Source/DMXRuntime/Public/Library/DMXEntityFixturePatch.h b/Engine/Plugins/VirtualProduction/DMX/DMXEngine/Source/DMXRuntime/Public/Library/DMXEntityFixturePatch.h
--- a/Engine/Plugins/VirtualProduction/DMX/DMXEngine/Source/DMXRuntime/Public/Library/DMXEntityFixturePatch.h
+++ b/Engine/Plugins/VirtualProduction/DMX/DMXEngine/Source/DMXRuntime/Public/Library/DMXEntityFixturePatch.h
@@ -255,7 +255,7 @@ public:
 
 	/** 
 	 * If true, the patch receives dmx and raises the OnFixturePatchReceivedDMX event in editor. 
-	 * NOTE: DMXComponent does not support this. To receive the patch in editor, drag out from the patch and bind its OnFixturePatchReceivedDMX event!
+	 * NOTE: If 'All Fixture Patches receive DMX in editor' is set to true in Project Settings -> Plugins -> DMX, this setting here is ignored.
 	 */
 	UPROPERTY(EditAnywhere, BlueprintReadOnly, Category = "Fixture Patch")
 	bool bReceiveDMXInEditor;
diff --git a/Engine/Plugins/VirtualProduction/DMX/DMXFixtures/Source/DMXFixtures/Private/DMXFixtureActor.cpp b/Engine/Plugins/VirtualProduction/DMX/DMXFixtures/Source/DMXFixtures/Private/DMXFixtureActor.cpp
--- a/Engine/Plugins/VirtualProduction/DMX/DMXFixtures/Source/DMXFixtures/Private/DMXFixtureActor.cpp
+++ b/Engine/Plugins/VirtualProduction/DMX/DMXFixtures/Source/DMXFixtures/Private/DMXFixtureActor.cpp
@@ -67,11 +67,14 @@ void ADMXFixtureActor::InitializeFixture(UStaticMeshComponent* StaticMeshLens, U
 {
 	GetComponents<UStaticMeshComponent>(StaticMeshComponents);
 
-	// Create dynamic materials
-	DynamicMaterialLens = UMaterialInstanceDynamic::Create(LensMaterialInstance, NULL);
-	DynamicMaterialBeam = UMaterialInstanceDynamic::Create(BeamMaterialInstance, NULL);
-	DynamicMaterialSpotLight = UMaterialInstanceDynamic::Create(SpotLightMaterialInstance, NULL);
-	DynamicMaterialPointLight = UMaterialInstanceDynamic::Create(PointLightMaterialInstance, NULL);
+	// Create dynamic materials if needed
+	if (!DynamicMaterialLens || !DynamicMaterialBeam || !DynamicMaterialSpotLight || !DynamicMaterialPointLight)
+	{
+		DynamicMaterialLens = UMaterialInstanceDynamic::Create(LensMaterialInstance, nullptr);
+		DynamicMaterialBeam = UMaterialInstanceDynamic::Create(BeamMaterialInstance, nullptr);
+		DynamicMaterialSpotLight = UMaterialInstanceDynamic::Create(SpotLightMaterialInstance, nullptr);
+		DynamicMaterialPointLight = UMaterialInstanceDynamic::Create(PointLightMaterialInstance, nullptr);
+	}
 
 	// Get lens width (support scaling)
 	if (StaticMeshLens)
diff --git a/Engine/Plugins/VirtualProduction/DMX/DMXProtocol/Source/DMXProtocol/Private/DMXProtocolSettings.cpp b/Engine/Plugins/VirtualProduction/DMX/DMXProtocol/Source/DMXProtocol/Private/DMXProtocolSettings.cpp
--- a/Engine/Plugins/VirtualProduction/DMX/DMXProtocol/Source/DMXProtocol/Private/DMXProtocolSettings.cpp
+++ b/Engine/Plugins/VirtualProduction/DMX/DMXProtocol/Source/DMXProtocol/Private/DMXProtocolSettings.cpp
@@ -156,6 +156,10 @@ void UDMXProtocolSettings::PostEditChangeProperty(FPropertyChangedEvent& Propert
 		bOverrideSendDMXEnabled = bDefaultSendDMXEnabled;
 		UDMXProtocolBlueprintLibrary::SetSendDMXEnabled(bDefaultSendDMXEnabled);
 	}
+	else if (PropertyName == GET_MEMBER_NAME_CHECKED(UDMXProtocolSettings, bAllFixturePatchesReceiveDMXInEditor))
+	{
+		OnAllFixturePatchesReceiveDMXInEditorEnabled.Broadcast(bAllFixturePatchesReceiveDMXInEditor);
+	}
 }
 #endif // WITH_EDITOR
 
@@ -252,14 +256,24 @@ void UDMXProtocolSettings::OverrideSendDMXEnabled(bool bEnabled)
 {
 	bOverrideSendDMXEnabled = bEnabled; 
 	
+	OnSetSendDMXEnabledDelegate.Broadcast(bEnabled);
+
+	// OnSetSendDMXEnabled is deprecated 5.1 and can be removed in a future release
+PRAGMA_DISABLE_DEPRECATION_WARNINGS
 	OnSetSendDMXEnabled.Broadcast(bEnabled);
+PRAGMA_ENABLE_DEPRECATION_WARNINGS
 }
 
 void UDMXProtocolSettings::OverrideReceiveDMXEnabled(bool bEnabled) 
 { 
 	bOverrideReceiveDMXEnabled = bEnabled; 
 
+	OnSetReceiveDMXEnabledDelegate.Broadcast(bEnabled);
+
+	// OnSetReceiveDMXEnabled is deprecated 5.1 and can be removed in a future release
+PRAGMA_DISABLE_DEPRECATION_WARNINGS
 	OnSetReceiveDMXEnabled.Broadcast(bEnabled);
+PRAGMA_ENABLE_DEPRECATION_WARNINGS
 }
 
 FString UDMXProtocolSettings::GetUniqueInputPortName() const
diff --git a/Engine/Plugins/VirtualProduction/DMX/DMXProtocol/Source/DMXProtocol/Private/IO/DMXInputPort.cpp b/Engine/Plugins/VirtualProduction/DMX/DMXProtocol/Source/DMXProtocol/Private/IO/DMXInputPort.cpp
--- a/Engine/Plugins/VirtualProduction/DMX/DMXProtocol/Source/DMXProtocol/Private/IO/DMXInputPort.cpp
+++ b/Engine/Plugins/VirtualProduction/DMX/DMXProtocol/Source/DMXProtocol/Private/IO/DMXInputPort.cpp
@@ -380,7 +380,7 @@ FDMXInputPortSharedRef FDMXInputPort::CreateFromConfig(FDMXInputPortConfig& Inpu
 	NewInputPort->bReceiveDMXEnabled = Settings->IsReceiveDMXEnabled();
 
 	// Bind to receive dmx changes
-	Settings->OnSetReceiveDMXEnabled.AddThreadSafeSP(NewInputPort, &FDMXInputPort::OnSetReceiveDMXEnabled);
+	Settings->GetOnSetReceiveDMXEnabled().AddThreadSafeSP(NewInputPort, &FDMXInputPort::OnSetReceiveDMXEnabled);
 
 	NewInputPort->UpdateFromConfig(InputPortConfig);
 
diff --git a/Engine/Plugins/VirtualProduction/DMX/DMXProtocol/Source/DMXProtocol/Private/IO/DMXOutputPort.cpp b/Engine/Plugins/VirtualProduction/DMX/DMXProtocol/Source/DMXProtocol/Private/IO/DMXOutputPort.cpp
--- a/Engine/Plugins/VirtualProduction/DMX/DMXProtocol/Source/DMXProtocol/Private/IO/DMXOutputPort.cpp
+++ b/Engine/Plugins/VirtualProduction/DMX/DMXProtocol/Source/DMXProtocol/Private/IO/DMXOutputPort.cpp
@@ -495,8 +495,8 @@ FDMXOutputPortSharedRef FDMXOutputPort::CreateFromConfig(FDMXOutputPortConfig& O
 	NewOutputPort->CommunicationDeterminator.SetSendEnabled(Settings->IsSendDMXEnabled());
 	NewOutputPort->CommunicationDeterminator.SetReceiveEnabled(Settings->IsReceiveDMXEnabled());
 
-	Settings->OnSetSendDMXEnabled.AddThreadSafeSP(NewOutputPort, &FDMXOutputPort::OnSetSendDMXEnabled);
-	Settings->OnSetReceiveDMXEnabled.AddThreadSafeSP(NewOutputPort, &FDMXOutputPort::OnSetReceiveDMXEnabled);
+	Settings->GetOnSetSendDMXEnabled().AddThreadSafeSP(NewOutputPort, &FDMXOutputPort::OnSetSendDMXEnabled);
+	Settings->GetOnSetReceiveDMXEnabled().AddThreadSafeSP(NewOutputPort, &FDMXOutputPort::OnSetReceiveDMXEnabled);
 
 	NewOutputPort->UpdateFromConfig(OutputPortConfig);
 
diff --git a/Engine/Plugins/VirtualProduction/DMX/DMXProtocol/Source/DMXProtocol/Public/DMXProtocolSettings.h b/Engine/Plugins/VirtualProduction/DMX/DMXProtocol/Source/DMXProtocol/Public/DMXProtocolSettings.h
--- a/Engine/Plugins/VirtualProduction/DMX/DMXProtocol/Source/DMXProtocol/Public/DMXProtocolSettings.h
+++ b/Engine/Plugins/VirtualProduction/DMX/DMXProtocol/Source/DMXProtocol/Public/DMXProtocolSettings.h
@@ -26,6 +26,7 @@ class DMXPROTOCOL_API UDMXProtocolSettings
 {
 	DECLARE_MULTICAST_DELEGATE_OneParam(FDMXOnSendDMXEnabled, bool /** bEnabled */);
 	DECLARE_MULTICAST_DELEGATE_OneParam(FDMXOnReceiveDMXEnabled, bool /** bEnabled */);
+	DECLARE_MULTICAST_DELEGATE_OneParam(FDMXOnAllFixturePatchesReceiveDMXInEditorEnabled, bool /** bEnabled */);
 
 public:
 	GENERATED_BODY()
@@ -80,12 +81,6 @@ public:
 	UPROPERTY(Config, EditAnywhere, Category = "DMX|Fixture Settings", Meta = (DisplayName = "Fixture Attributes"))
 	TSet<FDMXAttribute> Attributes;
 
-	/** Broadcast when send DMX is enabled or disabled */
-	FDMXOnSendDMXEnabled OnSetSendDMXEnabled;
-
-	/** Broadcast when receive DMX is enabled or disabled */
-	FDMXOnReceiveDMXEnabled OnSetReceiveDMXEnabled;
-
 	/** Returns whether send DMX is currently enabled, considering runtime override */
 	bool IsSendDMXEnabled() const { return bOverrideSendDMXEnabled; }
 
@@ -104,6 +99,48 @@ public:
 	/** Returns an automatic/unique name for an OutputPort */
 	FString GetUniqueOutputPortName() const;
 
+	/** Gets a Delegate Broadcast when send DMX is enabled or disabled */
+	FDMXOnSendDMXEnabled& GetOnSetSendDMXEnabled()
+	{
+		return OnSetSendDMXEnabledDelegate;
+	}
+
+	/** Gets a Delegate Broadcast when receive DMX is enabled or disabled */
+	FDMXOnSendDMXEnabled& GetOnSetReceiveDMXEnabled()
+	{
+		return OnSetReceiveDMXEnabledDelegate;
+	}
+
+#if WITH_EDITOR
+	/** Returns true if Fixture Patches should receive DMX in Editor */
+	bool ShouldAllFixturePatchesReceiveDMXInEditor() const { return bAllFixturePatchesReceiveDMXInEditor; }
+
+	/** Gets a Delegate Broadcast when receive DMX is enabled or disabled */
+	FDMXOnAllFixturePatchesReceiveDMXInEditorEnabled& GetOnAllFixturePatchesReceiveDMXInEditorEnabled()
+	{
+		return OnAllFixturePatchesReceiveDMXInEditorEnabled;
+	}
+#endif
+	
+	/** DEPRECATED 5.1 */
+	UE_DEPRECATED(5.1, "The OnSetSendDMXEnabled delegate is no longer directly accessible. Please use UDMXProtocolSettings::GetOnSetDMXEnabled() instead to access the delegate")
+	FDMXOnSendDMXEnabled OnSetSendDMXEnabled;
+
+	/** DEPRECATED 5.1 */
+	UE_DEPRECATED(5.1, "The OnSetReceiveDMXEnabled delegate is no longer directly accessible. Please use UDMXProtocolSettings::GetOnSetReceiveDMXEnabled() instead to access the delegate")
+	FDMXOnReceiveDMXEnabled OnSetReceiveDMXEnabled;
+
+	/** Returns the property name of the bDefaultSendDMXEnabled Property */
+	FORCEINLINE static FName GetDefaultSendDMXEnabledPropertyName() { return GET_MEMBER_NAME_CHECKED(UDMXProtocolSettings, bDefaultSendDMXEnabled); }
+
+	/** Returns the property name of the bDefaultReceiveDMXEnabled Property */
+	FORCEINLINE static FName GetDefaultReceiveDMXEnabledPropertyName() { return GET_MEMBER_NAME_CHECKED(UDMXProtocolSettings, bDefaultReceiveDMXEnabled); }
+
+#if WITH_EDITOR	
+	/** Returns the property name of the bDMXComponentRecievesDMXInEditor Property */
+	FORCEINLINE static FName GetAllFixturePatchesReceiveDMXInEditorPropertyName() { return GET_MEMBER_NAME_CHECKED(UDMXProtocolSettings, bAllFixturePatchesReceiveDMXInEditor); }
+#endif
+
 private:
 	/** Whether DMX is sent to the network. Recalled whenever editor or game starts.  */
 	UPROPERTY(Config, EditAnywhere, Category = "DMX|Communication Settings", Meta = (AllowPrivateAccess = true, DisplayName = "Send DMX by default"))
@@ -119,6 +156,21 @@ private:
 	/** Overrides the default bDefaultReceiveDMXEnabled value at runtime */
 	bool bOverrideReceiveDMXEnabled;
 
+	/** Broadcast when send DMX is enabled or disabled */
+	FDMXOnSendDMXEnabled OnSetSendDMXEnabledDelegate;
+
+	/** Broadcast when receive DMX is enabled or disabled */
+	FDMXOnReceiveDMXEnabled OnSetReceiveDMXEnabledDelegate;
+
+#if WITH_EDITORONLY_DATA
+	/** If true, all fixture patches receive DMX in Editor. This overrides the fixture patches 'Receive DMX In Editor' property. */
+	UPROPERTY(Config, EditAnywhere, Category = "DMX|Communication Settings", Meta = (DisplayName = "All Fixture Patches receive DMX in Editor"))
+	bool bAllFixturePatchesReceiveDMXInEditor = false;
+
+	/** Broadcast when Fixture Patches recieve DMX in Editor is enabled or disabled */
+	FDMXOnAllFixturePatchesReceiveDMXInEditorEnabled OnAllFixturePatchesReceiveDMXInEditorEnabled;
+#endif 
+
 	///////////////////
 	// DEPRECATED 4.27
 public:
diff --git a/Engine/Plugins/VirtualProduction/DMX/DMXProtocol/Source/DMXProtocolEditor/DMXProtocolEditor.Build.cs b/Engine/Plugins/VirtualProduction/DMX/DMXProtocol/Source/DMXProtocolEditor/DMXProtocolEditor.Build.cs
--- a/Engine/Plugins/VirtualProduction/DMX/DMXProtocol/Source/DMXProtocolEditor/DMXProtocolEditor.Build.cs
+++ b/Engine/Plugins/VirtualProduction/DMX/DMXProtocol/Source/DMXProtocolEditor/DMXProtocolEditor.Build.cs
@@ -26,6 +26,7 @@ public class DMXProtocolEditor : ModuleRules
 				"EditorStyle",
 				"Engine",
 				"InputCore",
+				"Projects",
 				"Slate",
 				"SlateCore",
 				"TimeManagement",
diff --git a/Engine/Plugins/VirtualProduction/DMX/DMXProtocol/Source/DMXProtocolEditor/Private/DetailsCustomizations/DMXProtocolSettingsDetails.cpp b/Engine/Plugins/VirtualProduction/DMX/DMXProtocol/Source/DMXProtocolEditor/Private/DetailsCustomizations/DMXProtocolSettingsDetails.cpp
new file mode 100644
--- /dev/null
+++ b/Engine/Plugins/VirtualProduction/DMX/DMXProtocol/Source/DMXProtocolEditor/Private/DetailsCustomizations/DMXProtocolSettingsDetails.cpp
@@ -0,0 +1,38 @@
+// Copyright Epic Games, Inc. All Rights Reserved.
+
+#include "DMXProtocolSettingsDetails.h"
+
+#include "DMXProtocolSettings.h"
+
+#include "DetailLayoutBuilder.h"
+#include "Interfaces/IPluginManager.h"
+
+
+namespace UE::DMXProtocolEditor::Private::DMXProtocolSettingsDetails
+{
+	bool IsDMXEnginePluginEnabled()
+	{
+		static const FString DMXEnginePluginName = TEXT("DMXEngine");
+
+		const TSharedPtr<IPlugin> Plugin = IPluginManager::Get().FindPlugin(DMXEnginePluginName);
+		return Plugin.IsValid() && Plugin->IsEnabled();
+	}
+}
+
+FDMXProtocolSettingsDetails::FDMXProtocolSettingsDetails()
+	: bIsDMXEnginePluginEnabled(UE::DMXProtocolEditor::Private::DMXProtocolSettingsDetails::IsDMXEnginePluginEnabled())
+{}
+
+TSharedRef<IDetailCustomization> FDMXProtocolSettingsDetails::MakeInstance()
+{
+	const TSharedRef<FDMXProtocolSettingsDetails> ProtocolSettingsDetails = MakeShared<FDMXProtocolSettingsDetails>();
+	return ProtocolSettingsDetails;
+}
+
+void FDMXProtocolSettingsDetails::CustomizeDetails(IDetailLayoutBuilder& DetailBuilder)
+{
+	if (!bIsDMXEnginePluginEnabled)
+	{
+		DetailBuilder.HideProperty(UDMXProtocolSettings::GetAllFixturePatchesReceiveDMXInEditorPropertyName());
+	}
+}
diff --git a/Engine/Plugins/VirtualProduction/DMX/DMXProtocol/Source/DMXProtocolEditor/Private/DetailsCustomizations/DMXProtocolSettingsDetails.h b/Engine/Plugins/VirtualProduction/DMX/DMXProtocol/Source/DMXProtocolEditor/Private/DetailsCustomizations/DMXProtocolSettingsDetails.h
new file mode 100644
--- /dev/null
+++ b/Engine/Plugins/VirtualProduction/DMX/DMXProtocol/Source/DMXProtocolEditor/Private/DetailsCustomizations/DMXProtocolSettingsDetails.h
@@ -0,0 +1,28 @@
+// Copyright Epic Games, Inc. All Rights Reserved.
+
+#pragma once
+
+#include "CoreMinimal.h"
+#include "IDetailCustomization.h"
+
+
+/** Details customization for the DMMXFixturePatch class */
+class FDMXProtocolSettingsDetails
+	: public IDetailCustomization
+{
+public:
+	/** Constructor */
+	FDMXProtocolSettingsDetails();
+
+	/** Creates a detail customization instance */
+	static TSharedRef<IDetailCustomization> MakeInstance();
+
+protected:
+	//~ Begin IDetailCustomization interface
+	virtual void CustomizeDetails(IDetailLayoutBuilder& DetailBuilder) override;
+	//~ End IDetailCustomization interface
+
+private:
+	/** True if the DMX Engine Plugin is enabled */
+	const bool bIsDMXEnginePluginEnabled;
+};
