DMX - Fix an issue where duplicated and copy pasted DMX entities wouldn't update their ID leaving them unusable

diff --git a/Engine/Plugins/VirtualProduction/DMX/DMXEngine/Source/DMXEditor/Private/Widgets/FixturePatch/SDMXFixturePatchTree.cpp b/Engine/Plugins/VirtualProduction/DMX/DMXEngine/Source/DMXEditor/Private/Widgets/FixturePatch/SDMXFixturePatchTree.cpp
--- a/Engine/Plugins/VirtualProduction/DMX/DMXEngine/Source/DMXEditor/Private/Widgets/FixturePatch/SDMXFixturePatchTree.cpp
+++ b/Engine/Plugins/VirtualProduction/DMX/DMXEngine/Source/DMXEditor/Private/Widgets/FixturePatch/SDMXFixturePatchTree.cpp
@@ -500,7 +500,7 @@ void SDMXFixturePatchTree::OnPasteNodes()
 	{
 		// Start transaction for Undo and take a snapshot of the current Library state
 		const FScopedTransaction PasteEntities(NewObjects.Num() > 1 ? LOCTEXT("PasteFixturePatches", "Paste Fixture Patches") : LOCTEXT("PasteFixturePatch", "Paste Fixture Patch"));
-		Library->Modify();
+		Library->PreEditChange(UDMXEntityFixtureType::StaticClass()->FindPropertyByName(UDMXLibrary::GetEntitiesPropertyName()));
 
 		// Add each pasted Entity to the Library
 		TArray<TWeakObjectPtr<UDMXEntityFixturePatch>> NewFixturePatches;
@@ -537,6 +537,8 @@ void SDMXFixturePatchTree::OnPasteNodes()
 			}
 		}
 
+		Library->PostEditChange();
+
 		FixturePatchSharedData->SelectFixturePatches(NewFixturePatches);
 
 		UpdateTree();
@@ -568,7 +570,7 @@ void SDMXFixturePatchTree::OnDuplicateNodes()
 	if (SelectedEntities.Num() > 0 && Library)
 	{
 		const FScopedTransaction Transaction(SelectedEntities.Num() > 1 ? LOCTEXT("DuplicateFixturePatches", "Duplicate Fixture Patch") : LOCTEXT("DuplicateFixturePatch", "Duplicate Fixture Patch"));
-		Library->Modify();
+		Library->PreEditChange(UDMXEntityFixtureType::StaticClass()->FindPropertyByName(UDMXLibrary::GetEntitiesPropertyName()));
 
 		TArray<TWeakObjectPtr<UDMXEntityFixturePatch>> NewFixturePatches;
 		NewFixturePatches.Reserve(SelectedEntities.Num());
@@ -581,6 +583,7 @@ void SDMXFixturePatchTree::OnDuplicateNodes()
 			
 			if (UDMXEntityFixturePatch* EntityCopy = CastChecked<UDMXEntityFixturePatch>(StaticDuplicateObjectEx(DuplicationParams)))
 			{
+				EntityCopy->RefreshID();
 				EntityCopy->SetName(FDMXRuntimeUtils::FindUniqueEntityName(Library, EntityCopy->GetClass(), EntityCopy->GetDisplayName()));
 				NewFixturePatches.Add(EntityCopy);
 
@@ -593,6 +596,8 @@ void SDMXFixturePatchTree::OnDuplicateNodes()
 			}
 		}
 
+		Library->PostEditChange();
+
 		FixturePatchSharedData->SelectFixturePatches(NewFixturePatches);
 
 		UpdateTree(); 
diff --git a/Engine/Plugins/VirtualProduction/DMX/DMXEngine/Source/DMXEditor/Private/Widgets/FixtureType/SDMXFixtureTypeTree.cpp b/Engine/Plugins/VirtualProduction/DMX/DMXEngine/Source/DMXEditor/Private/Widgets/FixtureType/SDMXFixtureTypeTree.cpp
--- a/Engine/Plugins/VirtualProduction/DMX/DMXEngine/Source/DMXEditor/Private/Widgets/FixtureType/SDMXFixtureTypeTree.cpp
+++ b/Engine/Plugins/VirtualProduction/DMX/DMXEngine/Source/DMXEditor/Private/Widgets/FixtureType/SDMXFixtureTypeTree.cpp
@@ -299,7 +299,7 @@ void SDMXFixtureTypeTree::OnPasteNodes()
 	{
 		// Start transaction for Undo and take a snapshot of the current Library state
 		const FScopedTransaction PasteEntities(NewObjects.Num() > 1 ? LOCTEXT("PasteFixtureTypes", "Paste Fixture Types") : LOCTEXT("PasteFixtureType", "Paste Fixture Type"));
-		Library->Modify();
+		Library->PreEditChange(UDMXLibrary::StaticClass()->FindPropertyByName(UDMXLibrary::GetEntitiesPropertyName()));
 
 		// Add each pasted Entity to the Library
 		TArray<TWeakObjectPtr<UDMXEntityFixtureType>> NewFixtureTypes;
@@ -320,6 +320,8 @@ void SDMXFixtureTypeTree::OnPasteNodes()
 			NewFixtureTypes.Add(CastChecked<UDMXEntityFixtureType>(NewEntity));
 		}
 
+		Library->PostEditChange();
+
 		FixtureTypeSharedData->SelectFixtureTypes(NewFixtureTypes);
 
 		UpdateTree();
@@ -344,7 +346,7 @@ void SDMXFixtureTypeTree::OnDuplicateNodes()
 	if (Library && SelectedEntities.Num() > 0)
 	{
 		const FScopedTransaction Transaction(SelectedEntities.Num() > 1 ? LOCTEXT("DuplicateFixtureTypes", "Duplicate Fixture Types") : LOCTEXT("DuplicateFixtureType", "Duplicate Fixture Type"));
-		Library->Modify();
+		Library->PreEditChange(UDMXLibrary::StaticClass()->FindPropertyByName(UDMXLibrary::GetEntitiesPropertyName()));
 
 		TArray<TWeakObjectPtr<UDMXEntityFixtureType>> NewFixtureTypes;
 		NewFixtureTypes.Reserve(SelectedEntities.Num());
@@ -364,6 +366,8 @@ void SDMXFixtureTypeTree::OnDuplicateNodes()
 			}
 		}
 
+		Library->PostEditChange();
+
 		FixtureTypeSharedData->SelectFixtureTypes(NewFixtureTypes);
 
 		UpdateTree();
diff --git a/Engine/Plugins/VirtualProduction/DMX/DMXEngine/Source/DMXRuntime/Private/Library/DMXEntityFixturePatch.cpp b/Engine/Plugins/VirtualProduction/DMX/DMXEngine/Source/DMXRuntime/Private/Library/DMXEntityFixturePatch.cpp
--- a/Engine/Plugins/VirtualProduction/DMX/DMXEngine/Source/DMXRuntime/Private/Library/DMXEntityFixturePatch.cpp
+++ b/Engine/Plugins/VirtualProduction/DMX/DMXEngine/Source/DMXRuntime/Private/Library/DMXEntityFixturePatch.cpp
@@ -64,14 +64,14 @@ UDMXEntityFixturePatch* UDMXEntityFixturePatch::CreateFixturePatchInLibrary(FDMX
 			NewFixturePatch->SetStartingChannel(ConstructionParams.StartingAddress);
 			NewFixturePatch->SetActiveModeIndex(ConstructionParams.ActiveMode);
 
-			return NewFixturePatch;
-
 #if WITH_EDITOR
 			if (bMarkDMXLibraryDirty)
 			{
 				DMXLibrary->PostEditChange();
 			}
 #endif
+			OnFixturePatchChangedDelegate.Broadcast(NewFixturePatch);
+			return NewFixturePatch;
 		}
 	}
 
