DMX - Create DMX Libraries from MVR files, spawn DMX Libraries by finding the most appropriate actors for their MVR content

diff --git a/Engine/Plugins/VirtualProduction/DMX/DMXEngine/DMXEngine.uplugin b/Engine/Plugins/VirtualProduction/DMX/DMXEngine/DMXEngine.uplugin
--- a/Engine/Plugins/VirtualProduction/DMX/DMXEngine/DMXEngine.uplugin
+++ b/Engine/Plugins/VirtualProduction/DMX/DMXEngine/DMXEngine.uplugin
@@ -14,8 +14,8 @@
 	"IsBetaVersion": false,
 	"Modules": [
 		{
-			"Name": "DMXRuntime",
-			"Type": "Runtime",
+			"Name": "DMXBlueprintGraph",
+			"Type": "UncookedOnly",
 			"LoadingPhase": "Default"
 		},
 		{
@@ -24,12 +24,16 @@
 			"LoadingPhase": "Default"
 		},
 		{
-			"Name": "DMXBlueprintGraph",
-			"Type": "UncookedOnly",
+			"Name": "DMXRuntime",
+			"Type": "Runtime",
 			"LoadingPhase": "Default"
 		}
 	],
 	"Plugins": [
+		{
+			"Name": "DatasmithContent",
+			"Enabled": true
+		},
 		{
 			"Name": "DMXProtocol",
 			"Enabled": true
diff --git a/Engine/Plugins/VirtualProduction/DMX/DMXEngine/Source/DMXEditor/DMXEditor.Build.cs b/Engine/Plugins/VirtualProduction/DMX/DMXEngine/Source/DMXEditor/DMXEditor.Build.cs
--- a/Engine/Plugins/VirtualProduction/DMX/DMXEngine/Source/DMXEditor/DMXEditor.Build.cs
+++ b/Engine/Plugins/VirtualProduction/DMX/DMXEngine/Source/DMXEditor/DMXEditor.Build.cs
@@ -25,28 +25,28 @@ public class DMXEditor : ModuleRules
 				"ApplicationCore",
 				"AssetRegistry",
 				"AssetTools",
+				"ContentBrowser",
 				"CoreUObject",
-				"Kismet",
 				"EditorFramework",
-				"UnrealEd",
 				"EditorStyle",
-				"PropertyEditor",
-				"KismetWidgets",
 				"Engine",
-				"Slate",
-				"SlateCore",
 				"InputCore",
 				"Json",
 				"JsonUtilities",
+				"Kismet",
+				"KismetWidgets",
 				"Projects",
 				"MainFrame",
-				"XmlParser",
-				"Sequencer",
 				"MovieScene",
+				"PropertyEditor",
+				"Sequencer",
+				"Slate",
+				"SlateCore",
 				"TakesCore",
 				"TakeRecorder",
 				"TakeTrackRecorders",
-				"ContentBrowser"
+				"UnrealEd",
+				"XmlParser",
 			}
 		);
 	}
diff --git a/Engine/Plugins/VirtualProduction/DMX/DMXEngine/Source/DMXEditor/Private/Customizations/DMXEntityFixtureTypeDetails.cpp b/Engine/Plugins/VirtualProduction/DMX/DMXEngine/Source/DMXEditor/Private/Customizations/DMXEntityFixtureTypeDetails.cpp
--- a/Engine/Plugins/VirtualProduction/DMX/DMXEngine/Source/DMXEditor/Private/Customizations/DMXEntityFixtureTypeDetails.cpp
+++ b/Engine/Plugins/VirtualProduction/DMX/DMXEngine/Source/DMXEditor/Private/Customizations/DMXEntityFixtureTypeDetails.cpp
@@ -2,8 +2,8 @@
 
 #include "Customizations/DMXEntityFixtureTypeDetails.h"
 
+#include "Library/DMXImportGDTF.h"
 #include "Library/DMXEntityFixtureType.h"
-#include "Library/DMXImport.h"
 
 #include "CoreMinimal.h"
 #include "DetailLayoutBuilder.h"
@@ -33,7 +33,7 @@ void FDMXEntityFixtureTypeDetails::OnDMXImportChanged()
 	UObject* DMXImportObject = nullptr;
 	if (DMXImportHandle->GetValue(DMXImportObject) == FPropertyAccess::Success)
 	{
-		if (UDMXImport* DMXImport = Cast<UDMXImport>(DMXImportObject))
+		if (UDMXImportGDTF* DMXImportGDTF = Cast<UDMXImportGDTF>(DMXImportObject))
 		{
 			const TArray<TWeakObjectPtr<UObject>>& SelectedObjects = PropertyUtilities->GetSelectedObjects();
 
@@ -41,7 +41,7 @@ void FDMXEntityFixtureTypeDetails::OnDMXImportChanged()
 			{
 				if (UDMXEntityFixtureType* FixtureType = Cast<UDMXEntityFixtureType>(WeakFixtureTypeObject.Get()))
 				{
-					FixtureType->SetModesFromDMXImport(DMXImport);
+					FixtureType->SetGDTF(DMXImportGDTF);
 				}
 			}
 		}
diff --git a/Engine/Plugins/VirtualProduction/DMX/DMXEngine/Source/DMXEditor/Private/Customizations/DMXLibraryDetails.cpp b/Engine/Plugins/VirtualProduction/DMX/DMXEngine/Source/DMXEditor/Private/Customizations/DMXLibraryDetails.cpp
--- a/Engine/Plugins/VirtualProduction/DMX/DMXEngine/Source/DMXEditor/Private/Customizations/DMXLibraryDetails.cpp
+++ b/Engine/Plugins/VirtualProduction/DMX/DMXEngine/Source/DMXEditor/Private/Customizations/DMXLibraryDetails.cpp
@@ -29,8 +29,7 @@ void FDMXLibraryDetails::CustomizeDetails(IDetailLayoutBuilder& DetailBuilder)
 
 	IDetailCategoryBuilder& DetailCategoryBuilder = DetailBuilder.EditCategory("DMX Ports");
 
-	const TSharedPtr<IPropertyHandle> PortReferencesHandle = DetailBuilder.GetProperty(UDMXLibrary::GetPortReferencesPropertyName());
-	
+	// Add a button to open DMX Project Settings
 	DetailCategoryBuilder.AddCustomRow(LOCTEXT("OpenSettingsFilterString", "Settings"))
 		.WholeRowContent()
 		.HAlign(HAlign_Left)
@@ -66,7 +65,12 @@ void FDMXLibraryDetails::CustomizeDetails(IDetailLayoutBuilder& DetailBuilder)
 			]
 		];
 
-	DetailCategoryBuilder.AddProperty(PortReferencesHandle.ToSharedRef());
+
+	// Add the Port References
+	const TSharedPtr<IPropertyHandle> PortReferencesHandle = DetailBuilder.GetProperty(UDMXLibrary::GetPortReferencesPropertyName());
+	DetailCategoryBuilder
+		.AddProperty(PortReferencesHandle.ToSharedRef())
+		.ShouldAutoExpand(true);
 }
- 
+
 #undef LOCTEXT_NAMESPACE
diff --git a/Engine/Plugins/VirtualProduction/DMX/DMXEngine/Source/DMXEditor/Private/DMXEditorUtils.cpp b/Engine/Plugins/VirtualProduction/DMX/DMXEngine/Source/DMXEditor/Private/DMXEditorUtils.cpp
--- a/Engine/Plugins/VirtualProduction/DMX/DMXEngine/Source/DMXEditor/Private/DMXEditorUtils.cpp
+++ b/Engine/Plugins/VirtualProduction/DMX/DMXEngine/Source/DMXEditor/Private/DMXEditorUtils.cpp
@@ -5,7 +5,6 @@
 #include "DMXEditorLog.h"
 #include "DMXRuntimeUtils.h"
 #include "DMXSubsystem.h"
-#include "Interfaces/IDMXProtocol.h"
 #include "IO/DMXInputPort.h"
 #include "IO/DMXOutputPort.h"
 #include "IO/DMXPortManager.h"
@@ -15,6 +14,7 @@
 
 #include "AssetRegistryModule.h"
 #include "Factories.h"
+#include "PackageTools.h"
 #include "ScopedTransaction.h"
 #include "UnrealExporter.h"
 #include "Dialogs/Dialogs.h"
@@ -22,6 +22,7 @@
 #include "HAL/PlatformApplicationMisc.h"
 #include "UObject/Package.h"
 
+
 #define LOCTEXT_NAMESPACE "FDMXEditorUtils"
 // In blueprints name verification, it is said that '.' is known for causing problems
 #define DMX_INVALID_NAME_CHARACTERS TEXT(".")
@@ -994,4 +995,35 @@ void FDMXEditorUtils::ClearFixturePatchCachedData()
 	}
 }
 
+UPackage* FDMXEditorUtils::GetOrCreatePackage(TWeakObjectPtr<UObject> Parent, const FString& DesiredName)
+{
+	UPackage* Package = nullptr;
+	FString NewPackageName;
+
+	if (Parent.IsValid() && Parent->IsA(UPackage::StaticClass()))
+	{
+		Package = StaticCast<UPackage*>(Parent.Get());
+	}
+
+	if (!Package)
+	{
+		if (Parent.IsValid() && Parent->GetOutermost())
+		{
+			NewPackageName = FPackageName::GetLongPackagePath(Parent->GetOutermost()->GetName()) + "/" + DesiredName;
+		}
+		else
+		{
+			return nullptr;
+		}
+
+		NewPackageName = UPackageTools::SanitizePackageName(NewPackageName);
+		Package = CreatePackage(*NewPackageName);
+		Package->FullyLoad();
+	}
+
+	return Package;
+}
+
+#undef DMX_INVALID_NAME_CHARACTERS
+
 #undef LOCTEXT_NAMESPACE
diff --git a/Engine/Plugins/VirtualProduction/DMX/DMXEngine/Source/DMXEditor/Private/Factories/DMXEditorFactoryNew.cpp b/Engine/Plugins/VirtualProduction/DMX/DMXEngine/Source/DMXEditor/Private/Factories/DMXEditorFactoryNew.cpp
--- a/Engine/Plugins/VirtualProduction/DMX/DMXEngine/Source/DMXEditor/Private/Factories/DMXEditorFactoryNew.cpp
+++ b/Engine/Plugins/VirtualProduction/DMX/DMXEngine/Source/DMXEditor/Private/Factories/DMXEditorFactoryNew.cpp
@@ -22,11 +22,11 @@
 
 #define LOCTEXT_NAMESPACE "UDMXEditorFactoryNew"
 
+
+// The whole class is deprecated 5.1, but we keep the implementation for the time being
+PRAGMA_DISABLE_DEPRECATION_WARNINGS
 UDMXEditorFactoryNew::UDMXEditorFactoryNew()
 {
-	SupportedClass = UDMXLibrary::StaticClass();
-	bCreateNew = true;
-	bEditAfterNew = true;
 }
 
 UObject * UDMXEditorFactoryNew::FactoryCreateNew(UClass * Class, UObject * InParent, FName Name, EObjectFlags Flags, UObject * Context, FFeedbackContext * Warn)
@@ -45,14 +45,10 @@ UObject * UDMXEditorFactoryNew::FactoryCreateNew(UClass * Class, UObject * InPar
 	return DMXLibrary;
 }
 
-uint32 UDMXEditorFactoryNew::GetMenuCategories() const
-{
-	return (uint32)FDMXEditorModule::GetAssetCategory();
-}
-
 UDMXLibrary * UDMXEditorFactoryNew::MakeNewEditor(UObject * InParent, FName Name, EObjectFlags Flags)
 {
 	return NewObject<UDMXLibrary>(InParent, Name, Flags);
 }
+PRAGMA_ENABLE_DEPRECATION_WARNINGS
 
-#undef LOCTEXT_NAMESPACE
\ No newline at end of file
+#undef LOCTEXT_NAMESPACE
diff --git a/Engine/Plugins/VirtualProduction/DMX/DMXEngine/Source/DMXEditor/Private/Factories/DMXGDTFFactory.cpp b/Engine/Plugins/VirtualProduction/DMX/DMXEngine/Source/DMXEditor/Private/Factories/DMXGDTFFactory.cpp
--- a/Engine/Plugins/VirtualProduction/DMX/DMXEngine/Source/DMXEditor/Private/Factories/DMXGDTFFactory.cpp
+++ b/Engine/Plugins/VirtualProduction/DMX/DMXEngine/Source/DMXEditor/Private/Factories/DMXGDTFFactory.cpp
@@ -115,7 +115,7 @@ UObject* UDMXGDTFFactory::FactoryCreateFile(UClass* InClass, UObject* InParent,
     ImportArgs.ImportUI = ImportUI;
     ImportArgs.Name = InName;
     ImportArgs.Parent = InParent;
-    ImportArgs.CurrentFilename = CurrentFilename;
+    ImportArgs.CurrentFilename = InFilename;
     ImportArgs.Flags = Flags;
     ImportArgs.bCancelOperation = bOperationCanceled;
     TUniquePtr<FDMXGDTFImporter> Importer = MakeUnique<FDMXGDTFImporter>(ImportArgs);
@@ -159,7 +159,7 @@ UObject* UDMXGDTFFactory::FactoryCreateFile(UClass* InClass, UObject* InParent,
 
         if (CreatedObject != nullptr)
         {
-            CreatedObject->SourceFilename = GetCurrentFilename();
+            CreatedObject->SourceFilename = InFilename;
         }
 		GEditor->GetEditorSubsystem<UImportSubsystem>()->BroadcastAssetPostImport(this, CreatedObject);
     }
diff --git a/Engine/Plugins/VirtualProduction/DMX/DMXEngine/Source/DMXEditor/Private/Factories/DMXGDTFImporter.cpp b/Engine/Plugins/VirtualProduction/DMX/DMXEngine/Source/DMXEditor/Private/Factories/DMXGDTFImporter.cpp
--- a/Engine/Plugins/VirtualProduction/DMX/DMXEngine/Source/DMXEditor/Private/Factories/DMXGDTFImporter.cpp
+++ b/Engine/Plugins/VirtualProduction/DMX/DMXEngine/Source/DMXEditor/Private/Factories/DMXGDTFImporter.cpp
@@ -27,6 +27,49 @@
 
 #define LOCTEXT_NAMESPACE "DMXGDTFImporter"
 
+
+namespace UE::DMX::DMXGDTFImporter::Private
+{
+	/** 
+	 * Depending on the GDTF Spec, some XML Attribute Names might start with or without DMX. 
+	 * 
+	 * @param ParentNode		The Parent Node to search in 
+	 * @param NodeName			The Attribute Name with DMX tag. This is the correct form, e.g. 'DMXMode', not 'Mode', 'DMXChannels', not 'Channels'.
+	 * @return					Returns a pointer to the child node, or nullptr if the child cannot be found 
+	 */
+	const FXmlNode* FindChildNodeEvenIfDMXSubstringIsMissing(const FXmlNode& ParentNode, const FString& AttributeNameWithDMXTag)
+	{
+		if (const FXmlNode* ChildNodePtrWithSubstring = ParentNode.FindChildNode(AttributeNameWithDMXTag))
+		{
+			return ChildNodePtrWithSubstring;
+		}
+		else if (const FXmlNode* ChildNodePtrWithoutSubstring = ParentNode.FindChildNode(AttributeNameWithDMXTag.RightChop(3)))
+		{
+			return ChildNodePtrWithoutSubstring;
+		}
+
+		return nullptr;
+	}
+
+	/** 
+	 * Depending on the GDTF Spec, some XML Attribute Names might start with or without DMX. 
+	 * 
+	 * @param ParentNode		The Parent Node to search in 
+	 * @param NodeName			The Attribute Name with DMX tag. This is the correct form, e.g. 'DMXMode', not 'Mode', 'DMXChannels', not 'Channels'.
+	 * @return					Returns the Attribute as String
+	 */
+	FString FindAttributeEvenIfDMXSubstringIsMissing(const FXmlNode& ParentNode, const FString& AttributeNameWithDMXTag)
+	{
+		FString Attribute = ParentNode.GetAttribute(AttributeNameWithDMXTag);
+		if (Attribute.IsEmpty())
+		{
+			Attribute = ParentNode.GetAttribute(AttributeNameWithDMXTag.RightChop(3));
+		}
+
+		return Attribute;
+	}
+};
+
 FDMXGDTFImporter::FDMXGDTFImporter(const FDMXGDTFImportArgs& InImportArgs)
     : ImportArgs(InImportArgs)
     , bIsXMLParsedSuccessfully(false)
@@ -393,8 +436,13 @@ UDMXImportGDTF* FDMXGDTFImporter::CreateGDTFDesctription()
                             {
                                 FDMXImportGDTFBreak ImportBreak;
 
-                                LexTryParseString(ImportBreak.DMXOffset, *BreakNode->GetAttribute("DMXOffset"));
-                                LexTryParseString(ImportBreak.DMXBreak, *BreakNode->GetAttribute("DMXBreak"));
+								// Depending on the GDTF spec the Offset node may be called either 'DMXOffset' or 'Offset'
+								const FString DMXOffsetAttribute = UE::DMX::DMXGDTFImporter::Private::FindAttributeEvenIfDMXSubstringIsMissing(*BreakNode, TEXT("DMXOffset"));
+								LexTryParseString(ImportBreak.DMXOffset, *DMXOffsetAttribute);
+
+								// Depending on the GDTF spec the Break node may be called either 'DMXBreak' or 'Break'
+								const FString DMXBreakAttribute = UE::DMX::DMXGDTFImporter::Private::FindAttributeEvenIfDMXSubstringIsMissing(*BreakNode, TEXT("DMXBreak"));
+								LexTryParseString(ImportBreak.DMXBreak, *DMXBreakAttribute);
 
                                 ImportGeometryReference.Breaks.Add(ImportBreak);
                             }
@@ -406,17 +454,21 @@ UDMXImportGDTF* FDMXGDTFImporter::CreateGDTFDesctription()
                     }
                 }
 
-                if (const FXmlNode* DMXModesNode = FixtureTypeNode->FindChildNode("DMXModes"))
+				// Depending on the GDTF spec the Break node may be called either 'DMXModes' or 'Modes'.
+				const FXmlNode* const ModesNode = UE::DMX::DMXGDTFImporter::Private::FindChildNodeEvenIfDMXSubstringIsMissing(*FixtureTypeNode, TEXT("DMXModes"));
+                if (ModesNode)
                 {
-                    for (const FXmlNode* DMXModeNode : DMXModesNode->GetChildrenNodes())
+                    for (const FXmlNode* DMXModeNode : ModesNode->GetChildrenNodes())
                     {
                         FDMXImportGDTFDMXMode DMXImportGDTFDMXMode;
                         DMXImportGDTFDMXMode.Name = FName(*DMXModeNode->GetAttribute("Name"));
                         DMXImportGDTFDMXMode.Geometry = FName(*DMXModeNode->GetAttribute("Geometry"));
 
-                        if (const FXmlNode* DMXChannelsNode = DMXModeNode->FindChildNode("DMXChannels"))
+						// Depending on the GDTF spec in use modes may be stored in the "Modes" or "DMXModes" node.
+						const FXmlNode* const ChannelsNode = UE::DMX::DMXGDTFImporter::Private::FindChildNodeEvenIfDMXSubstringIsMissing(*DMXModeNode, TEXT("DMXChannels"));
+                        if (ChannelsNode)
                         {
-                            for (const FXmlNode* DMXChannelNode : DMXChannelsNode->GetChildrenNodes())
+                            for (const FXmlNode* DMXChannelNode : ChannelsNode->GetChildrenNodes())
                             {
 								FDMXImportGDTFDMXChannel ImportDMXChannel;
 
diff --git a/Engine/Plugins/VirtualProduction/DMX/DMXEngine/Source/DMXEditor/Private/Factories/DMXLibraryFactory.cpp b/Engine/Plugins/VirtualProduction/DMX/DMXEngine/Source/DMXEditor/Private/Factories/DMXLibraryFactory.cpp
new file mode 100644
--- /dev/null
+++ b/Engine/Plugins/VirtualProduction/DMX/DMXEngine/Source/DMXEditor/Private/Factories/DMXLibraryFactory.cpp
@@ -0,0 +1,27 @@
+// Copyright Epic Games, Inc. All Rights Reserved.
+
+#include "Factories/DMXLibraryFactory.h"
+
+#include "DMXEditorLog.h"
+#include "DMXEditorModule.h"
+#include "Library/DMXLibrary.h"
+
+
+UDMXLibraryFactory::UDMXLibraryFactory()
+{
+	SupportedClass = UDMXLibrary::StaticClass();
+	bCreateNew = true;
+	bEditAfterNew = true;
+}
+
+UObject* UDMXLibraryFactory::FactoryCreateNew(UClass* Class, UObject* InParent, FName Name, EObjectFlags Flags, UObject* Context, FFeedbackContext* Warn)
+{
+	UDMXLibrary* DMXLibrary = NewObject<UDMXLibrary>(InParent, Name, Flags);
+
+	return DMXLibrary;
+}
+
+uint32 UDMXLibraryFactory::GetMenuCategories() const
+{
+	return (uint32)FDMXEditorModule::GetAssetCategory();
+}
diff --git a/Engine/Plugins/VirtualProduction/DMX/DMXEngine/Source/DMXEditor/Private/Factories/DMXLibraryFromMVRFactory.cpp b/Engine/Plugins/VirtualProduction/DMX/DMXEngine/Source/DMXEditor/Private/Factories/DMXLibraryFromMVRFactory.cpp
new file mode 100644
--- /dev/null
+++ b/Engine/Plugins/VirtualProduction/DMX/DMXEngine/Source/DMXEditor/Private/Factories/DMXLibraryFromMVRFactory.cpp
@@ -0,0 +1,253 @@
+// Copyright Epic Games, Inc. All Rights Reserved.
+
+#include "Factories/DMXLibraryFromMVRFactory.h"
+
+#include "DMXEditorLog.h"
+#include "DMXEditorModule.h"
+#include "DMXEditorUtils.h"
+#include "Factories/DMXGDTFFactory.h"
+#include "Library/DMXEntityFixturePatch.h"
+#include "Library/DMXEntityFixtureType.h"
+#include "Library/DMXImportGDTF.h"
+#include "Library/DMXLibrary.h"
+#include "MVR/DMXMVRGeneralSceneDescription.h"
+#include "MVR/DMXMVRUnzip.h"
+
+#include "AssetRegistryModule.h"
+#include "AssetToolsModule.h"
+#include "ObjectTools.h"
+#include "XmlFile.h"
+#include "Misc/FileHelper.h"
+#include "Misc/MessageDialog.h"
+
+
+#define LOCTEXT_NAMESPACE "DMXLibraryFromMVRFactory"
+
+const FName UDMXLibraryFromMVRFactory::MVRFileExtension = "MVR";
+const FName UDMXLibraryFromMVRFactory::GDTFFileExtension = "GDTF";
+
+UDMXLibraryFromMVRFactory::UDMXLibraryFromMVRFactory()
+{
+	bEditorImport = true;
+	bEditAfterNew = true;
+	SupportedClass = UDMXLibrary::StaticClass();
+
+	Formats.Add(TEXT("mvr;My Virtual Rig"));
+}
+
+UObject* UDMXLibraryFromMVRFactory::FactoryCreateFile(UClass* InClass, UObject* Parent, FName InName, EObjectFlags Flags, const FString& InFilename, const TCHAR* Parms, FFeedbackContext* Warn, bool& bOutOperationCanceled)
+{
+	bOutOperationCanceled = false;
+	CurrentFilename = InFilename;
+
+	UDMXLibrary* NewDMXLibrary = CreateDMXLibraryAsset(Parent, Flags, InFilename);
+	if (!NewDMXLibrary)
+	{
+		UE_LOG(LogDMXEditor, Error, TEXT("Failed to create DMX Library for MVR '%s'."), *InFilename);
+		return nullptr;
+	}
+
+	const TSharedPtr<FDMXMVRUnzip> MVRUnzip = FDMXMVRUnzip::CreateFromFile(InFilename);
+	if (!MVRUnzip.IsValid())
+	{
+		UE_LOG(LogDMXEditor, Error, TEXT("Cannot read MVR '%s'. File is not a valid MVR."), *InFilename);
+		return nullptr;
+	}
+
+	TArray64<uint8> XMLData;
+	MVRUnzip->GetFileContent(TEXT("GeneralSceneDescription.xml"), XMLData);
+
+	// MVR implicitly adpots UTF-8 encoding of Xml Files by adopting the GDT standard (DIN-15800).
+	// Content is NOT null-terminated; we need to specify lengths here.
+	const FUTF8ToTCHAR TCHARData(reinterpret_cast<const ANSICHAR*>(XMLData.GetData()), XMLData.Num());
+
+	const TSharedRef<FXmlFile> GeneralSceneDescriptionXml = MakeShared<FXmlFile>(FString(TCHARData.Length(), TCHARData.Get()), EConstructMethod::ConstructFromBuffer);
+	FDMXMVRGeneralSceneDescription GeneralSceneDescription(GeneralSceneDescriptionXml);
+
+	TArray<UDMXImportGDTF*> GDTFs = CreateGDTFAssets(Parent, Flags, MVRUnzip.ToSharedRef(), GeneralSceneDescription);
+
+	InitDMXLibrary(NewDMXLibrary, GDTFs, GeneralSceneDescription);
+
+	return NewDMXLibrary;
+}
+
+bool UDMXLibraryFromMVRFactory::FactoryCanImport(const FString& Filename)
+{
+	const FString TargetExtension = FPaths::GetExtension(Filename);
+
+	if (TargetExtension.Equals(UDMXLibraryFromMVRFactory::MVRFileExtension.ToString(), ESearchCase::IgnoreCase))
+	{
+		return true;
+	}
+
+	return false;
+}
+
+UDMXLibrary* UDMXLibraryFromMVRFactory::CreateDMXLibraryAsset(UObject* Parent, EObjectFlags Flags, const FString& InFilename)
+{
+	constexpr bool bRemovePathFromDesiredName = true;
+	const FString BaseFileName = FPaths::GetBaseFilename(InFilename, bRemovePathFromDesiredName);
+	const FString PackageName = Parent->GetName() + TEXT("_MVR/");
+	const FString AssetName = ObjectTools::SanitizeObjectName(BaseFileName);
+
+	FString UniquePackageName;
+	FString UniqueAssetName;
+	FAssetToolsModule& AssetToolsModule = FModuleManager::LoadModuleChecked<FAssetToolsModule>("AssetTools");
+	AssetToolsModule.Get().CreateUniqueAssetName(PackageName / AssetName, TEXT(""), UniquePackageName, UniqueAssetName);
+
+	UPackage* Package = CreatePackage(*UniquePackageName);
+	check(Package);
+	Package->FullyLoad();
+
+	UImportSubsystem* ImportSubsystem = GEditor->GetEditorSubsystem<UImportSubsystem>();
+	ImportSubsystem->BroadcastAssetPreImport(this, UDMXLibrary::StaticClass(), Parent, *UniqueAssetName, *MVRFileExtension.ToString());
+
+	UDMXLibrary* NewDMXLibrary = NewObject<UDMXLibrary>(Package, *UniqueAssetName, Flags | RF_Public);
+	if (!NewDMXLibrary)
+	{
+		return nullptr;
+	}
+
+	check(sizeof(this) == sizeof(int64));
+	ImportSubsystem->BroadcastAssetPostImport(this, NewDMXLibrary);
+	
+	return NewDMXLibrary;
+}
+
+TArray<UDMXImportGDTF*> UDMXLibraryFromMVRFactory::CreateGDTFAssets(UObject* Parent, EObjectFlags Flags, const TSharedRef<FDMXMVRUnzip>& MVRUnzip, const FDMXMVRGeneralSceneDescription& GeneralSceneDescription)
+{
+	TArray<UDMXImportGDTF*> ImportedGDTFs;
+	TArray<FString> ImportedGDTFNames;
+
+	constexpr bool bRemovePathFromDesiredName = true;
+	const FString PackageName = Parent->GetName() + FString(TEXT("_MVR")) / FString(TEXT("GDTFs"));
+
+	TArray<FAssetData> ExistingGDTFAssets;
+	FAssetRegistryModule& AssetRegistryModule = FModuleManager::LoadModuleChecked<FAssetRegistryModule>(TEXT("AssetRegistry"));
+	AssetRegistryModule.Get().GetAssetsByPath(FName(PackageName), ExistingGDTFAssets);
+
+	TArray<UDMXImportGDTF*> ExistingGDTFs;
+	for (const FAssetData& AssetData : ExistingGDTFAssets)
+	{
+		if (UDMXImportGDTF* GDTFAsset = Cast<UDMXImportGDTF>(AssetData.GetAsset()))
+		{
+			GeneralSceneDescription.MVRFixtures.ContainsByPredicate([GDTFAsset](const FDMXMVRFixture& MVRFixture)
+				{
+					return GDTFAsset->SourceFilename.Equals(MVRFixture.GDTFSpec, ESearchCase::IgnoreCase);
+				});
+
+			ExistingGDTFs.AddUnique(GDTFAsset);
+		}
+	}
+
+	UDMXGDTFFactory* GDTFFactory = NewObject<UDMXGDTFFactory>();
+	if (ExistingGDTFs.Num() > 0)
+	{
+		const FText MessageText = LOCTEXT("MVRImportReimportsGDTFDialog", "MVR contains existing GDTFs. Do you want to reimport the existing GDTF assets?");
+		if (FMessageDialog::Open(EAppMsgType::YesNo, MessageText) == EAppReturnType::Yes)
+		{
+			for (UDMXImportGDTF* GDTFAsset : ExistingGDTFs)
+			{
+				if (GDTFFactory->Reimport(GDTFAsset) == EReimportResult::Succeeded)
+				{
+					ImportedGDTFs.Add(GDTFAsset);
+					ImportedGDTFNames.Add(GDTFAsset->SourceFilename);
+				}
+			}
+		}
+	}
+
+	// Import GDTF Assets that aren't yet imported in the reimport procedure
+	for (const FDMXMVRFixture& MVRFixture : GeneralSceneDescription.MVRFixtures)
+	{
+		// Don't import the same GDTF twice
+		if (!ImportedGDTFNames.Contains(MVRFixture.GDTFSpec))
+		{
+			const FDMXMVRUnzip::FDMXScopedUnzipToTempFile ScopedUnzipGDTF(MVRUnzip, MVRFixture.GDTFSpec);
+			if (!ScopedUnzipGDTF.TempFilePathAndName.IsEmpty())
+			{
+				UPackage* Package = FDMXEditorUtils::GetOrCreatePackage(Parent, MVRFixture.GDTFSpec);
+
+				bool bCanceled;
+				UObject* NewGDTFObject = GDTFFactory->FactoryCreateFile(UDMXImportGDTF::StaticClass(), Package, FName(*MVRFixture.GDTFSpec), Flags | RF_Public, ScopedUnzipGDTF.TempFilePathAndName, nullptr, GWarn, bCanceled);
+				
+				if (UDMXImportGDTF* NewGDTF = Cast<UDMXImportGDTF>(NewGDTFObject))
+				{
+					ImportedGDTFs.Add(NewGDTF);
+					ImportedGDTFNames.Add(MVRFixture.GDTFSpec);
+
+					FAssetRegistryModule::AssetCreated(NewGDTF);
+					Package->MarkPackageDirty();
+				}
+			}
+		}
+	}
+
+	return ImportedGDTFs;
+}
+
+void UDMXLibraryFromMVRFactory::InitDMXLibrary(UDMXLibrary* DMXLibrary, const TArray<UDMXImportGDTF*>& GDTFAssets, const FDMXMVRGeneralSceneDescription& GeneralSceneDescription) const
+{
+	if (DMXLibrary)
+	{
+		// Create a Fixture Type for each GDTF
+		TMap<FString, UDMXEntityFixtureType*> GDTFSpecToFixtureTypeMap;
+		for (UDMXImportGDTF* GDTF : GDTFAssets)
+		{
+			const FString GDTFFilename = FPaths::GetCleanFilename(GDTF->SourceFilename);
+			FDMXEntityFixtureTypeConstructionParams FixtureTypeConstructionParams;
+			FixtureTypeConstructionParams.DMXCategory = FDMXFixtureCategory(FDMXFixtureCategory::GetFirstValue());
+			FixtureTypeConstructionParams.ParentDMXLibrary = DMXLibrary;
+
+			UDMXEntityFixtureType* FixtureType = UDMXEntityFixtureType::CreateFixtureTypeInLibrary(FixtureTypeConstructionParams, FPaths::GetBaseFilename(GDTFFilename));
+			FixtureType->SetGDTF(GDTF);
+
+			GDTFSpecToFixtureTypeMap.Add(GDTFFilename, FixtureType);
+		}
+
+		// Create Fixture Patches for the MVR Fixtures
+		for (const FDMXMVRFixture& MVRFixture : GeneralSceneDescription.MVRFixtures)
+		{
+			if (!GDTFSpecToFixtureTypeMap.Contains(MVRFixture.GDTFSpec))
+			{
+				continue;
+			}
+
+			UDMXEntityFixtureType* FixtureType = GDTFSpecToFixtureTypeMap[MVRFixture.GDTFSpec];
+			UDMXEntityFixturePatch* const* FixturePatchPtr = DMXLibrary->GetEntitiesTypeCast<UDMXEntityFixturePatch>().FindByPredicate([MVRFixture, FixtureType](UDMXEntityFixturePatch* FixturePatch)
+				{
+					const FDMXFixtureMode* ActiveModePtr = FixturePatch->GetActiveMode();
+					return
+						FixturePatch->GetFixtureType() == FixtureType &&
+						FixturePatch->GetUniverseID() == MVRFixture.Addresses.Universe &&
+						FixturePatch->GetStartingChannel() == MVRFixture.Addresses.Address &&
+						ActiveModePtr && ActiveModePtr->ModeName == MVRFixture.GDTFMode;
+				});
+			
+			if (!FixturePatchPtr)
+			{
+				int32 ActiveModeIndex = FixtureType->Modes.IndexOfByPredicate([MVRFixture](const FDMXFixtureMode& Mode)
+					{
+						return Mode.ModeName == MVRFixture.GDTFMode;
+					});
+
+				if (ActiveModeIndex != INDEX_NONE)
+				{
+					FDMXEntityFixturePatchConstructionParams FixturePatchConstructionParams;
+					FixturePatchConstructionParams.ActiveMode = ActiveModeIndex;
+					FixturePatchConstructionParams.FixtureTypeRef = FDMXEntityFixtureTypeRef(FixtureType);
+					FixturePatchConstructionParams.UniverseID = MVRFixture.Addresses.Universe;
+					FixturePatchConstructionParams.StartingAddress = MVRFixture.Addresses.Address;
+
+					UDMXEntityFixturePatch* FixturePatch = UDMXEntityFixturePatch::CreateFixturePatchInLibrary(FixturePatchConstructionParams, MVRFixture.Name);
+					FixturePatchPtr = &FixturePatch;
+				}
+			}
+			check(FixturePatchPtr);
+
+			(*FixturePatchPtr)->AddMVRFixtureInstance(MVRFixture);
+		}
+	}
+}
+
+#undef LOCTEXT_NAMESPACE
diff --git a/Engine/Plugins/VirtualProduction/DMX/DMXEngine/Source/DMXEditor/Private/Factories/DMXMVRSceneActorFactory.cpp b/Engine/Plugins/VirtualProduction/DMX/DMXEngine/Source/DMXEditor/Private/Factories/DMXMVRSceneActorFactory.cpp
new file mode 100644
--- /dev/null
+++ b/Engine/Plugins/VirtualProduction/DMX/DMXEngine/Source/DMXEditor/Private/Factories/DMXMVRSceneActorFactory.cpp
@@ -0,0 +1,77 @@
+// Copyright Epic Games, Inc. All Rights Reserved.
+
+#include "Factories/DMXMVRSceneActorFactory.h"
+
+#include "Engine/StaticMeshActor.h"
+#include "Library/DMXLibrary.h"
+#include "MVR/DMXMVRSceneActor.h"
+
+#include "Components/StaticMeshComponent.h"
+#include "Engine/StaticMesh.h"
+
+
+UDMXMVRSceneActorFactory::UDMXMVRSceneActorFactory()
+{
+	DisplayName = NSLOCTEXT("DMXMVRSceneActorFactory", "DMXMVRSceneActorFactoryDisplayName", "MVR Scene Actor");
+	NewActorClass = ADMXMVRSceneActor::StaticClass();
+}
+
+bool UDMXMVRSceneActorFactory::CanCreateActorFrom(const FAssetData& AssetData, FText& OutErrorMsg)
+{
+	if (AssetData.IsValid())
+	{
+		UClass* AssetClass = AssetData.GetClass();
+		if (AssetClass && AssetClass == UDMXLibrary::StaticClass())
+		{
+			return true;
+		}
+	}
+
+	return false;
+}
+
+AActor* UDMXMVRSceneActorFactory::SpawnActor(UObject* Asset, ULevel* InLevel, const FTransform& Transform, const FActorSpawnParameters& InSpawnParams)
+{
+	UDMXLibrary* DMXLibrary = Cast<UDMXLibrary>(Asset);
+
+	if (!DMXLibrary)
+	{
+		return nullptr;
+	}
+
+	if (InSpawnParams.ObjectFlags & RF_Transient)
+	{
+		// This is a hack for drag and drop so that we don't spawn all the actors for the preview actor since it gets deleted right after.
+
+		FActorSpawnParameters SpawnInfo(InSpawnParams);
+		SpawnInfo.OverrideLevel = InLevel;
+		AStaticMeshActor* DragActor = Cast<AStaticMeshActor>(InLevel->GetWorld()->SpawnActor(AStaticMeshActor::StaticClass(), &Transform, SpawnInfo));
+		DragActor->GetStaticMeshComponent()->SetStaticMesh(Cast< UStaticMesh >(FSoftObjectPath(TEXT("StaticMesh'/Engine/EditorMeshes/EditorSphere.EditorSphere'")).TryLoad()));
+		DragActor->SetActorScale3D(FVector(0.1f));
+
+		return DragActor;
+	}
+	else 
+	{
+		if (!InLevel)
+		{
+			return nullptr;
+		}
+
+		UWorld* World = InLevel->GetWorld();
+		if (!World)
+		{
+			return nullptr;
+		}
+
+		ADMXMVRSceneActor* MVRSceneActor = World->SpawnActor<ADMXMVRSceneActor>(ADMXMVRSceneActor::StaticClass());
+		if (MVRSceneActor)
+		{
+			MVRSceneActor->SetDMXLibrary(DMXLibrary);
+
+			return MVRSceneActor;
+		}
+	}
+
+	return nullptr;
+}
diff --git a/Engine/Plugins/VirtualProduction/DMX/DMXEngine/Source/DMXEditor/Private/Factories/DMXMVRSceneActorFactory.h b/Engine/Plugins/VirtualProduction/DMX/DMXEngine/Source/DMXEditor/Private/Factories/DMXMVRSceneActorFactory.h
new file mode 100644
--- /dev/null
+++ b/Engine/Plugins/VirtualProduction/DMX/DMXEngine/Source/DMXEditor/Private/Factories/DMXMVRSceneActorFactory.h
@@ -0,0 +1,29 @@
+// Copyright Epic Games, Inc. All Rights Reserved.
+
+
+#pragma once
+
+#include "CoreMinimal.h"
+#include "ActorFactories/ActorFactory.h"
+#include "UObject/ObjectMacros.h"
+
+#include "DMXMVRSceneActorFactory.generated.h"
+
+class AActor;
+struct FAssetData;
+
+
+UCLASS()
+class UDMXMVRSceneActorFactory
+	: public UActorFactory
+{
+	GENERATED_BODY()
+
+public:
+	UDMXMVRSceneActorFactory();
+
+	//~ Begin UActorFactory Interface
+	virtual bool CanCreateActorFrom(const FAssetData& AssetData, FText& OutErrorMsg) override;
+	virtual AActor* SpawnActor(UObject* Asset, ULevel* InLevel, const FTransform& Transform, const FActorSpawnParameters& InSpawnParams) override;
+	//~ End UActorFactory Interface
+};
diff --git a/Engine/Plugins/VirtualProduction/DMX/DMXEngine/Source/DMXEditor/Public/DMXEditorUtils.h b/Engine/Plugins/VirtualProduction/DMX/DMXEngine/Source/DMXEditor/Public/DMXEditorUtils.h
--- a/Engine/Plugins/VirtualProduction/DMX/DMXEngine/Source/DMXEditor/Public/DMXEditorUtils.h
+++ b/Engine/Plugins/VirtualProduction/DMX/DMXEngine/Source/DMXEditor/Public/DMXEditorUtils.h
@@ -148,6 +148,9 @@ public:
 	/** Clears cached data fixture patches received */
 	static void ClearFixturePatchCachedData();
 
+	/** Gets the package or creates a new one if it doesn't exist */
+	static UPackage* GetOrCreatePackage(TWeakObjectPtr<UObject> Parent, const FString& DesiredName);
+
 	// can't instantiate this class
 	FDMXEditorUtils() = delete;
 };
diff --git a/Engine/Plugins/VirtualProduction/DMX/DMXEngine/Source/DMXEditor/Public/Factories/DMXEditorFactoryNew.h b/Engine/Plugins/VirtualProduction/DMX/DMXEngine/Source/DMXEditor/Public/Factories/DMXEditorFactoryNew.h
--- a/Engine/Plugins/VirtualProduction/DMX/DMXEngine/Source/DMXEditor/Public/Factories/DMXEditorFactoryNew.h
+++ b/Engine/Plugins/VirtualProduction/DMX/DMXEngine/Source/DMXEditor/Public/Factories/DMXEditorFactoryNew.h
@@ -2,13 +2,14 @@
 
 #pragma once
 
-#include "EditorReimportHandler.h"
 #include "Factories/Factory.h"
 
 #include "DMXEditorFactoryNew.generated.h"
 
 class UDMXLibrary;
 
+/** DEPRECATED 5.1 */
+class UE_DEPRECATED(5.1, "Deprecated class. Please use DMXLibraryFactory instead.") UDMXEditorFactoryNew;
 UCLASS(hidecategories = Object)
 class DMXEDITOR_API UDMXEditorFactoryNew : public UFactory
 {
@@ -17,8 +18,8 @@ public:
 	UDMXEditorFactoryNew();
 
 	//~ Begin UFactory Interface
+	virtual FText GetDisplayName() const override { return FText::FromString(FString("DEPRECATED_DMXEditorFactoryNew")); }
 	virtual UObject* FactoryCreateNew(UClass* Class, UObject* InParent, FName Name, EObjectFlags Flags, UObject* Context, FFeedbackContext* Warn) override;
-	uint32 GetMenuCategories() const override;
 	//~ Begin UFactory Interface	
 
 private:
diff --git a/Engine/Plugins/VirtualProduction/DMX/DMXEngine/Source/DMXEditor/Public/Factories/DMXLibraryFactory.h b/Engine/Plugins/VirtualProduction/DMX/DMXEngine/Source/DMXEditor/Public/Factories/DMXLibraryFactory.h
new file mode 100644
--- /dev/null
+++ b/Engine/Plugins/VirtualProduction/DMX/DMXEngine/Source/DMXEditor/Public/Factories/DMXLibraryFactory.h
@@ -0,0 +1,25 @@
+// Copyright Epic Games, Inc. All Rights Reserved.
+
+#pragma once
+
+#include "Factories/Factory.h"
+
+#include "DMXLibraryFactory.generated.h"
+
+class UDMXLibrary;
+
+
+UCLASS()
+class DMXEDITOR_API UDMXLibraryFactory 
+	: public UFactory
+{
+	GENERATED_BODY()
+public:
+	UDMXLibraryFactory();
+
+	//~ Begin UFactory Interface
+	virtual FText GetDisplayName() const override { return FText::FromString(TEXT("DMX Library")); }
+	virtual UObject* FactoryCreateNew(UClass* Class, UObject* InParent, FName Name, EObjectFlags Flags, UObject* Context, FFeedbackContext* Warn) override;
+	uint32 GetMenuCategories() const override;
+	//~ Begin UFactory Interface	
+};
diff --git a/Engine/Plugins/VirtualProduction/DMX/DMXEngine/Source/DMXEditor/Public/Factories/DMXLibraryFromMVRFactory.h b/Engine/Plugins/VirtualProduction/DMX/DMXEngine/Source/DMXEditor/Public/Factories/DMXLibraryFromMVRFactory.h
new file mode 100644
--- /dev/null
+++ b/Engine/Plugins/VirtualProduction/DMX/DMXEngine/Source/DMXEditor/Public/Factories/DMXLibraryFromMVRFactory.h
@@ -0,0 +1,46 @@
+// Copyright Epic Games, Inc. All Rights Reserved.
+
+#pragma once
+
+#include "Factories/Factory.h"
+
+#include "DMXLibraryFromMVRFactory.generated.h"
+
+struct FDMXMVRGeneralSceneDescription;
+class FDMXMVRUnzip;
+class UDMXImportGDTF;
+class UDMXLibrary;
+
+
+
+UCLASS()
+class DMXEDITOR_API UDMXLibraryFromMVRFactory 
+	: public UFactory
+{
+	GENERATED_BODY()
+
+public:
+	UDMXLibraryFromMVRFactory();
+
+	//~ Begin UFactory Interface
+	virtual UObject* FactoryCreateFile(UClass* InClass, UObject* Parent, FName InName, EObjectFlags Flags, const FString& Filename, const TCHAR* Parms, FFeedbackContext* Warn, bool& bOutOperationCanceled);
+	virtual bool FactoryCanImport(const FString& Filename) override;
+	//~ End UFactory Interface	
+
+	/** File extention for MVR files */
+	static const FName MVRFileExtension;
+
+	/** File extention for GDTF files */
+	static const FName GDTFFileExtension;
+
+private:
+	/** Creates a DMX Library asset. Returns nullptr if the library could not be created */
+	UDMXLibrary* CreateDMXLibraryAsset(UObject* Parent, EObjectFlags Flags, const FString& InFilename);
+
+	/** Creates GDTF assets from the MVR */
+	TArray<UDMXImportGDTF*> CreateGDTFAssets(UObject* Parent, EObjectFlags Flags, const TSharedRef<FDMXMVRUnzip>& MVRUnzip, const FDMXMVRGeneralSceneDescription& GeneralSceneDescription);
+
+	/** Initializes the DMX Library from the General Scene Description and GDTF assets */
+	void InitDMXLibrary(UDMXLibrary* DMXLibrary, const TArray<UDMXImportGDTF*>& GDTFAssets, const FDMXMVRGeneralSceneDescription& GeneralSceneDescription) const;
+
+};
diff --git a/Engine/Plugins/VirtualProduction/DMX/DMXEngine/Source/DMXRuntime/DMXRuntime.Build.cs b/Engine/Plugins/VirtualProduction/DMX/DMXEngine/Source/DMXRuntime/DMXRuntime.Build.cs
--- a/Engine/Plugins/VirtualProduction/DMX/DMXEngine/Source/DMXRuntime/DMXRuntime.Build.cs
+++ b/Engine/Plugins/VirtualProduction/DMX/DMXEngine/Source/DMXRuntime/DMXRuntime.Build.cs
@@ -13,19 +13,32 @@ public class DMXRuntime : ModuleRules
 			{
 				"Core",
                 "CoreUObject",
-                "Engine",
+				"DatasmithContent",
+				"Engine",
 				"DMXProtocol",
 				"JsonUtilities",
-
 			}
 		);
 
 		PrivateDependencyModuleNames.AddRange(
 			new string[]
 			{
+				"AssetRegistry",
 				"Json",
 				"MovieScene",
+				"Projects",
+				"XmlParser",
 			}
 		);
+
+		if (Target.bBuildEditor == true)
+		{
+			PrivateDependencyModuleNames.AddRange(
+				new string[]
+				{
+					"UnrealEd",
+				}
+			);
+		}
 	}
 }
diff --git a/Engine/Plugins/VirtualProduction/DMX/DMXEngine/Source/DMXRuntime/Private/Library/DMXEntityFixturePatch.cpp b/Engine/Plugins/VirtualProduction/DMX/DMXEngine/Source/DMXRuntime/Private/Library/DMXEntityFixturePatch.cpp
--- a/Engine/Plugins/VirtualProduction/DMX/DMXEngine/Source/DMXRuntime/Private/Library/DMXEntityFixturePatch.cpp
+++ b/Engine/Plugins/VirtualProduction/DMX/DMXEngine/Source/DMXRuntime/Private/Library/DMXEntityFixturePatch.cpp
@@ -13,6 +13,7 @@
 #include "IO/DMXOutputPort.h"
 #include "Library/DMXEntityController.h"
 #include "Library/DMXEntityFixtureType.h"
+#include "Library/DMXImportGDTF.h"
 #include "Library/DMXLibrary.h"
 #include "Modulators/DMXModulator.h"
 
@@ -24,6 +25,32 @@ DECLARE_CYCLE_STAT(TEXT("FixturePatch cache values"), STAT_DMXFixturePatchCacheV
 
 #define LOCTEXT_NAMESPACE "DMXEntityFixturePatch"
 
+
+UDMXMVRFixtureInstance* UDMXMVRFixtureInstance::CreateFromMVRFixture(const FDMXMVRFixture& InMVRFixture)
+{
+	UDMXMVRFixtureInstance* NewMVRFixtureInstance = NewObject<UDMXMVRFixtureInstance>();
+	NewMVRFixtureInstance->MVRFixture = InMVRFixture;
+
+	return NewMVRFixtureInstance;
+}
+
+void UDMXMVRFixtureInstance::Serialize(FArchive& Ar)
+{
+	Super::Serialize(Ar);
+
+	Ar << MVRFixture;
+}
+
+void UDMXMVRFixtureInstance::SetMVRFixture(const FDMXMVRFixture& InMVRFixture)
+{
+	if (!ensureAlwaysMsgf(InMVRFixture.UUID == MVRFixture.UUID, TEXT("Cannot change the MVR UUID of an existing MVR Fixture Instance")))
+	{
+		return;
+	}
+
+	MVRFixture = InMVRFixture;
+}
+
 FDMXOnFixturePatchChangedDelegate UDMXEntityFixturePatch::OnFixturePatchChangedDelegate;
 
 UDMXEntityFixturePatch::UDMXEntityFixturePatch()
@@ -55,7 +82,7 @@ UDMXEntityFixturePatch* UDMXEntityFixturePatch::CreateFixturePatchInLibrary(FDMX
 			}
 #endif
 
-			FString EntityName = FDMXRuntimeUtils::FindUniqueEntityName(DMXLibrary, UDMXEntityFixturePatch::StaticClass(), DesiredName);
+			const FString EntityName = FDMXRuntimeUtils::FindUniqueEntityName(DMXLibrary, UDMXEntityFixturePatch::StaticClass(), DesiredName);
 
 			UDMXEntityFixturePatch* NewFixturePatch = NewObject<UDMXEntityFixturePatch>(DMXLibrary, UDMXEntityFixturePatch::StaticClass(), NAME_None, RF_Transactional);
 			NewFixturePatch->SetName(EntityName);
@@ -63,7 +90,7 @@ UDMXEntityFixturePatch* UDMXEntityFixturePatch::CreateFixturePatchInLibrary(FDMX
 			NewFixturePatch->SetUniverseID(ConstructionParams.UniverseID);
 			NewFixturePatch->SetStartingChannel(ConstructionParams.StartingAddress);
 			NewFixturePatch->SetActiveModeIndex(ConstructionParams.ActiveMode);
-
+			
 #if WITH_EDITOR
 			if (bMarkDMXLibraryDirty)
 			{
@@ -499,6 +526,44 @@ bool UDMXEntityFixturePatch::SetActiveModeIndex(int32 NewActiveModeIndex)
 	return false;
 }
 
+void UDMXEntityFixturePatch::AddMVRFixtureInstance(const FDMXMVRFixture& MVRFixtureToAdd)
+{
+	if (!ensureMsgf(ParentFixtureTypeTemplate, TEXT("Trying to create an MVR Fixture Instance in Fixture Patch '%s', but it has not Fixture Type."), *Name))
+	{
+		return;
+	}
+	if (!ensureMsgf(ParentFixtureTypeTemplate->DMXImport, TEXT("Trying to create an MVR Fixture Instance in Fixture Patch '%s', but its Fixture Type '%s' has no GDTF set."), *Name, *ParentFixtureTypeTemplate->Name))
+	{
+		return;
+	}
+	if (!ensureMsgf(MVRFixtureToAdd.UUID.IsValid(), TEXT("Trying to create an MVR Fixture Instance in Fixture Patch '%s', the UUID is not valid."), *Name, *ParentFixtureTypeTemplate->Name))
+	{
+		return;
+	}
+	const bool bMVRUUIDAlreadyContained = MVRFixtureInstances.ContainsByPredicate([MVRFixtureToAdd](const UDMXMVRFixtureInstance* MVRFixtureInstance)
+		{
+			return MVRFixtureInstance->GetMVRFixture().UUID == MVRFixtureToAdd.UUID;
+		});
+	if (!ensureMsgf(!bMVRUUIDAlreadyContained, TEXT("Trying to add an MVR Fixture to a Fixture Patch, but its UUID is already contained.")))
+	{
+		return;
+	}	
+
+	UDMXMVRFixtureInstance* NewMVRFixtureInstance = UDMXMVRFixtureInstance::CreateFromMVRFixture(MVRFixtureToAdd);
+	MVRFixtureInstances.Add(NewMVRFixtureInstance);
+}
+
+TArray<FDMXMVRFixture> UDMXEntityFixturePatch::MakeMVRFixtures() const
+{
+	TArray<FDMXMVRFixture> Result;
+	for (UDMXMVRFixtureInstance* MVRFixtureInstance : MVRFixtureInstances)
+	{
+		Result.Add(MVRFixtureInstance->GetMVRFixture());
+	}
+
+	return Result;
+}
+
 int32 UDMXEntityFixturePatch::GetRemoteUniverse() const
 {	
 	/** DEPRECATED 4.27 */
diff --git a/Engine/Plugins/VirtualProduction/DMX/DMXEngine/Source/DMXRuntime/Private/Library/DMXEntityFixtureType.cpp b/Engine/Plugins/VirtualProduction/DMX/DMXEngine/Source/DMXRuntime/Private/Library/DMXEntityFixtureType.cpp
--- a/Engine/Plugins/VirtualProduction/DMX/DMXEngine/Source/DMXRuntime/Private/Library/DMXEntityFixtureType.cpp
+++ b/Engine/Plugins/VirtualProduction/DMX/DMXEngine/Source/DMXRuntime/Private/Library/DMXEntityFixtureType.cpp
@@ -300,20 +300,20 @@ void UDMXEntityFixtureType::PostEditUndo()
 }
 #endif // WITH_EDITOR
 
-#if WITH_EDITOR
-void UDMXEntityFixtureType::SetModesFromDMXImport(UDMXImport* DMXImportAsset)
+void UDMXEntityFixtureType::SetGDTF(UDMXImportGDTF* GDTFAsset)
 {
-	if (DMXImportAsset == nullptr || !DMXImportAsset->IsValidLowLevelFast())
+	if (!IsValid(GDTFAsset))
 	{
 		return;
 	}
 
-	if (UDMXImportGDTFDMXModes* GDTFDMXModes = Cast<UDMXImportGDTFDMXModes>(DMXImportAsset->DMXModes))
+	DMXImport = GDTFAsset;
+	if (UDMXImportGDTFDMXModes* GDTFDMXModes = Cast<UDMXImportGDTFDMXModes>(DMXImport->DMXModes))
 	{
 		// Clear existing modes
-		Modes.Empty(DMXImportAsset->DMXModes != nullptr ? GDTFDMXModes->DMXModes.Num() : 0);
+		Modes.Empty(DMXImport->DMXModes != nullptr ? GDTFDMXModes->DMXModes.Num() : 0);
 
-		if (DMXImportAsset->DMXModes == nullptr)
+		if (DMXImport->DMXModes == nullptr)
 		{
 			return;
 		}
@@ -435,7 +435,17 @@ void UDMXEntityFixtureType::SetModesFromDMXImport(UDMXImport* DMXImportAsset)
 
 	OnFixtureTypeChangedDelegate.Broadcast(this);
 }
-#endif // WITH_EDITOR
+
+void UDMXEntityFixtureType::SetModesFromDMXImport(UDMXImport* DMXImportAsset)
+{
+	UDMXImportGDTF* GDTFAsset = Cast<UDMXImportGDTF>(DMXImportAsset);
+	if (!ensureAlwaysMsgf(GDTFAsset, TEXT("DMXImportAsset %s is not a DMXImportGDTF type, this was never supported. Cannot SetModesFromDMXImport"), *DMXImportAsset->GetName()))
+	{
+		return;
+	}
+
+	SetGDTF(GDTFAsset);
+}
 
 FDMXOnFixtureTypeChangedDelegate& UDMXEntityFixtureType::GetOnFixtureTypeChanged()
 {
diff --git a/Engine/Plugins/VirtualProduction/DMX/DMXEngine/Source/DMXRuntime/Private/Library/DMXLibrary.cpp b/Engine/Plugins/VirtualProduction/DMX/DMXEngine/Source/DMXRuntime/Private/Library/DMXLibrary.cpp
--- a/Engine/Plugins/VirtualProduction/DMX/DMXEngine/Source/DMXRuntime/Private/Library/DMXLibrary.cpp
+++ b/Engine/Plugins/VirtualProduction/DMX/DMXEngine/Source/DMXRuntime/Private/Library/DMXLibrary.cpp
@@ -530,6 +530,19 @@ void UDMXLibrary::UpdatePorts()
 	}
 }
 
+FDMXMVRGeneralSceneDescription UDMXLibrary::MakeGeneralSceneDescription() const
+{
+	FDMXMVRGeneralSceneDescription GeneralSceneDescription;
+	const TArray<UDMXEntityFixturePatch*> FixturePatches = GetEntitiesTypeCast<UDMXEntityFixturePatch>();
+	for (UDMXEntityFixturePatch* FixturePatch : FixturePatches)
+	{
+		const TArray<FDMXMVRFixture> MVRFixtures = FixturePatch->MakeMVRFixtures();
+		GeneralSceneDescription.MVRFixtures.Append(MVRFixtures);
+	}
+
+	return GeneralSceneDescription;
+}
+
 #if WITH_EDITOR
 void UDMXLibrary::UpgradeFromControllersToPorts()
 {	
diff --git a/Engine/Plugins/VirtualProduction/DMX/DMXEngine/Source/DMXRuntime/Private/MVR/DMXMVRAssetUserData.cpp b/Engine/Plugins/VirtualProduction/DMX/DMXEngine/Source/DMXRuntime/Private/MVR/DMXMVRAssetUserData.cpp
new file mode 100644
--- /dev/null
+++ b/Engine/Plugins/VirtualProduction/DMX/DMXEngine/Source/DMXRuntime/Private/MVR/DMXMVRAssetUserData.cpp
@@ -0,0 +1,56 @@
+// Copyright Epic Games, Inc. All Rights Reserved.
+
+#include "MVR/DMXMVRAssetUserData.h"
+
+#include "GameFramework/Actor.h"
+
+
+const FName UDMXMVRAssetUserData::MVRUUIDMetaDataKey = TEXT("MVRUUID");
+
+
+UDMXMVRAssetUserData* UDMXMVRAssetUserData::GetMVRAssetUserData(AActor* Actor)
+{
+	if (USceneComponent* RootComponent = Actor->GetRootComponent())
+	{
+		return RootComponent->GetAssetUserData<UDMXMVRAssetUserData>();
+	}
+
+	return nullptr;
+}
+
+FString UDMXMVRAssetUserData::GetMVRAssetUserDataValueForkey(AActor* Actor, const FName Key)
+{
+	if (UDMXMVRAssetUserData* MVRAssetUserData = GetMVRAssetUserData(Actor))
+	{
+		const FString* const ValuePtr = MVRAssetUserData->MetaData.Find(Key);
+		return ValuePtr ? *ValuePtr : FString();
+	}
+
+	return FString();
+}
+
+bool UDMXMVRAssetUserData::SetMVRAssetUserDataValueForKey(AActor* Actor, FName Key, const FString& Value)
+{
+	// For AActor, the interface is actually implemented by the ActorComponent
+	if (USceneComponent* RootComponent = Actor->GetRootComponent())
+	{
+		if (IInterface_AssetUserData* AssetUserData = Cast<IInterface_AssetUserData>(RootComponent))
+		{
+			UDMXMVRAssetUserData* MVRUserData = AssetUserData->GetAssetUserData<UDMXMVRAssetUserData>();
+
+			if (!MVRUserData)
+			{
+				MVRUserData = NewObject<UDMXMVRAssetUserData>(RootComponent, NAME_None, RF_Public | RF_Transactional);
+				AssetUserData->AddAssetUserData(MVRUserData);
+			}
+
+			// Add Datasmith meta data
+			MVRUserData->MetaData.Add(Key, Value);
+			MVRUserData->MetaData.KeySort(FNameLexicalLess());
+
+			return true;
+		}
+	}
+
+	return false;
+}
diff --git a/Engine/Plugins/VirtualProduction/DMX/DMXEngine/Source/DMXRuntime/Private/MVR/DMXMVRFixtureActorLibrary.cpp b/Engine/Plugins/VirtualProduction/DMX/DMXEngine/Source/DMXRuntime/Private/MVR/DMXMVRFixtureActorLibrary.cpp
new file mode 100644
--- /dev/null
+++ b/Engine/Plugins/VirtualProduction/DMX/DMXEngine/Source/DMXRuntime/Private/MVR/DMXMVRFixtureActorLibrary.cpp
@@ -0,0 +1,376 @@
+// Copyright Epic Games, Inc. All Rights Reserved.
+
+#include "MVR/DMXMVRFixtureActorLibrary.h"
+
+#include "DMXRuntimeLog.h"
+#include "MVR/DMXMVRFixtureActorInterface.h"
+#include "Library/DMXEntityFixturePatch.h"
+
+#include "AssetRegistryModule.h"
+#include "PreviewScene.h"
+#include "Engine/SkeletalMesh.h"
+#include "Interfaces/IPluginManager.h"
+#include "Misc/ScopedSlowTask.h"
+
+
+class FDMXMVRFixtureActorAssetHierarchy
+{
+private:
+	class FDMXVRAssetNode
+		: public TSharedFromThis<FDMXVRAssetNode>
+	{
+	public:
+		FDMXVRAssetNode(FName InClassPath, FString InClassName)
+			: ClassName(InClassName)
+			, ClassPath(InClassPath)
+		{}
+
+		TSharedPtr<FDMXVRAssetNode> Parent;
+		TArray<TSharedRef<FDMXVRAssetNode>> Children;
+
+		void AppendClassPathsOfSelfAndChildren(TArray<FName>& OutChildClassPaths) const
+		{
+			OutChildClassPaths.Add(ClassPath);
+			for (const TSharedRef<FDMXVRAssetNode>& ChildNode : Children)
+			{
+				ChildNode->AppendClassPathsOfSelfAndChildren(OutChildClassPaths);
+			}
+		}
+
+		FString ClassName;
+		FName ClassPath;
+		FName ParentClassPath;
+		FName BlueprintAssetPath;
+		TArray<FString> ImplementedInterfaces;
+	};
+
+public:
+	static void GetClassPathsWithInterface(const FString& InInterfaceClassPathName, TArray<FName>& OutClassPaths)
+	{
+		FDMXMVRFixtureActorAssetHierarchy MVRAssetHierarchy;
+
+		if (MVRAssetHierarchy.ObjectClassRoot.IsValid())
+		{
+			const TSharedRef<FDMXVRAssetNode> StartNode = MVRAssetHierarchy.ObjectClassRoot.ToSharedRef();
+			MVRAssetHierarchy.GetClassPathsWithInterface(InInterfaceClassPathName, OutClassPaths, StartNode);
+		}
+	}
+
+private:
+	void GetClassPathsWithInterface(const FString& InInterfaceClassPathName, TArray<FName>& OutClassPaths, const TSharedRef<FDMXVRAssetNode>& StartNode)
+	{
+		if (StartNode->ImplementedInterfaces.Contains(InInterfaceClassPathName))
+		{
+			StartNode->AppendClassPathsOfSelfAndChildren(OutClassPaths);
+			return;
+		}
+
+		// Recursive on all children
+		for (const TSharedRef<FDMXVRAssetNode>& ChildNode : StartNode->Children)
+		{
+			GetClassPathsWithInterface(InInterfaceClassPathName, OutClassPaths, ChildNode);
+		}
+	}
+
+	FDMXMVRFixtureActorAssetHierarchy()
+	{
+		// Fetch all classes from AssetRegistry blueprint data (which covers unloaded classes), and in-memory UClasses.
+		// Create a node for each one with unioned data from the AssetRegistry or UClass for that class.
+		// Set parent/child pointers to create a tree, and store this tree in this->ObjectClassRoot
+		TMap<FName, TSharedPtr<FDMXVRAssetNode>> ClassPathToNode;
+
+		// Create a node for every Blueprint class listed in the AssetRegistry and set the Blueprint fields
+		// Retrieve all blueprint classes 
+		FAssetRegistryModule& AssetRegistryModule = FModuleManager::LoadModuleChecked<FAssetRegistryModule>(TEXT("AssetRegistry"));
+		TArray<FAssetData> BlueprintList;
+		FARFilter Filter;
+		Filter.ClassNames.Add(UBlueprint::StaticClass()->GetFName());
+		Filter.RecursiveClassesExclusionSet.FindOrAdd(*USkeletalMesh::StaticClass()->GetName());
+		Filter.bRecursiveClasses = true;
+
+		AssetRegistryModule.Get().GetAssets(Filter, BlueprintList);
+
+		FString ClassPathString;
+		for (FAssetData& AssetData : BlueprintList)
+		{
+			FName ClassPath;
+			if (AssetData.GetTagValue(FBlueprintTags::GeneratedClassPath, ClassPathString))
+			{
+				ClassPath = FName(*FPackageName::ExportTextPathToObjectPath(ClassPathString));
+			}
+			if (ClassPath.IsNone())
+			{
+				continue;
+			}
+
+			TSharedPtr<FDMXVRAssetNode>& Node = ClassPathToNode.FindOrAdd(ClassPath);
+			if (!Node.IsValid())
+			{
+				const FString ClassName = AssetData.AssetName.ToString();
+				Node = MakeShared<FDMXVRAssetNode>(ClassPath, ClassName);
+			}
+			
+			SetAssetDataFields(Node, AssetData);
+		}
+
+		// FindOrCreate a node for every loaded UClass, and set the UClass fields
+		BuildHierarchyFromLoadedClassees(ObjectClassRoot, ClassPathToNode);
+
+		// Set the parent and child pointers
+		for (TPair<FName, TSharedPtr<FDMXVRAssetNode>>& KVPair : ClassPathToNode)
+		{
+			TSharedPtr<FDMXVRAssetNode>& Node = KVPair.Value;
+			if (Node == ObjectClassRoot)
+			{
+				// No parent expected for the root class
+				continue;
+			}
+			TSharedPtr<FDMXVRAssetNode>* ParentNodePtr = nullptr;
+			if (!Node->ParentClassPath.IsNone())
+			{
+				ParentNodePtr = ClassPathToNode.Find(Node->ParentClassPath);
+			}
+			if (!ParentNodePtr)
+			{
+				continue;
+			}
+			TSharedPtr<FDMXVRAssetNode>& ParentNode = *ParentNodePtr;
+			check(ParentNode);
+			ParentNode->Children.Add(Node.ToSharedRef());
+		}
+	}
+
+	void BuildHierarchyFromLoadedClassees(TSharedPtr<FDMXVRAssetNode>& OutRootNode, TMap<FName, TSharedPtr<FDMXVRAssetNode>>& InOutClassPathToNode)
+	{
+		for (TObjectIterator<UClass> ClassIt; ClassIt; ++ClassIt)
+		{
+			UClass* CurrentClass = *ClassIt;
+			// Ignore deprecated and temporary trash classes.
+			if (CurrentClass->HasAnyClassFlags(CLASS_Deprecated | CLASS_NewerVersionExists | CLASS_Hidden))
+			{
+				continue;
+			}
+			TSharedPtr<FDMXVRAssetNode>& Node = InOutClassPathToNode.FindOrAdd(FName(*CurrentClass->GetPathName()));
+			if (!Node.IsValid())
+			{
+				Node = MakeShared<FDMXVRAssetNode>(*CurrentClass->GetPathName(), CurrentClass->GetPathName());
+			}
+			SetClassFields(Node, *CurrentClass);
+		}
+		TSharedPtr<FDMXVRAssetNode>* ExistingRoot = InOutClassPathToNode.Find(FName(*AActor::StaticClass()->GetPathName()));
+		check(ExistingRoot && ExistingRoot->IsValid());
+		OutRootNode = *ExistingRoot;
+	}
+
+	void SetClassFields(TSharedPtr<FDMXVRAssetNode>& InOutClassNode, UClass& Class)
+	{
+		// Fields that can als be set from FAssetData
+		if (InOutClassNode->ClassPath.IsNone())
+		{
+			InOutClassNode->ClassPath = FName(*Class.GetPathName());
+		}
+		if (InOutClassNode->ParentClassPath.IsNone())
+		{
+			if (Class.GetSuperClass())
+			{
+				InOutClassNode->ParentClassPath = FName(*Class.GetSuperClass()->GetPathName());
+			}
+		}
+
+		for (const FImplementedInterface& Interface : Class.Interfaces)
+		{
+			InOutClassNode->ImplementedInterfaces.Add(Interface.Class->GetPathName());
+		}
+	}
+
+	TSharedPtr<FDMXVRAssetNode> FindNodeByGeneratedClassPath(const TSharedRef<FDMXVRAssetNode>& InRootNode, FName InGeneratedClassPath)
+	{
+		if (InRootNode->ClassPath == InGeneratedClassPath)
+		{
+			return InRootNode;
+		}
+
+		TSharedPtr<FDMXVRAssetNode> ReturnNode;
+		// Search the children recursively, one of them might have the parent.
+		for (int32 ChildClassIndex = 0; !ReturnNode.IsValid() && ChildClassIndex < InRootNode->Children.Num(); ChildClassIndex++)
+		{
+			// Check the child, then check the return to see if it is valid. If it is valid, end the recursion.
+			ReturnNode = FindNodeByGeneratedClassPath(InRootNode->Children[ChildClassIndex], InGeneratedClassPath);
+
+			if (ReturnNode.IsValid())
+			{
+				return ReturnNode;
+			}
+		}
+
+		return nullptr;
+	}
+
+	void SetAssetDataFields(TSharedPtr<FDMXVRAssetNode>& InOutMXVRAssetNode, const FAssetData& InAssetData)
+	{
+		if (InOutMXVRAssetNode->ClassPath.IsNone())
+		{
+			FString GeneratedClassPath;
+			UClass* AssetClass = InAssetData.GetClass();
+			if (AssetClass && AssetClass->IsChildOf(UBlueprintGeneratedClass::StaticClass()))
+			{
+				InOutMXVRAssetNode->ClassPath = InAssetData.ObjectPath;
+			}
+			else if (InAssetData.GetTagValue(FBlueprintTags::GeneratedClassPath, GeneratedClassPath))
+			{
+				InOutMXVRAssetNode->ClassPath = FName(*FPackageName::ExportTextPathToObjectPath(GeneratedClassPath));
+			}
+		}
+		if (InOutMXVRAssetNode->ParentClassPath.IsNone())
+		{
+			FString ParentClassPathString;
+			if (InAssetData.GetTagValue(FBlueprintTags::ParentClassPath, ParentClassPathString))
+			{
+				InOutMXVRAssetNode->ParentClassPath = FName(*FPackageName::ExportTextPathToObjectPath(ParentClassPathString));
+			}
+		}
+
+		// Blueprint-specific fields
+		InOutMXVRAssetNode->BlueprintAssetPath = InAssetData.ObjectPath;
+
+		// Get interface class paths.
+		TArray<FString> ImplementedInterfaces;
+		GetImplementedInterfaceClassPathsFromAsset(InAssetData, ImplementedInterfaces);
+		InOutMXVRAssetNode->ImplementedInterfaces = ImplementedInterfaces;
+	}
+
+	void GetImplementedInterfaceClassPathsFromAsset(const FAssetData& InAssetData, TArray<FString>& OutClassPaths)
+	{
+		if (!InAssetData.IsValid())
+		{
+			return;
+		}
+
+		const FString ImplementedInterfaces = InAssetData.GetTagValueRef<FString>(FBlueprintTags::ImplementedInterfaces);
+		if (!ImplementedInterfaces.IsEmpty())
+		{
+			// Parse string like "((Interface=Class'"/Script/VPBookmark.VPBookmarkProvider"'),(Interface=Class'"/Script/VPUtilities.VPContextMenuProvider"'))"
+			// We don't want to actually resolve the hard ref so do some manual parsing
+
+			FString FullInterface;
+			FString CurrentString = *ImplementedInterfaces;
+			while (CurrentString.Split(TEXT("Interface="), nullptr, &FullInterface))
+			{
+				// Cutoff at next )
+				int32 RightParen = INDEX_NONE;
+				if (FullInterface.FindChar(TCHAR(')'), RightParen))
+				{
+					// Keep parsing
+					CurrentString = FullInterface.Mid(RightParen);
+
+					// Strip class name
+					FullInterface = *FPackageName::ExportTextPathToObjectPath(FullInterface.Left(RightParen));
+
+					// Handle quotes
+					FString InterfacePath;
+					const TCHAR* NewBuffer = FPropertyHelpers::ReadToken(*FullInterface, InterfacePath, true);
+
+					if (NewBuffer)
+					{
+						OutClassPaths.Add(InterfacePath);
+					}
+				}
+			}
+		}
+	}
+
+	TSharedPtr<FDMXVRAssetNode> ObjectClassRoot;
+};
+
+
+FDMXMVRFixtureActorLibrary::FDMXMVRFixtureActorLibrary()
+{
+	TArray<FName> ClassPaths;
+	FDMXMVRFixtureActorAssetHierarchy::GetClassPathsWithInterface(UDMXMVRFixtureActorInterface::StaticClass()->GetPathName(), ClassPaths);
+	ClassPaths.RemoveAll([](const FName& ClassPath)
+		{
+			return
+				ClassPath.ToString().Contains(TEXT(".SKEL_")) ||
+				ClassPath.ToString().Contains(TEXT(".REINST_"));
+		});
+
+	FPreviewScene PreviewScene;
+	TArray<UClass*> MVRFixtureActorClasses;
+	for (const FName& ClassPath : ClassPaths)
+	{
+		UClass* Class = LoadClass<UObject>(nullptr, *ClassPath.ToString());
+		AActor* Actor = PreviewScene.GetWorld()->SpawnActor<AActor>(Class);
+		check(Actor);
+
+		MVRActors.Add(Actor);
+	}
+}
+
+UClass* FDMXMVRFixtureActorLibrary::FindMostAppropriateActorClassForPatch(const UDMXEntityFixturePatch* const Patch) const
+{
+	const FDMXFixtureMode* const ModePtr = Patch->GetActiveMode();
+	if (!ModePtr)
+	{
+		return nullptr;
+	}
+
+	TArray<FName> AttributesOfPatch;
+	for (const FDMXFixtureFunction& Function : ModePtr->Functions)
+	{
+		AttributesOfPatch.Add(Function.Attribute.Name);
+	}
+	TArray<FName> MatrixAttributesOfPatch;
+	if (ModePtr->bFixtureMatrixEnabled)
+	{
+		for (const FDMXFixtureCellAttribute& CellAttribute : ModePtr->FixtureMatrixConfig.CellAttributes)
+		{
+			MatrixAttributesOfPatch.Add(CellAttribute.Attribute.Name);
+		}
+	}
+
+	// Best result
+	AActor* BestMatch = nullptr;
+	TArray<FName> BestMatchingAttributes;
+	TArray<FName> BestMatchingMatrixAttributes;
+
+	// Search through the actors
+	TArray<FName> SupportedAttributes;
+	TArray<FName> SupportedMatrixAttributes;
+	TArray<FName> MatchingAttributes;
+	TArray<FName> MatchingMatrixAttributes;
+	for (AActor* MVRActor : MVRActors)
+	{
+		// Look up attributes
+		SupportedAttributes.Reset();
+		SupportedMatrixAttributes.Reset();
+		IDMXMVRFixtureActorInterface::Execute_OnMVRGetSupportedDMXAttributes(MVRActor, SupportedAttributes, SupportedMatrixAttributes);
+
+		MatchingAttributes.Reset();
+		MatchingMatrixAttributes.Reset();
+		for (const FName& AttributeName : SupportedAttributes)
+		{
+			if (AttributesOfPatch.Contains(AttributeName))
+			{
+				MatchingAttributes.AddUnique(AttributeName);
+			}
+			else if (ModePtr->bFixtureMatrixEnabled && MatrixAttributesOfPatch.Contains(AttributeName))
+			{
+				MatchingMatrixAttributes.AddUnique(AttributeName);
+			}
+		}
+
+		if (BestMatchingAttributes.Num() + BestMatchingMatrixAttributes.Num() < MatchingAttributes.Num() + MatchingMatrixAttributes.Num())
+		{
+			BestMatch = MVRActor;
+			BestMatchingAttributes = MatchingAttributes;
+			BestMatchingMatrixAttributes = MatchingMatrixAttributes;
+		}
+	}
+
+	return BestMatch ? BestMatch->GetClass() : nullptr;
+}
+
+void FDMXMVRFixtureActorLibrary::AddReferencedObjects(FReferenceCollector& Collector)
+{
+	Collector.AddReferencedObjects(MVRActors);
+}
diff --git a/Engine/Plugins/VirtualProduction/DMX/DMXEngine/Source/DMXRuntime/Private/MVR/DMXMVRGeneralSceneDescription.cpp b/Engine/Plugins/VirtualProduction/DMX/DMXEngine/Source/DMXRuntime/Private/MVR/DMXMVRGeneralSceneDescription.cpp
new file mode 100644
--- /dev/null
+++ b/Engine/Plugins/VirtualProduction/DMX/DMXEngine/Source/DMXRuntime/Private/MVR/DMXMVRGeneralSceneDescription.cpp
@@ -0,0 +1,648 @@
+// Copyright Epic Games, Inc. All Rights Reserved.
+
+#include "MVR/DMXMVRGeneralSceneDescription.h"
+
+#include "DMXProtocolCommon.h"
+#include "DMXRuntimeLog.h"
+
+#include "XmlFile.h"
+#include "XmlNode.h"
+
+
+namespace UE::DMXRuntime::DMXMVRFixture::Private
+{
+	/** Parses the Name from the MVR Fixture Node, returns true on success. */
+	bool ParseName(const FXmlNode* const InFixtureNode, FString& OutName)
+	{
+		check(InFixtureNode);
+
+		static const FString InFixtureNodeAttribute_Name = TEXT("Name");
+		OutName = InFixtureNode->GetAttribute(InFixtureNodeAttribute_Name);
+		if (!OutName.IsEmpty())
+		{
+			return true;
+		}
+
+		return false;
+	}
+
+	/** Parses the UUID from the MVR Fixture Node, returns true on success. */
+	bool ParseUUID(const FXmlNode* const InFixtureNode, FGuid& OutUUID)
+	{
+		check(InFixtureNode);
+
+		static const FString InFixtureNodeAttribute_UUID = TEXT("UUID");
+		const FString UUIDString = InFixtureNode->GetAttribute(InFixtureNodeAttribute_UUID);
+		if (FGuid::Parse(UUIDString, OutUUID))
+		{
+			return true;
+		}
+		return false;
+	}
+
+	/** Parses the Matrix from the MVR Fixture Node, optional, always succeeds. */
+	void ParseMatrix(const FXmlNode* const InFixtureNode, TOptional<FTransform>& OutTransform)
+	{
+		check(InFixtureNode);
+		OutTransform.Reset();
+
+		static const FString ChildNodeName_Matrix = TEXT("Matrix");
+		const FXmlNode* const MatrixNode = InFixtureNode->FindChildNode(ChildNodeName_Matrix);
+
+		if (!MatrixNode)
+		{
+			return;
+		}
+
+		auto GetVectorStringsFromMatrixStringLambda([](const FString& MatrixString) -> TArray<FString>
+			{
+				TArray<FString> Result;
+				int32 SubstringIndex = 0;
+				for (int32 SubStringNumber = 0; SubStringNumber < 4; SubStringNumber++)
+				{
+					SubstringIndex = MatrixString.Find(TEXT("{"), ESearchCase::IgnoreCase, ESearchDir::FromStart, SubstringIndex) + 1;
+					const int32 EndIndex = MatrixString.Find(TEXT("}"), ESearchCase::IgnoreCase, ESearchDir::FromStart, SubstringIndex);
+
+					const FString SubString = MatrixString.Mid(SubstringIndex, EndIndex - SubstringIndex);
+					Result.Add(SubString);
+				}
+
+				return Result;
+			});
+
+		auto GetVectorFromVectorStringLambda([](const FString& VectorString) -> FVector
+			{
+				const int32 XIndex = 0;
+				const int32 YIndex = VectorString.Find(TEXT(","), ESearchCase::IgnoreCase, ESearchDir::FromStart, 0) + 1;
+				const int32 ZIndex = VectorString.Find(TEXT(","), ESearchCase::IgnoreCase, ESearchDir::FromStart, YIndex) + 1;
+
+				const FString XString = "X=" + VectorString.Mid(XIndex, YIndex - XIndex);
+				const FString YString = "Y=" + VectorString.Mid(YIndex, ZIndex - YIndex);
+				const FString ZString = "Z=" + VectorString.Mid(ZIndex, VectorString.Len() - ZIndex);
+
+				FVector Vector;
+				if (Vector.InitFromString(XString + YString + ZString))
+				{
+					return Vector;
+				}
+
+				return FVector::ZeroVector;
+			});
+
+		const TArray<FString> SubStrings = GetVectorStringsFromMatrixStringLambda(MatrixNode->GetContent());
+		if (SubStrings.Num() < 4)
+		{
+			return;
+		}
+
+		const FVector U = GetVectorFromVectorStringLambda(SubStrings[0]) / 10.f;
+		const FVector V = GetVectorFromVectorStringLambda(SubStrings[1]) / 10.f;
+		const FVector W = GetVectorFromVectorStringLambda(SubStrings[2]) / 10.f;
+		const FVector O = GetVectorFromVectorStringLambda(SubStrings[3]) / 10.f;
+
+		FMatrix Matrix;
+		Matrix.M[0][0] = U.X;
+		Matrix.M[0][1] = V.X;
+		Matrix.M[0][2] = W.X;
+		Matrix.M[0][3] = O.X;
+
+		Matrix.M[1][0] = U.Y;
+		Matrix.M[1][1] = V.Y;
+		Matrix.M[1][2] = W.Y;
+		Matrix.M[1][3] = -O.Y; // From MVR's right hand to UE's left hand coordinate system
+
+		Matrix.M[2][0] = U.Z;
+		Matrix.M[2][1] = V.Z;
+		Matrix.M[2][2] = W.Z;
+		Matrix.M[2][3] = O.Z;
+
+		Matrix.M[3][0] = 0.0;
+		Matrix.M[3][1] = 0.0;
+		Matrix.M[3][2] = 0.0;
+		Matrix.M[3][3] = 1.0;
+
+		OutTransform = FTransform(Matrix.GetTransposed());
+
+		// Scale from milimeters to centimeters
+		OutTransform->SetScale3D(OutTransform->GetScale3D() * 10.f);
+
+		// GDTFs are facing down on the Z-Axis, but UE Actors are facing up 
+		const FQuat InvertUpRotationQuaternion = FQuat(OutTransform->GetRotation().GetAxisY(), PI);
+		OutTransform->SetRotation(OutTransform->GetRotation() * InvertUpRotationQuaternion);
+	}
+	
+	/** Parses the GDTFSpec from the MVR Fixture Node, returns true on success. */
+	bool ParseGDTFSpec(const FXmlNode* const InFixtureNode, FString& OutGDTFSpec)
+	{
+		check(InFixtureNode);
+
+		static const FString ChildNodeName_GDTFSpec = TEXT("GDTFSpec");
+		const FXmlNode* const GDTFSpecNode = InFixtureNode->FindChildNode(ChildNodeName_GDTFSpec);
+		if (GDTFSpecNode)
+		{
+			OutGDTFSpec = GDTFSpecNode->GetContent();
+			return true;
+		}
+
+		return false;
+	}
+
+	/** Parses the GDTFMode from the MVR Fixture Node, returns true on success. */
+	bool ParseGDTFMode(const FXmlNode* const InFixtureNode, FString& OutGDTFMode)
+	{
+		check(InFixtureNode);
+
+		static const FString ChildNodeName_GDTFMode = TEXT("GDTFMode");
+		const FXmlNode* const GDTFModeNode = InFixtureNode->FindChildNode(ChildNodeName_GDTFMode);
+		if (GDTFModeNode)
+		{
+			OutGDTFMode = GDTFModeNode->GetContent();
+			return true;
+		}
+		return true;
+	}
+
+	/** Parses Focus from the MVR Fixture Node, optional, always succeeds. */
+	void ParseFocus(const FXmlNode* const InFixtureNode, TOptional<FGuid>& OutFocus)
+	{
+		check(InFixtureNode);
+		OutFocus.Reset();
+
+		static const FString ChildNodeName_Focus = TEXT("Focus");
+		const FXmlNode* const FocusNode = InFixtureNode->FindChildNode(ChildNodeName_Focus);
+		if (FocusNode)
+		{
+			const FString FocusString = InFixtureNode->GetContent();
+			FGuid FocusGuid;
+			if (FGuid::Parse(FocusString, FocusGuid))
+			{
+				OutFocus = FocusGuid;
+			}
+		};
+	}
+
+	/** Parses Cast Shadows from the MVR Fixture Node, optional, always succeeds. */
+	void ParseCastShadow(const FXmlNode* const InFixtureNode, TOptional<bool>& OutCastShadow)
+	{
+		check(InFixtureNode);
+		OutCastShadow.Reset();
+
+		static const FString ChidNodeName_CastShadow = TEXT("CastShadow");
+		const FXmlNode* const CastShadowNode = InFixtureNode->FindChildNode(ChidNodeName_CastShadow);
+		if (CastShadowNode)
+		{
+			OutCastShadow = !(CastShadowNode->GetContent() == TEXT("0")) || !CastShadowNode->GetContent().Equals(TEXT("true"), ESearchCase::IgnoreCase);
+		}
+	}
+
+	/** Parses Position of from MVR Fixture Node, optional, always succeeds. */
+	void ParsePosition(const FXmlNode* const InFixtureNode, TOptional<FGuid>& OutPosition)
+	{
+		check(InFixtureNode);
+		OutPosition.Reset();
+
+		static const FString ChildNodenName_Position = TEXT("Position");
+		const FXmlNode* const PositionNode = InFixtureNode->FindChildNode(ChildNodenName_Position);
+		if (PositionNode)
+		{
+			FGuid PositionGuid;
+			if (FGuid::Parse(PositionNode->GetContent(), PositionGuid))
+			{
+				OutPosition = PositionGuid;
+			}
+		}
+	}
+
+	/** Parses the Fixture Id from the MVR Fixture Node, returns true on success. */
+	bool ParseFixtureId(const FXmlNode* const InFixtureNode, FString& OutFixtureId)
+	{
+		check(InFixtureNode);
+		OutFixtureId.Reset();
+
+		static const FString ChildNodenName_FixtureId = TEXT("FixtureId");
+		const FXmlNode* const FixtureIdNode = InFixtureNode->FindChildNode(ChildNodenName_FixtureId);
+		if (FixtureIdNode)
+		{
+			OutFixtureId = FixtureIdNode->GetContent();
+			return true;
+		}
+
+		return false;
+	}
+
+	/** Parses the Unit Number from the MVR Fixture Node, returns true on success. */
+	bool ParseUnitNumber(const FXmlNode* const InFixtureNode, int32& OutUnitNumber)
+	{
+		check(InFixtureNode);
+
+		static const FString ChildNodenName_UnitNumber = TEXT("FixtureId");
+		const FXmlNode* const UnitNumberNode = InFixtureNode->FindChildNode(ChildNodenName_UnitNumber);
+		if (UnitNumberNode)
+		{
+			int32 UnitNumber;
+			if (LexTryParseString(UnitNumber, *UnitNumberNode->GetContent()))
+			{
+				OutUnitNumber = UnitNumber;
+				return true;
+			}
+		}
+		return false;
+	}
+
+	/** Parses the Addresses from the MVR Fixture Node, defaulted, always succeeds */
+	bool ParseAddresses(const FXmlNode* const InFixtureNode, FDMXMVRFixtureAddresses& OutAddresses)
+	{
+		check(InFixtureNode);
+
+		static const FString ChildNodenName_Addresses = TEXT("Addresses");
+		const FXmlNode* const AddressesNode = InFixtureNode->FindChildNode(ChildNodenName_Addresses);
+
+		static const FString ChildNodenName_Address = TEXT("Address");
+		const FXmlNode* const AddressNode = AddressesNode ? AddressesNode->FindChildNode(ChildNodenName_Address) : nullptr;
+
+		if (AddressNode)
+		{
+			int32 Universe;
+			int32 Address;
+
+			auto ParseAddressStringLambda = [](const FXmlNode* const InAddressNode, int32& OutUniverse, int32& OutAddress) -> const bool
+			{
+				const FString InAddressString = InAddressNode->GetContent();
+
+				int32 SeparatorIndex = INDEX_NONE;
+				if (InAddressString.FindChar('.', SeparatorIndex))
+				{
+					if (SeparatorIndex > 0 && SeparatorIndex < InAddressString.Len() - 2)
+					{
+						const FString UniverseString = InAddressString.Left(SeparatorIndex);
+						const FString AddressString = InAddressString.RightChop(SeparatorIndex + 1);
+
+						const bool bValidUniverse = LexTryParseString(OutUniverse, *UniverseString) && OutUniverse >= 0 && OutUniverse <= DMX_MAX_UNIVERSE;
+						const bool bValidAddress = LexTryParseString(OutAddress, *UniverseString) && OutAddress >= 0 && OutAddress <= DMX_MAX_ADDRESS;
+
+						return bValidUniverse && bValidAddress;
+					}
+				}
+				else if (LexTryParseString(OutAddress, *InAddressString))
+				{
+					static const int32 UniverseSize = 512;
+
+					OutUniverse = OutAddress / UniverseSize;
+					OutAddress = OutAddress % UniverseSize;
+
+					return 
+						OutUniverse >= 0 && 
+						OutUniverse <= DMX_MAX_UNIVERSE &&
+						OutAddress >= 0 && 
+						OutAddress <= DMX_MAX_ADDRESS;
+				}
+
+				return false;
+			}(AddressNode, Universe, Address);
+			if (ParseAddressStringLambda)
+			{
+				OutAddresses.Universe = Universe;
+				OutAddresses.Address = Address;
+
+				return true;
+			}
+		}
+		return false;
+	}
+
+	/** Parses the CID Color from the MVR Fixture Node, optional, always succeeds. */
+	void ParseCIEColor(const FXmlNode* const InFixtureNode, TOptional<FDMXMVRColorCIE>& OutCIEColor)
+	{
+		check(InFixtureNode);
+
+		static const FString ChildNodenName_CIEColor = TEXT("CIEColor");
+		const FXmlNode* const CIEColorNode = InFixtureNode->FindChildNode(ChildNodenName_CIEColor);
+		if (CIEColorNode)
+		{
+			const FString CIEColorString = CIEColorNode->GetContent();
+			if (CIEColorString.IsEmpty())
+			{
+				return;
+			}
+
+			TArray<FString> CIEColorArray;
+			CIEColorString.ParseIntoArray(CIEColorArray, TEXT(","));
+
+			if (CIEColorArray.Num() != 3	)
+			{
+				return;
+			}
+
+			FDMXMVRColorCIE Result;
+			bool bSuccess = LexTryParseString<float>(Result.X, *CIEColorArray[0]);
+			bSuccess = bSuccess && LexTryParseString<float>(Result.Y, *CIEColorArray[1]);
+			bSuccess = bSuccess && LexTryParseString<uint8>(Result.YY, *CIEColorArray[2]);
+
+			if (bSuccess)
+			{
+				OutCIEColor = Result;
+			}
+		}
+	}
+
+	/** Parses the Fixture Type Id from MVR Fixture Node, optional, always succeeds. */
+	void ParseFixtureTypeId(const FXmlNode* const InFixtureNode, TOptional<int32>& OutFixtureTypeId)
+	{
+		check(InFixtureNode);
+		OutFixtureTypeId.Reset();
+
+		static const FString ChildNodenName_FixtureTypeId = TEXT("FixtureTypeId");
+		const FXmlNode* const FixtureTypeIdNode = InFixtureNode->FindChildNode(ChildNodenName_FixtureTypeId);
+		if (FixtureTypeIdNode)
+		{
+			int32 FixtureTypeId;
+			if (LexTryParseString(FixtureTypeId, *FixtureTypeIdNode->GetContent()))
+			{
+				OutFixtureTypeId = FixtureTypeId;
+			}
+		}
+	}
+
+	/** Parses the Custom Id from MVR Fixture Node, optional, always succeeds. */
+	void ParseCustomId(const FXmlNode* const InFixtureNode, TOptional<int32>& OutCustomId)
+	{
+		check(InFixtureNode);
+		OutCustomId.Reset();
+
+		static const FString ChildNodenName_Position = TEXT("Position");
+		const FXmlNode* const PositionNode = InFixtureNode->FindChildNode(ChildNodenName_Position);
+		if (PositionNode)
+		{
+			int32 FixtureTypeId;
+			if (LexTryParseString(FixtureTypeId, *PositionNode->GetContent()))
+			{
+				OutCustomId = FixtureTypeId;
+			}
+		}
+	}
+
+	/** Parses the Mapping from MVR Fixture Node, optional, always succeeds. */
+	void ParseMapping(const FXmlNode* const InFixtureNode, TOptional<FDMXMVRFixtureMapping>& OutMapping)
+	{
+		check(InFixtureNode);
+		OutMapping.Reset();
+
+		static const FString ChildNodenName_Mapping = TEXT("Position");
+		const FXmlNode* const MappingNode = InFixtureNode->FindChildNode(ChildNodenName_Mapping);
+		if (MappingNode)
+		{
+			int32 FixtureTypeId;
+			if (LexTryParseString(FixtureTypeId, *MappingNode->GetContent()))
+			{
+				static const FString ChildNodenName_LinkDef = TEXT("LinkDef");
+				const FXmlNode* const LinkDefNOde = MappingNode->FindChildNode(ChildNodenName_LinkDef);
+				FDMXMVRFixtureMapping Mapping;
+				if (LinkDefNOde && !FGuid::Parse(LinkDefNOde->GetContent(), Mapping.LinkDef))
+				{
+					return;
+				}
+
+				static const FString ChildNodenName_UX = TEXT("UX");
+				const FXmlNode* const UXNode = InFixtureNode->FindChildNode(ChildNodenName_UX);
+				int32 UX;
+				if (UXNode && !LexTryParseString(UX, *UXNode->GetContent()))
+				{
+					Mapping.UX = UX;
+					return;
+				}
+
+				static const FString ChildNodenName_UY = TEXT("UY");
+				const FXmlNode* const UYNode = InFixtureNode->FindChildNode(ChildNodenName_UY);
+				int32 UY;
+				if (UYNode && !LexTryParseString(UY, *UYNode->GetContent()))
+				{
+					return;
+				}
+				Mapping.UY = UY;
+
+				static const FString ChildNodenName_OX = TEXT("OX");
+				const FXmlNode* const OXNode = InFixtureNode->FindChildNode(ChildNodenName_OX);
+				int32 OX;
+				if (OXNode && !LexTryParseString(OX, *OXNode->GetContent()))
+				{
+					return;
+				}
+				Mapping.OX = OX;
+
+				static const FString ChildNodenName_OY = TEXT("OY");
+				const FXmlNode* const OYNode = InFixtureNode->FindChildNode(ChildNodenName_OY);
+				int32 OY;
+				if (OYNode && !LexTryParseString(OY, *OYNode->GetContent()))
+				{
+					return;
+				}
+				Mapping.OY = OY;
+
+				static const FString ChildNodenName_RZ = TEXT("RZ");
+				const FXmlNode* const RZNode = InFixtureNode->FindChildNode(ChildNodenName_RZ);
+				int32 RZ;
+				if (RZNode && !LexTryParseString(RZ, *RZNode->GetContent()))
+				{
+					return;
+				}
+				Mapping.RZ = RZ;
+
+				OutMapping = Mapping;
+			}
+		}
+	}
+
+	/** Parses the Gobo from MVR Fixture Node, optional, always succeeds. */
+	void ParseGobo(const FXmlNode* const InFixtureNode, TOptional<FDMXMVRFixtureGobo>& OutGobo)
+	{
+		check(InFixtureNode);
+		OutGobo.Reset();
+
+		static const FString ChildNodenName_Gobo = TEXT("Position");
+		const FXmlNode* const GoboNode = InFixtureNode->FindChildNode(ChildNodenName_Gobo);
+		if (GoboNode)
+		{
+			const FString& GoboResource = GoboNode->GetContent();
+			OutGobo->Value = GoboResource;
+
+			static const FString ChildNodeName_Gobo = TEXT("Rotation");
+			const FXmlNode* const GoboRotationNode = GoboNode->FindChildNode(ChildNodeName_Gobo);
+			float GoboRotation;
+			if (GoboRotationNode && LexTryParseString(GoboRotation, *GoboRotationNode->GetContent()))
+			{
+				OutGobo->Rotation = GoboRotation;
+			}
+		}
+	}
+}
+
+FDMXMVRFixture::FDMXMVRFixture()
+{}
+
+FDMXMVRFixture::FDMXMVRFixture(const FXmlNode* const FixtureNode)
+{
+	// Ensure a valid node
+	if (!ensureAlwaysMsgf(FixtureNode, TEXT("Trying to parse MVR Fixture Node, but node is null.")))
+	{
+		return;
+	}
+
+	static const FString NodeName_Fixture = TEXT("Fixture");
+	if (!ensureAlwaysMsgf(FixtureNode->GetTag() == NodeName_Fixture, TEXT("Trying to parse MVR Fixture Node, but node name doesn't match.")))
+	{
+		return;
+	}
+
+	// Parse nodes. Those that don't fail are either defaulted or optional.
+	using namespace UE::DMXRuntime::DMXMVRFixture::Private;
+	if (!ParseName(FixtureNode, Name))
+	{
+		// Non-optional
+		UE_LOG(LogDMXRuntime, Warning, TEXT("Cannot parse Fixture Node as MVR Fixture Name. Failed to parse MVR Node %s."), *FixtureNode->GetTag());
+		return;
+	}
+
+	if(!ParseUUID(FixtureNode, UUID))
+	{
+		// Non-optional
+		UE_LOG(LogDMXRuntime, Warning, TEXT("Cannot parse Fixture '%s' from MVR. Failed to parse MVR UUID."), *Name);
+		return;
+	}
+
+	ParseMatrix(FixtureNode, Transform);
+	if (Transform.IsSet() && !Transform.GetValue().IsValid())
+	{
+		// Mend invalid transforms
+		UE_LOG(LogDMXRuntime, Warning, TEXT("Imported Fixture '%s' from MVR but its Transform is not valid, using Identity instead."), *Name);
+		Transform = FTransform::Identity;
+	}
+
+	if (!ParseGDTFSpec(FixtureNode, GDTFSpec))
+	{
+		// Non-optional
+		UE_LOG(LogDMXRuntime, Warning, TEXT("Cannot parse Fixture '%s' from MVR. No valid GDTF Spec specified."), *Name);
+		return;
+	}
+
+	if (!ParseGDTFMode(FixtureNode, GDTFMode))
+	{
+		// Non-optional
+		UE_LOG(LogDMXRuntime, Warning, TEXT("Cannot parse Fixture '%s' from MVR. No valid GDTF Mode specified."), *Name);
+		return;
+	}
+
+	ParseFocus(FixtureNode, Focus);
+	ParseCastShadow(FixtureNode, bCastShadows);
+	ParsePosition(FixtureNode, Position);
+	
+	if (!ParseFixtureId(FixtureNode, FixtureId))
+	{
+		// Non-optional, mend with empty string
+		UE_LOG(LogDMXRuntime, Warning, TEXT("Imported Fixture '%s' from MVR. No valid Fixture Id specified. Fixture Id set to Fixture Name."), *Name);
+		FixtureId = TEXT("");
+	}
+
+	if (!ParseUnitNumber(FixtureNode, UnitNumber))
+	{
+		// Non-optional, but defaulted by standard
+		UE_LOG(LogDMXRuntime, Warning, TEXT("Imported Fixture '%s' from MVR. No valid Unit Number specified. Unit Number set to ''0'."), *Name);
+		UnitNumber = 0;
+	}
+
+	if (!ParseAddresses(FixtureNode, Addresses))
+	{
+		// Non-optional
+		UE_LOG(LogDMXRuntime, Warning, TEXT("Cannot parse Fixture '%s' from MVR. No valid Addresses specified."), *Name);
+		return;
+	}
+
+	ParseCIEColor(FixtureNode, CIEColor);
+	ParseFixtureTypeId(FixtureNode, FixtureTypeId);
+	ParseCustomId(FixtureNode, CustomId);
+	ParseMapping(FixtureNode, Mapping);
+	ParseGobo(FixtureNode, Gobo);
+}
+
+bool FDMXMVRFixture::IsValid() const
+{
+	if (!UUID.IsValid())
+	{
+		return false;
+	}
+	
+	if (Name.IsEmpty())
+	{
+		return false;
+	}
+
+	if (GDTFSpec.IsEmpty())
+	{
+		return false;
+	}
+
+	if (GDTFMode.IsEmpty())
+	{
+		return false;
+	}
+
+	if (Addresses.Universe < 0 || Addresses.Universe > DMX_MAX_UNIVERSE)
+	{
+		return false;
+	}
+
+	if (Addresses.Address < 1 || Addresses.Address > DMX_MAX_ADDRESS)
+	{
+		return false;
+	}
+
+	return true;
+}
+
+FDMXMVRGeneralSceneDescription::FDMXMVRGeneralSceneDescription()
+{}
+
+FDMXMVRGeneralSceneDescription::FDMXMVRGeneralSceneDescription(TSharedRef<FXmlFile> GeneralSceneDescription)
+{
+	ParseMVRFixtures(GeneralSceneDescription);
+}
+
+void FDMXMVRGeneralSceneDescription::ParseMVRFixtures(const TSharedRef<FXmlFile>& GeneralSceneDescription)
+{
+	if (const FXmlNode* RootNode = GeneralSceneDescription->GetRootNode())
+	{
+		static const FString NodeName_Scene = TEXT("Scene");
+		if (const FXmlNode* SceneNode = RootNode->FindChildNode(NodeName_Scene))
+		{
+			static const FString NodeName_Layers = TEXT("Layers");
+			if (const FXmlNode* LayersNode = SceneNode->FindChildNode(NodeName_Layers))
+			{
+				for (const FXmlNode* LayerNode : LayersNode->GetChildrenNodes())
+				{
+					static const FString NodeName_Layer = TEXT("Layer");
+					if (LayerNode->GetTag() == NodeName_Layer)
+					{
+						for (const FXmlNode* ChildListNode : LayerNode->GetChildrenNodes())
+						{
+							static const FString NodeName_ChildList = TEXT("ChildList");
+							if (ChildListNode->GetTag() == NodeName_ChildList)
+							{
+								for (const FXmlNode* FixtureNode : ChildListNode->GetChildrenNodes())
+								{
+									static const FString NodeName_Fixture = TEXT("Fixture");
+									if (FixtureNode->GetTag() == NodeName_Fixture)
+									{
+										FDMXMVRFixture MVRFixture(FixtureNode);
+										if (MVRFixture.IsValid())
+										{
+											MVRFixtures.Add(MoveTemp(MVRFixture));
+										}
+									}
+								}
+							}
+						}
+					}
+				}
+			}
+		}
+	}
+}
diff --git a/Engine/Plugins/VirtualProduction/DMX/DMXEngine/Source/DMXRuntime/Private/MVR/DMXMVRSceneActor.cpp b/Engine/Plugins/VirtualProduction/DMX/DMXEngine/Source/DMXRuntime/Private/MVR/DMXMVRSceneActor.cpp
new file mode 100644
--- /dev/null
+++ b/Engine/Plugins/VirtualProduction/DMX/DMXEngine/Source/DMXRuntime/Private/MVR/DMXMVRSceneActor.cpp
@@ -0,0 +1,246 @@
+// Copyright Epic Games, Inc. All Rights Reserved.
+
+#include "MVR/DMXMVRSceneActor.h"
+
+#include "DMXRuntimeLog.h"
+#include "Game/DMXComponent.h"
+#include "Library/DMXEntityFixturePatch.h"
+#include "Library/DMXLibrary.h"
+#include "MVR/DMXMVRFixtureActorLibrary.h"
+#include "MVR/DMXMVRAssetUserData.h"
+#include "MVR/DMXMVRFixtureActorInterface.h"
+
+#include "DatasmithAssetUserData.h"
+#include "Components/SceneComponent.h"
+
+#if WITH_EDITOR
+#include "Editor.h"
+#endif 
+
+
+ADMXMVRSceneActor::ADMXMVRSceneActor()
+{
+#if WITH_EDITOR
+	if (HasAnyFlags(RF_ClassDefaultObject))
+	{
+		return;
+	}
+
+	FEditorDelegates::MapChange.AddUObject(this, &ADMXMVRSceneActor::OnMapChange);
+
+	if (GEngine)
+	{
+		GEngine->OnLevelActorDeleted().AddUObject(this, &ADMXMVRSceneActor::OnActorDeleted);
+	}
+
+	if (GEditor)
+	{
+		GEditor->GetEditorSubsystem<UImportSubsystem>()->OnAssetPostImport.AddUObject(this, &ADMXMVRSceneActor::OnAssetPostImport);
+	}
+#endif // WITH_EDITOR
+
+	MVRSceneRoot = CreateDefaultSubobject<USceneComponent>("MVRSceneRoot");
+	SetRootComponent(MVRSceneRoot);
+}
+
+ADMXMVRSceneActor::~ADMXMVRSceneActor()
+{
+#if WITH_EDITOR
+	if (HasAnyFlags(RF_ClassDefaultObject))
+	{
+		return;
+	}
+
+	FEditorDelegates::MapChange.RemoveAll(this);
+
+	if (GEngine)
+	{
+		GEngine->OnLevelActorDeleted().RemoveAll(this);
+	}
+
+	if (GEditor)
+	{
+		GEditor->GetEditorSubsystem<UImportSubsystem>()->OnAssetPostImport.RemoveAll(this);
+	}
+#endif // WITH_EDITOR
+}
+
+void ADMXMVRSceneActor::PostLoad()
+{
+	Super::PostLoad();
+	SetMVRUUIDsForRelatedActors();
+}
+
+void ADMXMVRSceneActor::PostRegisterAllComponents()
+{
+	Super::PostRegisterAllComponents();
+
+	// If the actor was created as a Datasmith Element, set the library from there
+	const FString DMXLibraryPathString = UDatasmithAssetUserData::GetDatasmithUserDataValueForKey(this, TEXT("DMXLibraryPath"));
+	if (!DMXLibraryPathString.IsEmpty() && !DMXLibrary)
+	{
+		const FSoftObjectPath DMXLibraryPath(DMXLibraryPathString);
+		UObject* NewDMXLibraryObject = DMXLibraryPath.TryLoad();
+		if (UDMXLibrary* NewDMXLibrary = Cast<UDMXLibrary>(NewDMXLibraryObject))
+		{
+			SetDMXLibrary(NewDMXLibrary);
+		}
+	}
+
+	SetMVRUUIDsForRelatedActors();
+}
+
+void ADMXMVRSceneActor::SetMVRUUIDsForRelatedActors()
+{
+	for (const TTuple<FGuid, TSoftObjectPtr<AActor>>& MVRUUIDToActorPair : MVRUUIDToRelatedActorMap)
+	{
+		if (AActor* Actor = MVRUUIDToActorPair.Value.Get())
+		{			
+			const FString MVRUUID = UDMXMVRAssetUserData::GetMVRAssetUserDataValueForkey(Actor, UDMXMVRAssetUserData::MVRUUIDMetaDataKey);
+			if (MVRUUID.IsEmpty())
+			{
+				if (!UDMXMVRAssetUserData::SetMVRAssetUserDataValueForKey(Actor, UDMXMVRAssetUserData::MVRUUIDMetaDataKey, MVRUUIDToActorPair.Key.ToString()))
+				{
+					UE_LOG(LogDMXRuntime, Warning, TEXT("Actor %s is referenced by MVR scene but doesn't have an MVR UUID."), *Actor->GetName());
+				}
+			}
+		}
+	}
+}
+
+void ADMXMVRSceneActor::SetDMXLibrary(UDMXLibrary* NewDMXLibrary)
+{
+	if (!ensureAlwaysMsgf(!DMXLibrary, TEXT("Tried to set the DMXLibrary for %s, but it already has one set. Changing the library is not supported."), *GetName()))
+	{
+		return;
+	}
+
+	UWorld* World = GetWorld();
+	if (!World)
+	{
+		return;
+	}
+
+	if (!NewDMXLibrary)
+	{
+		return;
+	}
+
+	if (NewDMXLibrary == DMXLibrary)
+	{
+		return;
+	}
+
+	DMXLibrary = NewDMXLibrary;
+
+	const TSharedRef<FDMXMVRFixtureActorLibrary> MVRFixtureActorLibrary = MakeShared<FDMXMVRFixtureActorLibrary>();
+	const TArray<UDMXEntityFixturePatch*> FixturePatches = DMXLibrary->GetEntitiesTypeCast<UDMXEntityFixturePatch>();
+	for (UDMXEntityFixturePatch* FixturePatch : FixturePatches)
+	{
+		UClass* ActorClass = MVRFixtureActorLibrary->FindMostAppropriateActorClassForPatch(FixturePatch);
+		if (!ActorClass)
+		{
+			continue;
+		}
+
+		const TArray<UDMXMVRFixtureInstance*>& MVRFixturesInstances = FixturePatch->GetMVRFixtureInstances();
+		for (const UDMXMVRFixtureInstance* MVRFixtureInstance : MVRFixturesInstances)
+		{
+			const FDMXMVRFixture& MVRFixture = MVRFixtureInstance->GetMVRFixture();
+			if (!ensureAlwaysMsgf(MVRFixture.UUID.IsValid(), TEXT("Tried to set MVR Fixture %s for MVR Scene Actor %s, but the MVR UUID is invalid."), *MVRFixture.Name, *GetName()))
+			{
+				continue;
+			}
+
+			const FTransform Transform = MVRFixture.Transform.IsSet() ? MVRFixture.Transform.GetValue() : FTransform::Identity;
+			AActor* ChildActor = World->SpawnActor<AActor>(ActorClass, Transform);
+			if (!ChildActor)
+			{
+				continue;
+			}
+
+			ChildActor->RegisterAllComponents();
+			USceneComponent* RootComponentOfChildActor = ChildActor->GetRootComponent();
+			if (!RootComponentOfChildActor)
+			{
+				continue;
+			}
+			RootComponentOfChildActor->AttachToComponent(MVRSceneRoot, FAttachmentTransformRules::KeepRelativeTransform);
+			
+			// Set the patch either via the interface or via a present DMX Component.
+			// Prefer the interface way as it may further customize how the patch is set.
+			if (IDMXMVRFixtureActorInterface* MVRFixtureActorInterface = Cast<IDMXMVRFixtureActorInterface>(ChildActor))
+			{
+				MVRFixtureActorInterface->Execute_OnMVRSetFixturePatch(ChildActor, FixturePatch);
+			}
+			else if(UActorComponent* Component = ChildActor->GetComponentByClass(UDMXComponent::StaticClass()))
+			{
+				CastChecked<UDMXComponent>(Component)->SetFixturePatch(FixturePatch);
+			}
+
+			MVRUUIDToRelatedActorMap.Add(MVRFixture.UUID, ChildActor);
+		}
+	}
+}
+
+#if WITH_EDITOR
+void ADMXMVRSceneActor::OnMapChange(uint32 MapEventFlags)
+{
+	// Whenever a sub-level is loaded, we need to apply the fix
+	if (MapEventFlags == MapChangeEventFlags::NewMap)
+	{
+		SetMVRUUIDsForRelatedActors();
+	}
+}
+#endif // WITH_EDITOR
+
+#if WITH_EDITOR
+void ADMXMVRSceneActor::OnActorDeleted(AActor* ActorDeleted)
+{
+	const FString MVRUUIDString = UDMXMVRAssetUserData::GetMVRAssetUserDataValueForkey(ActorDeleted, UDMXMVRAssetUserData::MVRUUIDMetaDataKey);
+	FGuid MVRUUID;
+	if (FGuid::Parse(MVRUUIDString, MVRUUID))
+	{
+		TSoftObjectPtr<AActor>* RelatedActorPtr = MVRUUIDToRelatedActorMap.Find(MVRUUID);
+		if (RelatedActorPtr)
+		{
+			AActor* RelatedActor = RelatedActorPtr->Get();
+			if (RelatedActor == ActorDeleted)
+			{
+				// This will add this actor to the transaction if there is one currently recording
+				Modify();
+
+				RelatedActorPtr->Reset();
+			}
+		}
+	}
+}
+#endif // WITH_EDITOR
+
+#if WITH_EDITOR
+void ADMXMVRSceneActor::OnAssetPostImport(UFactory* InFactory, UObject* ActorAdded)
+{
+	for (TObjectIterator<AActor> It; It; ++It)
+	{
+		AActor* Actor = *It;
+
+		const FString MVRUUIDString = UDMXMVRAssetUserData::GetMVRAssetUserDataValueForkey(Actor, UDMXMVRAssetUserData::MVRUUIDMetaDataKey);
+		FGuid MVRUUID;
+		if (FGuid::Parse(MVRUUIDString, MVRUUID))
+		{
+			TSoftObjectPtr< AActor >* RelatedActorPtr = MVRUUIDToRelatedActorMap.Find(MVRUUID);
+			if (RelatedActorPtr)
+			{
+				AActor* RelatedActor = RelatedActorPtr->Get();
+				if (!RelatedActor)
+				{
+					// This will add this actor to the transaction if there is one currently recording
+					Modify();
+
+					*RelatedActorPtr = Actor;
+				}
+			}
+		}
+	}
+}
+#endif // WITH_EDITOR
diff --git a/Engine/Plugins/VirtualProduction/DMX/DMXEngine/Source/DMXRuntime/Private/MVR/DMXMVRUnzip.cpp b/Engine/Plugins/VirtualProduction/DMX/DMXEngine/Source/DMXRuntime/Private/MVR/DMXMVRUnzip.cpp
new file mode 100644
--- /dev/null
+++ b/Engine/Plugins/VirtualProduction/DMX/DMXEngine/Source/DMXRuntime/Private/MVR/DMXMVRUnzip.cpp
@@ -0,0 +1,246 @@
+// Copyright Epic Games, Inc. All Rights Reserved.
+
+#include "MVR/DMXMVRUnzip.h"
+
+#include "DMXRuntimeLog.h"
+
+#include "GenericPlatform/GenericPlatform.h"
+#include "Misc/Base64.h"
+#include "Misc/Compression.h"
+#include "Misc/FileHelper.h"
+#include "Misc/Paths.h"
+#include "Serialization/JsonSerializer.h"
+
+
+TSharedPtr<FDMXMVRUnzip> FDMXMVRUnzip::CreateFromData(const uint8* DataPtr, const int64 DataNum)
+{
+	const TSharedRef<FDMXMVRUnzip> NewMVRUnzip = MakeShared<FDMXMVRUnzip>();
+	if (NewMVRUnzip->InitializeFromDataInternal(DataPtr, DataNum))
+	{
+		return NewMVRUnzip;
+	}
+
+	return nullptr;
+}
+
+TSharedPtr<FDMXMVRUnzip> FDMXMVRUnzip::CreateFromFile(const FString& FilePathAndName)
+{
+	const TSharedRef<FDMXMVRUnzip> NewMVRUnzip = MakeShared<FDMXMVRUnzip>();
+	if (NewMVRUnzip->InitializeFromFileInternal(FilePathAndName))
+	{
+		return NewMVRUnzip;
+	}
+
+	return nullptr;
+}
+
+bool FDMXMVRUnzip::GetFileContent(const FString& FilenameInZip, TArray64<uint8>& OutData)
+{
+	uint32* Offset = OffsetsMap.Find(FilenameInZip);
+	if (!Offset)
+	{
+		return false;
+	}
+
+	constexpr uint64 LocalEntryMinSize = 30;
+
+	if (*Offset + LocalEntryMinSize > Data.Num())
+	{
+		return false;
+	}
+
+	uint16 Compression = 0;
+	uint32 CompressedSize;
+	uint32 UncompressedSize = 0;
+	uint16 FilenameLen = 0;
+	uint16 ExtraFieldLen = 0;
+
+	constexpr uint8 CompressionOffset = 8;
+	Data.Seek(*Offset + CompressionOffset);
+	Data << Compression;
+
+	constexpr uint8 CompressionSizeOffset = 18;
+	Data.Seek(*Offset + CompressionSizeOffset);
+	Data << CompressedSize;
+	Data << UncompressedSize;
+	Data << FilenameLen;
+	Data << ExtraFieldLen;
+
+	if (*Offset + LocalEntryMinSize + FilenameLen + ExtraFieldLen + CompressedSize > Data.Num())
+	{
+		return false;
+	}
+
+	if (Compression == 8)
+	{
+		OutData.AddUninitialized(UncompressedSize);
+		if (!FCompression::UncompressMemory(NAME_Zlib, OutData.GetData(), UncompressedSize, Data.GetData() + *Offset + LocalEntryMinSize + FilenameLen + ExtraFieldLen, CompressedSize, COMPRESS_NoFlags, -15))
+		{
+			return false;
+		}
+	}
+	else if (Compression == 0 && CompressedSize == UncompressedSize)
+	{
+		OutData.Append(Data.GetData() + *Offset + LocalEntryMinSize + FilenameLen + ExtraFieldLen, UncompressedSize);
+	}
+	else
+	{
+		return false;
+	}
+
+	return true;
+}
+
+bool FDMXMVRUnzip::Contains(const FString& FilenameInZip) const
+{
+	return OffsetsMap.Contains(FilenameInZip);
+}
+
+FString FDMXMVRUnzip::GetFirstFilenameByExtension(const FString& Extension) const
+{
+	for (const TPair<FString, uint32>& Pair : OffsetsMap)
+	{
+		if (Pair.Key.EndsWith(Extension, ESearchCase::IgnoreCase))
+		{
+			return Pair.Key;
+		}
+	}
+
+	return "";
+}
+
+FDMXMVRUnzip::FDMXScopedUnzipToTempFile::FDMXScopedUnzipToTempFile(const TSharedRef<FDMXMVRUnzip>& MVRUnzip, const FString& FilenameInZip)
+{
+	if (!ensureMsgf(MVRUnzip->Contains(FilenameInZip), TEXT("Tried to unzip %s from MVR, but zip does not contain the file.")))
+	{
+		return;
+	}
+	TArray64<uint8> FileData;
+	MVRUnzip->GetFileContent(FilenameInZip, FileData);
+
+	const FString Directory = FPaths::EngineSavedDir() / TEXT("DMXMVR");
+	TempFilePathAndName = FPaths::ConvertRelativePathToFull(Directory / FilenameInZip);
+
+	FFileHelper::SaveArrayToFile(FileData, *TempFilePathAndName);
+}
+
+FDMXMVRUnzip::FDMXScopedUnzipToTempFile::~FDMXScopedUnzipToTempFile()
+{
+	// Delete the extracted file
+	IFileManager& FileManager = IFileManager::Get();
+	constexpr bool bRequireExists = true;
+	constexpr bool bEvenIfReadOnly = false;
+	constexpr bool bQuiet = true;
+	FileManager.Delete(*TempFilePathAndName, bRequireExists, bEvenIfReadOnly, bQuiet);
+}
+
+bool FDMXMVRUnzip::InitializeFromDataInternal(const uint8* DataPtr, const int64 DataNum)
+{
+	Data.Append(DataPtr, DataNum);
+
+	// Step 0: retrieve the trailer magic
+	TArray<uint8> Magic;
+	bool bIndexFound = false;
+	uint64 Index = 0;
+	for (Index = Data.Num() - 1; Index >= 0; Index--)
+	{
+		Magic.Insert(Data[Index], 0);
+		if (Magic.Num() == 4)
+		{
+			if (Magic[0] == 0x50 && Magic[1] == 0x4b && Magic[2] == 0x05 && Magic[3] == 0x06)
+			{
+				bIndexFound = true;
+				break;
+			}
+			Magic.Pop();
+		}
+	}
+
+	if (!bIndexFound)
+	{
+		return false;
+	}
+
+	uint16 DiskEntries = 0;
+	uint16 TotalEntries = 0;
+	uint32 CentralDirectorySize = 0;
+	uint32 CentralDirectoryOffset = 0;
+	uint16 CommentLen = 0;
+
+	constexpr uint64 TrailerMinSize = 22;
+	constexpr uint64 CentralDirectoryMinSize = 46;
+
+	if (Index + TrailerMinSize > Data.Num())
+	{
+		return false;
+	}
+
+	// Skip signature and disk data
+	Data.Seek(Index + 8);
+	Data << DiskEntries;
+	Data << TotalEntries;
+	Data << CentralDirectorySize;
+	Data << CentralDirectoryOffset;
+	Data << CommentLen;
+
+	uint16 DirectoryEntries = FMath::Min(DiskEntries, TotalEntries);
+
+	for (uint16 DirectoryIndex = 0; DirectoryIndex < DirectoryEntries; DirectoryIndex++)
+	{
+		if (CentralDirectoryOffset + CentralDirectoryMinSize > Data.Num())
+		{
+			return false;
+		}
+
+		uint16 FilenameLen = 0;
+		uint16 ExtraFieldLen = 0;
+		uint16 EntryCommentLen = 0;
+		uint32 EntryOffset = 0;
+
+		// seek to FilenameLen
+		Data.Seek(CentralDirectoryOffset + 28);
+		Data << FilenameLen;
+		Data << ExtraFieldLen;
+		Data << EntryCommentLen;
+		// seek to EntryOffset
+		Data.Seek(CentralDirectoryOffset + 42);
+		Data << EntryOffset;
+
+		if (CentralDirectoryOffset + CentralDirectoryMinSize + FilenameLen + ExtraFieldLen + EntryCommentLen > Data.Num())
+		{
+			return false;
+		}
+
+		TArray64<uint8> FilenameBytes;
+		FilenameBytes.Append(Data.GetData() + CentralDirectoryOffset + CentralDirectoryMinSize, FilenameLen);
+		FilenameBytes.Add(0);
+
+		FString Filename = FString(UTF8_TO_TCHAR(FilenameBytes.GetData()));
+
+		OffsetsMap.Add(Filename, EntryOffset);
+
+		CentralDirectoryOffset += CentralDirectoryMinSize + FilenameLen + ExtraFieldLen + EntryCommentLen;
+	}
+
+	return true;
+}
+
+bool FDMXMVRUnzip::InitializeFromFileInternal(const FString& FilePathAndName)
+{
+	if (!FPaths::FileExists(FilePathAndName))
+	{
+		UE_LOG(LogDMXRuntime, Error, TEXT("Failed to load '%s'. File does not exists."), *FilePathAndName);
+		return false;
+	}
+
+	TArray<uint8> NewData;
+	if (!FFileHelper::LoadFileToArray(NewData, *FilePathAndName))
+	{
+		UE_LOG(LogDMXRuntime, Error, TEXT("Faild to load '%s'. Access denied."), *FilePathAndName);
+		return false;
+	}
+
+	return InitializeFromDataInternal(NewData.GetData(), NewData.Num());
+}
+
+#undef LOCTEXT_NAMESPACE
diff --git a/Engine/Plugins/VirtualProduction/DMX/DMXEngine/Source/DMXRuntime/Public/Library/DMXEntityFixturePatch.h b/Engine/Plugins/VirtualProduction/DMX/DMXEngine/Source/DMXRuntime/Public/Library/DMXEntityFixturePatch.h
--- a/Engine/Plugins/VirtualProduction/DMX/DMXEngine/Source/DMXRuntime/Public/Library/DMXEntityFixturePatch.h
+++ b/Engine/Plugins/VirtualProduction/DMX/DMXEngine/Source/DMXRuntime/Public/Library/DMXEntityFixturePatch.h
@@ -4,6 +4,7 @@
 
 #include "DMXProtocolCommon.h"
 #include "Library/DMXEntity.h"
+#include "MVR/DMXMVRGeneralSceneDescription.h"
 
 #include "DMXAttribute.h"
 #include "DMXProtocolCommon.h"
@@ -27,6 +28,35 @@ struct FPropertyChangedEvent;
 DECLARE_MULTICAST_DELEGATE_OneParam(FDMXOnFixturePatchChangedDelegate, const UDMXEntityFixturePatch* /** ChangedFixturePatch */);
 
 
+/**
+ * An instance of an MVR Fixture in a Fixture Patch.
+ */
+UCLASS(DefaultToInstanced)
+class DMXRUNTIME_API UDMXMVRFixtureInstance
+	: public UObject
+{
+	GENERATED_BODY()
+
+public:
+	/** Creates a new instance from an MVR Fixture */
+	static UDMXMVRFixtureInstance* CreateFromMVRFixture(const FDMXMVRFixture& InMVRFixture);
+
+	//~ Begin UObject interface
+	virtual void Serialize(FArchive& Ar) override;
+	//~ End UObject interface
+
+	/** Sets the MVR Fixture. Note, the UUID has to match the previous MVR Fixture. */
+	void SetMVRFixture(const FDMXMVRFixture& InMVRFixture);
+
+	/** Returns the MVR Fixture */
+	FORCEINLINE const FDMXMVRFixture& GetMVRFixture() const { return MVRFixture; }
+
+private:
+	/** The mvr fixture data */
+	FDMXMVRFixture MVRFixture;
+};
+
+
 /** Parameters to construct a Fixture Patch. */
 USTRUCT(BlueprintType)
 struct DMXRUNTIME_API FDMXEntityFixturePatchConstructionParams
@@ -50,6 +80,7 @@ struct DMXRUNTIME_API FDMXEntityFixturePatchConstructionParams
 	int32 StartingAddress = 1;
 };
 
+
 /** 
  * A DMX fixture patch that can be patch to channels in a DMX Universe via the DMX Library Editor. 
  * 
@@ -205,6 +236,15 @@ public:
 	/** Returns custom tags defined for the patch */
 	FORCEINLINE const TArray<FName>& GetCustomTags() const { return CustomTags; }
 
+	/** Adds an MVR Fixture to the patch */
+	void AddMVRFixtureInstance(const FDMXMVRFixture& MVRFixtureToAdd);
+
+	/** Makes an array of MVR Fixtures from the MVR fixture instances of this patch */
+	TArray<FDMXMVRFixture> MakeMVRFixtures() const;
+
+	/** Returns the MVR Fixture instances that use this patch */
+	FORCEINLINE const TArray<UDMXMVRFixtureInstance*>& GetMVRFixtureInstances() const { return MVRFixtureInstances; }
+
 #if WITH_EDITOR
 	/** Property name getters. When accessing the this way, use with care. The patch will need to update its cache after using property setters. Call RebuildCache to do that. */
 	static FName GetUniverseIDPropertyNameChecked() { return GET_MEMBER_NAME_CHECKED(UDMXEntityFixturePatch, UniverseID); }
@@ -213,6 +253,7 @@ public:
 	static FName GetAutoStartingAddressPropertyNameChecked() { return GET_MEMBER_NAME_CHECKED(UDMXEntityFixturePatch, AutoStartingAddress); }
 	static FName GetParentFixtureTypeTemplatePropertyNameChecked() { return GET_MEMBER_NAME_CHECKED(UDMXEntityFixturePatch, ParentFixtureTypeTemplate); }
 	static FName GetActiveModePropertyNameChecked() { return GET_MEMBER_NAME_CHECKED(UDMXEntityFixturePatch, ActiveMode); }
+	static FName GetMVRFixtureInstancesPropertyNameChecked() { return GET_MEMBER_NAME_CHECKED(UDMXEntityFixturePatch, MVRFixtureInstances); }
 #endif // WITH_EDITOR
 
 protected:
@@ -240,6 +281,10 @@ protected:
 	UPROPERTY(EditAnywhere, BlueprintReadOnly, Category = "Fixture Patch")
 	int32 ActiveMode;
 
+	/** MVR Fixtures that use this patch */
+	UPROPERTY(BlueprintReadWrite, EditAnywhere, Category = "Fixture Patch")
+	TArray<UDMXMVRFixtureInstance*> MVRFixtureInstances;
+
 	/** Delegate broadcast when a Fixture Patch changed */
 	static FDMXOnFixturePatchChangedDelegate OnFixturePatchChangedDelegate;
 
diff --git a/Engine/Plugins/VirtualProduction/DMX/DMXEngine/Source/DMXRuntime/Public/Library/DMXEntityFixtureType.h b/Engine/Plugins/VirtualProduction/DMX/DMXEngine/Source/DMXRuntime/Public/Library/DMXEntityFixtureType.h
--- a/Engine/Plugins/VirtualProduction/DMX/DMXEngine/Source/DMXRuntime/Public/Library/DMXEntityFixtureType.h
+++ b/Engine/Plugins/VirtualProduction/DMX/DMXEngine/Source/DMXRuntime/Public/Library/DMXEntityFixtureType.h
@@ -313,11 +313,14 @@ public:
 
 
 public:
-#if WITH_EDITOR
+	/** Sets the GDTF for this Fixture Type. Note, this clears the fixture type so that the Fixture Type reflects the GDTF. Merging is not supported. */
 	UFUNCTION(BlueprintCallable, Category = "Fixture Settings")
+	void SetGDTF(UDMXImportGDTF* DMXImportAsset);
+
+	UE_DEPRECATED(5.1, "Deprecated in favor of more clearly named member function SetGDTFAsset. Only supports the functional UDMXImportGDTF now.")
+	UFUNCTION(BlueprintCallable, Category = "Fixture Settings", Meta = (DeprecatedFunction, DeprecationMessage = "Deprecated 5.1. Please use the more clearly named SetGDTFAsset instead."))
 	void SetModesFromDMXImport(UDMXImport* DMXImportAsset);
-#endif // WITH_EDITOR
-	
+
 	/** Returns a delegate that is and should be broadcast whenever a Fixture Type changed */
 	static FDMXOnFixtureTypeChangedDelegate& GetOnFixtureTypeChanged();
 
diff --git a/Engine/Plugins/VirtualProduction/DMX/DMXEngine/Source/DMXRuntime/Public/Library/DMXLibrary.h b/Engine/Plugins/VirtualProduction/DMX/DMXEngine/Source/DMXRuntime/Public/Library/DMXLibrary.h
--- a/Engine/Plugins/VirtualProduction/DMX/DMXEngine/Source/DMXRuntime/Public/Library/DMXLibrary.h
+++ b/Engine/Plugins/VirtualProduction/DMX/DMXEngine/Source/DMXRuntime/Public/Library/DMXLibrary.h
@@ -7,6 +7,7 @@
 #include "IO/DMXInputPortReference.h"
 #include "IO/DMXOutputPortReference.h"
 #include "Library/DMXEntity.h"
+#include "MVR/DMXMVRGeneralSceneDescription.h"
 
 #include "Templates/SubclassOf.h"
 #include "Misc/Guid.h"
@@ -185,7 +186,6 @@ public:
 	UE_DEPRECATED(5.0, "Deprecated in favor of the more expressive UDMXLibrary::GetOnEntitiesAdded() and UDMXLibrary::GetOnEntitiesRemoved() that also forwards the added resp. removed entities")
 	FOnEntitiesUpdated_DEPRECATED& GetOnEntitiesUpdated();
 
-public:
 	/** Returns all local Universe IDs in Ports */
 	TSet<int32> GetAllLocalUniversesIDsInPorts() const;
 
@@ -201,6 +201,9 @@ public:
 	/** Updates the ports from what's set in the Input and Output Port References arrays */
 	void UpdatePorts();
 
+	/** Creates an MVR General Scene Description from the Library's MVR data */
+	FDMXMVRGeneralSceneDescription MakeGeneralSceneDescription() const;
+
 	/** Returns the Entity that was last added to the Library */
 	FORCEINLINE TWeakObjectPtr<UDMXEntity> GetLastAddedEntity() const { return LastAddedEntity; }
 
diff --git a/Engine/Plugins/VirtualProduction/DMX/DMXEngine/Source/DMXRuntime/Public/MVR/DMXMVRAssetUserData.h b/Engine/Plugins/VirtualProduction/DMX/DMXEngine/Source/DMXRuntime/Public/MVR/DMXMVRAssetUserData.h
new file mode 100644
--- /dev/null
+++ b/Engine/Plugins/VirtualProduction/DMX/DMXEngine/Source/DMXRuntime/Public/MVR/DMXMVRAssetUserData.h
@@ -0,0 +1,34 @@
+// Copyright Epic Games, Inc. All Rights Reserved.
+
+#pragma once
+
+#include "CoreMinimal.h"
+#include "Engine/AssetUserData.h"
+
+#include "DMXMVRAssetUserData.generated.h"
+
+
+/** Asset user data for Actors in an MVR Scene */
+UCLASS()
+class DMXRUNTIME_API UDMXMVRAssetUserData
+	: public UAssetUserData
+{
+	GENERATED_BODY()
+
+public:
+	/** Returns MVR Asset User Data for the Actor, or nullptr if there is no asset user data */
+	static UDMXMVRAssetUserData* GetMVRAssetUserData(AActor* Actor);
+
+	/** Returns the Value in Meta data given a Key, or an empty String if the Key doesn't exist */
+	static FString GetMVRAssetUserDataValueForkey(AActor* Actor, const FName Key);
+
+	/** Sets the Value for Key in Meta Data. Creates the Key if it doesn't exist. */
+	static bool SetMVRAssetUserDataValueForKey(AActor* Actor, FName Key, const FString& Value);
+
+	/** MVR Meta Data for the Actor */
+	UPROPERTY(VisibleAnywhere, BlueprintReadWrite, AdvancedDisplay, Category = "MVR")
+	TMap<FName, FString> MetaData;
+
+	/** The key for the MVR UUID meta data */
+	static const FName MVRUUIDMetaDataKey;
+};
diff --git a/Engine/Plugins/VirtualProduction/DMX/DMXEngine/Source/DMXRuntime/Public/MVR/DMXMVRFixtureActorInterface.h b/Engine/Plugins/VirtualProduction/DMX/DMXEngine/Source/DMXRuntime/Public/MVR/DMXMVRFixtureActorInterface.h
new file mode 100644
--- /dev/null
+++ b/Engine/Plugins/VirtualProduction/DMX/DMXEngine/Source/DMXRuntime/Public/MVR/DMXMVRFixtureActorInterface.h
@@ -0,0 +1,35 @@
+// Copyright Epic Games, Inc. All Rights Reserved.
+
+#pragma once
+
+#include "UObject/Interface.h"
+
+#include "DMXMVRFixtureActorInterface.generated.h"
+
+class UDMXEntityFixturePatch;
+
+
+/** 
+ * When implemented in an actor, MVR will find the Actor and consider it when auto-selecting Fixtures.
+ */
+UINTERFACE(BlueprintType)
+class DMXRUNTIME_API UDMXMVRFixtureActorInterface
+	: public UInterface
+{
+	GENERATED_BODY()
+
+};
+
+class DMXRUNTIME_API IDMXMVRFixtureActorInterface
+{
+	GENERATED_BODY()
+
+public:
+	/** Should return the DMX Attributes the Actor supports */
+	UFUNCTION(BlueprintNativeEvent, BlueprintCallable, Category = "DMX", Meta = (DisplayName = "On MVR Get Supported DMX Attributes"))
+	void OnMVRGetSupportedDMXAttributes(TArray<FName>& OutAttributeNames, TArray<FName>& OutMatrixAttributeNames) const;
+
+	/** Should set the Fixture Patch the MVR Fixture uses */
+	UFUNCTION(BlueprintNativeEvent, BlueprintCallable, CallInEditor, Category = "DMX", Meta = (DisplayName = "On MVR Set Fixture Patch"))
+	void OnMVRSetFixturePatch(UDMXEntityFixturePatch* FixturePatch);
+};
diff --git a/Engine/Plugins/VirtualProduction/DMX/DMXEngine/Source/DMXRuntime/Public/MVR/DMXMVRFixtureActorLibrary.h b/Engine/Plugins/VirtualProduction/DMX/DMXEngine/Source/DMXRuntime/Public/MVR/DMXMVRFixtureActorLibrary.h
new file mode 100644
--- /dev/null
+++ b/Engine/Plugins/VirtualProduction/DMX/DMXEngine/Source/DMXRuntime/Public/MVR/DMXMVRFixtureActorLibrary.h
@@ -0,0 +1,40 @@
+// Copyright Epic Games, Inc. All Rights Reserved.
+
+#pragma once
+
+#include "CoreMinimal.h"
+
+#include "Templates/SharedPointer.h"
+#include "UObject/GCObject.h"
+#include "UObject/StrongObjectPtr.h"
+
+class UDMXEntityFixturePatch;
+
+class AActor;
+struct FAssetData;
+
+
+class DMXRUNTIME_API FDMXMVRFixtureActorLibrary
+	: public FGCObject
+	, public TSharedFromThis<FDMXMVRFixtureActorLibrary>
+{
+public:
+	/** Constructor */
+	FDMXMVRFixtureActorLibrary();
+
+	/** Returns the Actor that can represent the patch best */
+	UClass* FindMostAppropriateActorClassForPatch(const UDMXEntityFixturePatch* const Patch) const;
+
+protected:
+	//~ Begin FGCObject interface
+	virtual void AddReferencedObjects(FReferenceCollector& Collector) override;
+	virtual FString GetReferencerName() const override
+	{
+		return TEXT("FDMXMVRFixtureActorLibrary");
+	}
+	//~ End FGCObject interface
+
+private:
+	/** Array of actors that implement the MVR Actor Interface */
+	TArray<AActor*> MVRActors;
+};
diff --git a/Engine/Plugins/VirtualProduction/DMX/DMXEngine/Source/DMXRuntime/Public/MVR/DMXMVRGeneralSceneDescription.h b/Engine/Plugins/VirtualProduction/DMX/DMXEngine/Source/DMXRuntime/Public/MVR/DMXMVRGeneralSceneDescription.h
new file mode 100644
--- /dev/null
+++ b/Engine/Plugins/VirtualProduction/DMX/DMXEngine/Source/DMXRuntime/Public/MVR/DMXMVRGeneralSceneDescription.h
@@ -0,0 +1,199 @@
+// Copyright Epic Games, Inc. All Rights Reserved.
+
+#pragma once
+
+#include "CoreMinimal.h"
+
+#include "Library/DMXImport.h"
+#include "Misc/Optional.h"
+#include "Serialization/Archive.h"
+
+class FXmlNode;
+class FXmlFile;
+
+
+struct DMXRUNTIME_API FDMXMVRFixtureAddresses
+{
+	int32 Address;
+	int32 Universe;
+
+	/** Serializes an MVR Fixture Address from or into an archive */
+	friend FArchive& operator<<(FArchive& Ar, FDMXMVRFixtureAddresses& Addresses)
+	{
+		Ar << Addresses.Address;
+		Ar << Addresses.Universe;
+
+		return Ar;
+	}
+};
+
+struct DMXRUNTIME_API FDMXMVRColorCIE
+{
+	float X;
+	float Y;
+	uint8 YY;
+
+	/** Serializes an MVR Color CIE from or into an archive */
+	friend FArchive& operator<<(FArchive& Ar, FDMXMVRColorCIE& ColorCIE)
+	{
+		Ar << ColorCIE.X;
+		Ar << ColorCIE.Y;
+		Ar << ColorCIE.YY;
+
+		return Ar;
+	}
+};
+
+struct DMXRUNTIME_API FDMXMVRFixtureMapping
+{
+	/** The unique identifier of the MappingDefinition node that will be the source of the mapping. */
+	FGuid LinkDef;
+
+	/** The offset in pixels in x direction from top left corner of the source that will be used for the mapped object. */
+	TOptional<int32> UX;
+
+	/** The offset in pixels in y direction from top left corner of the source that will be used for the mapped object. */
+	TOptional<int32> UY;
+
+	/** The size in pixels in x direction from top left of the starting point. */
+	TOptional<int32> OX;
+
+	/** The size in pixels in y direction from top left of the starting point. */
+	TOptional<int32> OY;
+
+	/** The rotation around the middle point of the defined rectangle in degree. Positive direction is counter cock wise. */
+	TOptional<int32> RZ;
+
+	/** Serializes an MVR Fixture Mapping from or into an archive */
+	friend FArchive& operator<<(FArchive& Ar, FDMXMVRFixtureMapping& Mapping)
+	{
+		Ar << Mapping.LinkDef;
+		Ar << Mapping.UX;
+		Ar << Mapping.UY;
+		Ar << Mapping.OX;
+		Ar << Mapping.OY;
+		Ar << Mapping.RZ;
+
+		return Ar;
+	}
+};
+
+struct DMXRUNTIME_API FDMXMVRFixtureGobo
+{
+	/** The node value is the Gobo used for the fixture. The image ressource must apply to the GDTF standard. Use a FileName to specify. */
+	FString Value;
+
+	/** The roation of the Gobo in degree. */
+	float Rotation = 0.f;
+
+	/** Serializes an MVR Fixture Gobo from or into an archive */
+	friend FArchive& operator<<(FArchive& Ar, FDMXMVRFixtureGobo& Gobo)
+	{
+		Ar << Gobo.Value;
+		Ar << Gobo.Rotation;
+
+		return Ar;
+	}
+};
+
+struct DMXRUNTIME_API FDMXMVRFixture
+{
+public:
+	FDMXMVRFixture();
+	FDMXMVRFixture(const FXmlNode* const FixtureNode);
+
+	/** Returns true if this is a valid MVR Fixture */
+	bool IsValid() const;
+
+	/** The unique identifier of the object. */
+	FGuid UUID;
+
+	/** The name of the object. */
+	FString Name;
+
+	/** The location of the object inside the parent coordinate system. */
+	TOptional<FTransform> Transform;
+	
+	/** The name of the file containing the GDTF information for this light fixture. */
+	FString GDTFSpec;
+
+	/** The name of the used DMX mode. This has to match the name of a DMXMode in the GDTF file. */
+	FString GDTFMode;
+
+	/** A focus point reference that this lighting fixture aims at if this reference exists. */
+	TOptional<FGuid> Focus;
+
+	/** Defines if a Object cast Shadows. */
+	TOptional<bool> bCastShadows;
+
+	/** A position reference that this lighting fixture belongs to if this reference exists. */
+	TOptional<FGuid> Position;
+
+	/** The Fixture Id of the lighting fixture. This is the short name of the fixture. */
+	FString FixtureId;
+
+	/** The unit number of the lighting fixture in a position. */
+	int32 UnitNumber = 0;
+
+	/** The container for DMX Addresses for this fixture. */
+	FDMXMVRFixtureAddresses Addresses;
+
+	/** A color assigned to a fixture. If it is not defined, there is no color for the fixture. */
+	TOptional<FDMXMVRColorCIE> CIEColor;
+
+	/** The Fixture Type ID is a value that can be used as a short name of the Fixture Type. This does not have to be unique. The default value is 0. */
+	TOptional<int32> FixtureTypeId;
+
+	/** The Custom ID is a value that can be used as a short name of the Fixture Instance. This does not have to be unique. The default value is 0. */
+	TOptional<int32> CustomId;
+
+	/** The container for Mappings for this fixture. */
+	TOptional<FDMXMVRFixtureMapping> Mapping;
+
+	/** The Gobo used for the fixture. The image ressource must apply to the GDTF standard. */
+	TOptional<FDMXMVRFixtureGobo> Gobo;
+
+	/** Serializes an MVR Fixture from or into an archive */
+	friend FArchive& operator<<(FArchive& Ar, FDMXMVRFixture& MVRFixture)
+	{
+		Ar << MVRFixture.UUID;
+		Ar << MVRFixture.Name;
+		Ar << MVRFixture.Transform;
+		Ar << MVRFixture.GDTFSpec;
+		Ar << MVRFixture.GDTFMode;
+		Ar << MVRFixture.Focus;
+		Ar << MVRFixture.bCastShadows;
+		Ar << MVRFixture.Position;
+		Ar << MVRFixture.FixtureId;
+		Ar << MVRFixture.UnitNumber;	
+		Ar << MVRFixture.Addresses;
+		Ar << MVRFixture.CIEColor;
+		Ar << MVRFixture.FixtureTypeId;
+		Ar << MVRFixture.CustomId;
+		Ar << MVRFixture.Mapping;
+		Ar << MVRFixture.Gobo;
+		
+		return Ar;		 
+	}
+};
+
+struct DMXRUNTIME_API FDMXMVRGeneralSceneDescription
+{
+	FDMXMVRGeneralSceneDescription();
+	FDMXMVRGeneralSceneDescription(TSharedRef<FXmlFile> GeneralSceneDescription);
+
+	/** The DMMX interactable fixture contained in an MVR file */
+	TArray<FDMXMVRFixture> MVRFixtures;
+
+	/** Serializes an General Scene Description from or into an archive */
+	friend FArchive& operator<<(FArchive& Ar, FDMXMVRGeneralSceneDescription& GeneralSceneDescription)
+	{
+		Ar << GeneralSceneDescription.MVRFixtures;
+
+		return Ar;
+	}
+
+private:
+	/** Gets the fixture nodes from the General Scene Description */
+	void ParseMVRFixtures(const TSharedRef<FXmlFile>& GeneralSceneDescription);
+};
diff --git a/Engine/Plugins/VirtualProduction/DMX/DMXEngine/Source/DMXRuntime/Public/MVR/DMXMVRSceneActor.h b/Engine/Plugins/VirtualProduction/DMX/DMXEngine/Source/DMXRuntime/Public/MVR/DMXMVRSceneActor.h
new file mode 100644
--- /dev/null
+++ b/Engine/Plugins/VirtualProduction/DMX/DMXEngine/Source/DMXRuntime/Public/MVR/DMXMVRSceneActor.h
@@ -0,0 +1,90 @@
+// Copyright Epic Games, Inc. All Rights Reserved.
+
+#pragma once
+
+#include "CoreMinimal.h"
+#include "GameFramework/Actor.h"
+
+#include "DMXMVRSceneActor.generated.h"
+
+class UDMXImportGDTF;
+class UDMXLibrary;
+class UDMXMVRSceneComponent;
+
+class UFactory;
+
+
+USTRUCT(BlueprintType)
+struct FDMXMVRSceneActorGDTFToActorClassPair
+{
+	GENERATED_BODY()
+
+#if WITH_EDITORONLY_DATA
+	/** The GDTF Asset */
+	UPROPERTY(EditAnywhere, Category = "MVR")
+	TSoftObjectPtr<UDMXImportGDTF> GDTF;
+
+	UPROPERTY(EditAnywhere, Category = "MVR", Meta = (MustImplement = "DMXMVRFixtureActorInterface"))
+	TSoftClassPtr<AActor> Actor;
+#endif // WITH_EDITORONLY_DATA
+};
+
+
+UCLASS(NotBlueprintable)
+class DMXRUNTIME_API ADMXMVRSceneActor
+	: public AActor
+{
+	GENERATED_BODY()
+
+public:
+	/** Constructor */
+	ADMXMVRSceneActor();
+
+	/** Destructor */
+	~ADMXMVRSceneActor();
+
+	//~ Begin AActor interface
+	virtual void PostLoad() override;
+	virtual void PostRegisterAllComponents() override;
+	//~ End AActor interface
+
+	/** Sets the dmx library for this MVR actor. Should only be called once, further calls will have no effect and hit an ensure condition */
+	void SetDMXLibrary(UDMXLibrary* NewDMXLibrary);
+
+	/** Returns the DMX Library of this MVR Scene Actor */
+	FORCEINLINE UDMXLibrary* GetDMXLibrary() const { return DMXLibrary; }
+
+#if WITH_EDITORONLY_DATA
+	/** The actor class that is spawned for a specific GDTF by default (can be overriden per MVR UUID, see below) */
+	UPROPERTY(EditAnywhere, Category = "MVR", Meta = (DispayName = "Default Actor used for GDTF"))
+	TArray<FDMXMVRSceneActorGDTFToActorClassPair> GDTFToDefaultActorClassArray;
+#endif // WITH_EDITORONLY_DATA
+
+private:
+	/** Set MVR UUIDs for related Actors */
+	void SetMVRUUIDsForRelatedActors();
+
+#if WITH_EDITOR
+	/** Called when a sub-level is loaded */
+	void OnMapChange(uint32 MapChangeFlags);
+
+	/** Called when an actor got deleted in editor */
+	void OnActorDeleted(AActor* ActorDeleted);
+
+	/** Called when an asset was imported */
+	void OnAssetPostImport(UFactory* InFactory, UObject* ActorAdded);
+#endif // WITH_EDITOR
+
+private:
+	/** The DMX Library this Scene Actor uses */
+	UPROPERTY(BlueprintReadOnly, VisibleAnywhere, Category = "MVR", Meta = (AllowPrivateAccess = true))
+	TObjectPtr<UDMXLibrary> DMXLibrary;
+
+	/** The actors that created along with this scene */
+	UPROPERTY(VisibleAnywhere, Category = "Actor", AdvancedDisplay, Meta = (AllowPrivateAccess = true))
+	TMap<FGuid, TSoftObjectPtr<AActor>> MVRUUIDToRelatedActorMap;
+
+	/** The root component to which all actors are attached initially */
+	UPROPERTY(VisibleAnywhere, Category = "Actor", AdvancedDisplay, Meta = (AllowPrivateAccess = true))
+	USceneComponent* MVRSceneRoot;
+};
diff --git a/Engine/Plugins/VirtualProduction/DMX/DMXEngine/Source/DMXRuntime/Public/MVR/DMXMVRUnzip.h b/Engine/Plugins/VirtualProduction/DMX/DMXEngine/Source/DMXRuntime/Public/MVR/DMXMVRUnzip.h
new file mode 100644
--- /dev/null
+++ b/Engine/Plugins/VirtualProduction/DMX/DMXEngine/Source/DMXRuntime/Public/MVR/DMXMVRUnzip.h
@@ -0,0 +1,47 @@
+// Copyright Epic Games, Inc. All Rights Reserved.
+
+#pragma once
+
+#include "CoreMinimal.h"
+#include "Serialization/ArrayReader.h"
+
+
+/** Provides access to the zipped content of an MVR file */
+class DMXRUNTIME_API FDMXMVRUnzip
+{
+public:
+	/** Creates a zip reader from data. Returns nullptr if not a valid zip */
+	static TSharedPtr<FDMXMVRUnzip> CreateFromData(const uint8* DataPtr, const int64 DataNum);
+
+	/** Creates a zip reader from file. Returns nullptr if not a valid zip */
+	static TSharedPtr<FDMXMVRUnzip> CreateFromFile(const FString& FilePathAndName);
+
+	/** Gets the data of a file in the zip */
+	bool GetFileContent(const FString& FilenameInZip, TArray64<uint8>& OutData);
+
+	/** Returns true if the zip contains the file with name */
+	bool Contains(const FString& FilenameInZip) const;
+
+	/** Returns the first file that matches the extension */
+	FString GetFirstFilenameByExtension(const FString& Extension) const;
+
+	/** Helper to unzip a file within the zip as temp file. Deletes the file when running out of scope. */
+	struct DMXRUNTIME_API FDMXScopedUnzipToTempFile
+	{
+		FDMXScopedUnzipToTempFile(const TSharedRef<FDMXMVRUnzip>& MVRUnzip, const FString& FilenameInZip);
+		~FDMXScopedUnzipToTempFile();
+
+		/** Absolute File Path and Name of the Temp File */
+		FString TempFilePathAndName;
+	};
+
+private:
+	/** Constructs itself from data. Returns true on success. */
+	bool InitializeFromDataInternal(const uint8* DataPtr, const int64 DataNum);
+
+	/** Initializes itself from a file. Returns true on success. */
+	bool InitializeFromFileInternal(const FString& FilePathAndName);
+
+	TMap<FString, uint32> OffsetsMap;
+	FArrayReader Data;
+};
diff --git a/Engine/Plugins/VirtualProduction/DMX/DMXFixtures/Source/DMXFixtures/DMXFixtures.Build.cs b/Engine/Plugins/VirtualProduction/DMX/DMXFixtures/Source/DMXFixtures/DMXFixtures.Build.cs
--- a/Engine/Plugins/VirtualProduction/DMX/DMXFixtures/Source/DMXFixtures/DMXFixtures.Build.cs
+++ b/Engine/Plugins/VirtualProduction/DMX/DMXFixtures/Source/DMXFixtures/DMXFixtures.Build.cs
@@ -8,7 +8,15 @@ public class DMXFixtures : ModuleRules
 	{
 	   PCHUsage = PCHUsageMode.UseExplicitOrSharedPCHs;
 
-       PublicDependencyModuleNames.AddRange(new string[] { "Core", "CoreUObject", "Engine", "DMXRuntime", "DMXProtocol", "RenderCore", "RHI", "ProceduralMeshComponent" });
+       PublicDependencyModuleNames.AddRange(new string[] { 
+		   "Core",
+		   "CoreUObject",
+		   "DMXRuntime", 
+		   "DMXProtocol",
+		   "Engine",
+		   "RenderCore", 
+		   "RHI", 
+		   "ProceduralMeshComponent" });
 
        PrivateDependencyModuleNames.AddRange(new string[] {  });
 
diff --git a/Engine/Plugins/VirtualProduction/DMX/DMXFixtures/Source/DMXFixtures/Private/DMXFixtureActor.cpp b/Engine/Plugins/VirtualProduction/DMX/DMXFixtures/Source/DMXFixtures/Private/DMXFixtureActor.cpp
--- a/Engine/Plugins/VirtualProduction/DMX/DMXFixtures/Source/DMXFixtures/Private/DMXFixtureActor.cpp
+++ b/Engine/Plugins/VirtualProduction/DMX/DMXFixtures/Source/DMXFixtures/Private/DMXFixtureActor.cpp
@@ -6,6 +6,7 @@
 
 #include "Components/StaticMeshComponent.h"
 
+
 DECLARE_CYCLE_STAT(TEXT("FixtureActor Push Normalized Values"), STAT_FixtureActorPushNormalizedValuesPerAttribute, STATGROUP_DMX);
 
 ADMXFixtureActor::ADMXFixtureActor()
@@ -55,6 +56,22 @@ ADMXFixtureActor::ADMXFixtureActor()
 	HasBeenInitialized = false;
 }
 
+void ADMXFixtureActor::OnMVRGetSupportedDMXAttributes_Implementation(TArray<FName>& OutAttributeNames, TArray<FName>& OutMatrixAttributeNames) const
+{
+	for (UDMXFixtureComponent* DMXFixtureComponent : TInlineComponentArray<UDMXFixtureComponent*>(this))
+	{
+		TArray<FName> SupportedAttributeNamesOfComponent;
+		DMXFixtureComponent->GetSupportedDMXAttributes(SupportedAttributeNamesOfComponent);
+
+		OutAttributeNames.Append(SupportedAttributeNamesOfComponent);
+	}
+}
+
+void ADMXFixtureActor::OnMVRSetFixturePatch_Implementation(UDMXEntityFixturePatch* FixturePatch)
+{
+	DMX->SetFixturePatch(FixturePatch);
+}
+
 #if WITH_EDITOR
 void ADMXFixtureActor::PostEditChangeProperty(FPropertyChangedEvent& PropertyChangedEvent)
 {
diff --git a/Engine/Plugins/VirtualProduction/DMX/DMXFixtures/Source/DMXFixtures/Private/DMXFixtureActorMatrix.cpp b/Engine/Plugins/VirtualProduction/DMX/DMXFixtures/Source/DMXFixtures/Private/DMXFixtureActorMatrix.cpp
--- a/Engine/Plugins/VirtualProduction/DMX/DMXFixtures/Source/DMXFixtures/Private/DMXFixtureActorMatrix.cpp
+++ b/Engine/Plugins/VirtualProduction/DMX/DMXFixtures/Source/DMXFixtures/Private/DMXFixtureActorMatrix.cpp
@@ -71,6 +71,38 @@ void ADMXFixtureActorMatrix::PostEditChangeProperty(FPropertyChangedEvent& Prope
 }
 #endif
 
+void ADMXFixtureActorMatrix::OnMVRGetSupportedDMXAttributes_Implementation(TArray<FName>& OutAttributeNames, TArray<FName>& OutMatrixAttributeNames) const
+{
+	Super::OnMVRGetSupportedDMXAttributes_Implementation(OutAttributeNames, OutMatrixAttributeNames);
+	
+	for (UDMXFixtureComponent* DMXComponent : TInlineComponentArray<UDMXFixtureComponent*>(this))
+	{
+		if (UDMXFixtureComponentColor* ColorComponent = Cast<UDMXFixtureComponentColor>(DMXComponent))
+		{
+			OutMatrixAttributeNames.Add(ColorComponent->DMXChannel1.Name);
+			OutMatrixAttributeNames.Add(ColorComponent->DMXChannel2.Name);
+			OutMatrixAttributeNames.Add(ColorComponent->DMXChannel3.Name);
+			OutMatrixAttributeNames.Add(ColorComponent->DMXChannel4.Name);
+		}
+		else if (UDMXFixtureComponentSingle* SingleComponent = Cast<UDMXFixtureComponentSingle>(DMXComponent))
+		{
+			// Single channel component - hardcoded and limited to specific names, as on PushFixtureMatrixCellData
+			if (SingleComponent->DMXChannel.Name.Name == FName("Dimmer"))
+			{
+				OutMatrixAttributeNames.Add(FName("Dimmer"));
+			}
+			else if (SingleComponent->DMXChannel.Name.Name == FName("Pan"))
+			{
+				OutMatrixAttributeNames.Add(FName("Pan"));
+			}
+			else if (SingleComponent->DMXChannel.Name.Name == FName("Tilt"))
+			{
+				OutMatrixAttributeNames.Add(FName("Tilt"));
+			}
+		}
+	}
+}
+
 void ADMXFixtureActorMatrix::InitializeMatrixFixture()
 {
 	UDMXEntityFixturePatch* FixturePatch = DMX->GetFixturePatch();
diff --git a/Engine/Plugins/VirtualProduction/DMX/DMXFixtures/Source/DMXFixtures/Private/DMXFixtureComponentColor.cpp b/Engine/Plugins/VirtualProduction/DMX/DMXFixtures/Source/DMXFixtures/Private/DMXFixtureComponentColor.cpp
--- a/Engine/Plugins/VirtualProduction/DMX/DMXFixtures/Source/DMXFixtures/Private/DMXFixtureComponentColor.cpp
+++ b/Engine/Plugins/VirtualProduction/DMX/DMXFixtures/Source/DMXFixtures/Private/DMXFixtureComponentColor.cpp
@@ -9,6 +9,17 @@ UDMXFixtureComponentColor::UDMXFixtureComponentColor()
 	PrimaryComponentTick.bCanEverTick = false;
 }
 
+void UDMXFixtureComponentColor::GetSupportedDMXAttributes_Implementation(TArray<FName>& OutAttributeNames)
+{
+	if (bIsEnabled)
+	{
+		OutAttributeNames.Add(DMXChannel1.Name);
+		OutAttributeNames.Add(DMXChannel2.Name);
+		OutAttributeNames.Add(DMXChannel3.Name);
+		OutAttributeNames.Add(DMXChannel4.Name);
+	}
+}
+
 void UDMXFixtureComponentColor::Initialize()
 {
 	Super::Initialize();
diff --git a/Engine/Plugins/VirtualProduction/DMX/DMXFixtures/Source/DMXFixtures/Private/DMXFixtureComponentDouble.cpp b/Engine/Plugins/VirtualProduction/DMX/DMXFixtures/Source/DMXFixtures/Private/DMXFixtureComponentDouble.cpp
--- a/Engine/Plugins/VirtualProduction/DMX/DMXFixtures/Source/DMXFixtures/Private/DMXFixtureComponentDouble.cpp
+++ b/Engine/Plugins/VirtualProduction/DMX/DMXFixtures/Source/DMXFixtures/Private/DMXFixtureComponentDouble.cpp
@@ -11,6 +11,15 @@ UDMXFixtureComponentDouble::UDMXFixtureComponentDouble()
 	PrimaryComponentTick.bCanEverTick = false;
 }
 
+void UDMXFixtureComponentDouble::GetSupportedDMXAttributes_Implementation(TArray<FName>& OutAttributeNames)
+{
+	if (bIsEnabled)
+	{
+		OutAttributeNames.Add(DMXChannel1.Name.Name);
+		OutAttributeNames.Add(DMXChannel2.Name.Name);
+	}
+}
+
 void UDMXFixtureComponentDouble::Initialize()
 {
 	Super::Initialize();
diff --git a/Engine/Plugins/VirtualProduction/DMX/DMXFixtures/Source/DMXFixtures/Private/DMXFixtureComponentSingle.cpp b/Engine/Plugins/VirtualProduction/DMX/DMXFixtures/Source/DMXFixtures/Private/DMXFixtureComponentSingle.cpp
--- a/Engine/Plugins/VirtualProduction/DMX/DMXFixtures/Source/DMXFixtures/Private/DMXFixtureComponentSingle.cpp
+++ b/Engine/Plugins/VirtualProduction/DMX/DMXFixtures/Source/DMXFixtures/Private/DMXFixtureComponentSingle.cpp
@@ -10,6 +10,14 @@ UDMXFixtureComponentSingle::UDMXFixtureComponentSingle()
 	PrimaryComponentTick.bCanEverTick = false;
 }
 
+void UDMXFixtureComponentSingle::GetSupportedDMXAttributes_Implementation(TArray<FName>& OutAttributeNames)
+{
+	if (bIsEnabled)
+	{
+		OutAttributeNames.Add(DMXChannel.Name.Name);
+	}
+}
+
 void UDMXFixtureComponentSingle::Initialize()
 {
 	Super::Initialize();
diff --git a/Engine/Plugins/VirtualProduction/DMX/DMXFixtures/Source/DMXFixtures/Public/DMXFixtureActor.h b/Engine/Plugins/VirtualProduction/DMX/DMXFixtures/Source/DMXFixtures/Public/DMXFixtureActor.h
--- a/Engine/Plugins/VirtualProduction/DMX/DMXFixtures/Source/DMXFixtures/Public/DMXFixtureActor.h
+++ b/Engine/Plugins/VirtualProduction/DMX/DMXFixtures/Source/DMXFixtures/Public/DMXFixtureActor.h
@@ -2,18 +2,21 @@
 
 #pragma once
 
-#include "CoreMinimal.h"
 #include "DMXFixtureComponent.h"
 #include "DMXFixtureComponentDouble.h"
 #include "DMXFixtureComponentSingle.h"
 #include "DMXFixtureComponentColor.h"
-#include "GameFramework/Actor.h"
-#include "Game/DMXComponent.h"
-#include "Components/SpotLightComponent.h"
-#include "Materials/MaterialInstanceDynamic.h"
-#include "Components/ArrowComponent.h"
 #include "DMXProtocolTypes.h"
 #include "Library/DMXEntityFixturePatch.h"
+#include "MVR/DMXMVRFixtureActorInterface.h"
+
+#include "CoreMinimal.h"
+#include "Components/ArrowComponent.h"
+#include "Components/SpotLightComponent.h"
+#include "Game/DMXComponent.h"
+#include "GameFramework/Actor.h"
+#include "Materials/MaterialInstanceDynamic.h"
+
 #include "DMXFixtureActor.generated.h"
 
 
@@ -28,18 +31,27 @@ enum EDMXFixtureQualityLevel
 };
 
 UCLASS()
-class DMXFIXTURES_API ADMXFixtureActor : public AActor
+class DMXFIXTURES_API ADMXFixtureActor 
+	: public AActor
+	, public IDMXMVRFixtureActorInterface
 {
 	GENERATED_BODY()
 
 protected:
+	//~ Begin AActor Interface
 #if WITH_EDITOR
 	virtual void PostEditChangeProperty(FPropertyChangedEvent& PropertyChangedEvent) override;
 #endif
+	//~ End AActor interface
 
 public:
 	ADMXFixtureActor();
 
+	//~ Begin DMXMVRFixtureActorInterface interface
+	virtual void OnMVRGetSupportedDMXAttributes_Implementation(TArray<FName>& OutAttributeNames, TArray<FName>& OutMatrixAttributeNames) const override;
+	virtual void OnMVRSetFixturePatch_Implementation(UDMXEntityFixturePatch* FixturePatch) override;
+	//~ End DMXMVRFixtureActorInterface interface
+
 	bool HasBeenInitialized;
 	float LensRadius;
 	void FeedFixtureData();
diff --git a/Engine/Plugins/VirtualProduction/DMX/DMXFixtures/Source/DMXFixtures/Public/DMXFixtureActorMatrix.h b/Engine/Plugins/VirtualProduction/DMX/DMXFixtures/Source/DMXFixtures/Public/DMXFixtureActorMatrix.h
--- a/Engine/Plugins/VirtualProduction/DMX/DMXFixtures/Source/DMXFixtures/Public/DMXFixtureActorMatrix.h
+++ b/Engine/Plugins/VirtualProduction/DMX/DMXFixtures/Source/DMXFixtures/Public/DMXFixtureActorMatrix.h
@@ -23,6 +23,9 @@ protected:
 #endif
 
 public:
+	//~ Begin DMXGDTFActorInterface interface
+	virtual void OnMVRGetSupportedDMXAttributes_Implementation(TArray<FName>& OutAttributeNames, TArray<FName>& OutMatrixAttributeNames) const override;
+	//~ End DMXGDTFActorInterface interface
 
 	TArray<uint8> MatrixData;
 	int MatrixDataSize;
diff --git a/Engine/Plugins/VirtualProduction/DMX/DMXFixtures/Source/DMXFixtures/Public/DMXFixtureComponent.h b/Engine/Plugins/VirtualProduction/DMX/DMXFixtures/Source/DMXFixtures/Public/DMXFixtureComponent.h
--- a/Engine/Plugins/VirtualProduction/DMX/DMXFixtures/Source/DMXFixtures/Public/DMXFixtureComponent.h
+++ b/Engine/Plugins/VirtualProduction/DMX/DMXFixtures/Source/DMXFixtures/Public/DMXFixtureComponent.h
@@ -41,7 +41,7 @@ protected:
 	/** Initializes the cells for the fixture */
 	void InitCells(int NumCells);
 
-public:	
+public:
 	UDMXFixtureComponent();
 
 	/** Initializes the component */
@@ -52,7 +52,7 @@ public:
 	virtual void SetCurrentCell(int Index);
 
 	/** If used within a DMX Fixture Actor or Fixture Matrix Actor, the component only receives data when set to true. Else needs be implemented in blueprints. */
-	UPROPERTY(EditAnywhere, Category = "DMX Parameters", meta= (DisplayPriority = 0))
+	UPROPERTY(EditAnywhere, Category = "DMX Parameters", meta = (DisplayPriority = 0))
 	bool bIsEnabled = true;
 
 	/** Value changes smaller than this threshold are ignored */
@@ -87,6 +87,11 @@ public:
 	UFUNCTION(BlueprintCallable, BlueprintImplementableEvent, Category = "DMX")
 	void InitializeComponent();
 
+	/** Should be implemented to let other objects (e.g. datasmith) know which attributes the component can handle */
+	UFUNCTION(BlueprintCallable, BlueprintNativeEvent, Category = "DMX")
+	void GetSupportedDMXAttributes(TArray<FName>& OutAttributeNames);
+	virtual void GetSupportedDMXAttributes_Implementation(TArray<FName>& OutAttributeNames) {};
+
 	/** Applies the speed scale property */
 	void ApplySpeedScale();
 
diff --git a/Engine/Plugins/VirtualProduction/DMX/DMXFixtures/Source/DMXFixtures/Public/DMXFixtureComponentColor.h b/Engine/Plugins/VirtualProduction/DMX/DMXFixtures/Source/DMXFixtures/Public/DMXFixtureComponentColor.h
--- a/Engine/Plugins/VirtualProduction/DMX/DMXFixtures/Source/DMXFixtures/Public/DMXFixtureComponentColor.h
+++ b/Engine/Plugins/VirtualProduction/DMX/DMXFixtures/Source/DMXFixtures/Public/DMXFixtureComponentColor.h
@@ -18,6 +18,10 @@ class DMXFIXTURES_API UDMXFixtureComponentColor : public UDMXFixtureComponent
 public:
 	UDMXFixtureComponentColor();
 
+	//~ Begin DMXFixtureComponent interface
+	virtual void GetSupportedDMXAttributes_Implementation(TArray<FName>& OutAttributeNames) override;
+	//~ End DMXFixtureComponent interface
+
 	UPROPERTY(BlueprintReadWrite, EditAnywhere, Category = "DMX Channel")
 	FDMXAttributeName DMXChannel1;
 
diff --git a/Engine/Plugins/VirtualProduction/DMX/DMXFixtures/Source/DMXFixtures/Public/DMXFixtureComponentDouble.h b/Engine/Plugins/VirtualProduction/DMX/DMXFixtures/Source/DMXFixtures/Public/DMXFixtureComponentDouble.h
--- a/Engine/Plugins/VirtualProduction/DMX/DMXFixtures/Source/DMXFixtures/Public/DMXFixtureComponentDouble.h
+++ b/Engine/Plugins/VirtualProduction/DMX/DMXFixtures/Source/DMXFixtures/Public/DMXFixtureComponentDouble.h
@@ -14,7 +14,11 @@ class DMXFIXTURES_API UDMXFixtureComponentDouble : public UDMXFixtureComponent
 
 public:
 	UDMXFixtureComponentDouble();
-		
+
+	//~ Begin DMXFixtureComponent interface
+	virtual void GetSupportedDMXAttributes_Implementation(TArray<FName>& OutAttributeNames) override;
+	//~ End DMXFixtureComponent interface
+
 	/** The first dmx attribute the component handles */
 	UPROPERTY(EditAnywhere, BlueprintReadOnly, Category = "DMX Channels")
 	FDMXChannelData DMXChannel1;
diff --git a/Engine/Plugins/VirtualProduction/DMX/DMXFixtures/Source/DMXFixtures/Public/DMXFixtureComponentSingle.h b/Engine/Plugins/VirtualProduction/DMX/DMXFixtures/Source/DMXFixtures/Public/DMXFixtureComponentSingle.h
--- a/Engine/Plugins/VirtualProduction/DMX/DMXFixtures/Source/DMXFixtures/Public/DMXFixtureComponentSingle.h
+++ b/Engine/Plugins/VirtualProduction/DMX/DMXFixtures/Source/DMXFixtures/Public/DMXFixtureComponentSingle.h
@@ -14,6 +14,10 @@ class DMXFIXTURES_API UDMXFixtureComponentSingle : public UDMXFixtureComponent
 
 public:
 	UDMXFixtureComponentSingle();
+
+	//~ Begin DMXFixtureComponent interface
+	virtual void GetSupportedDMXAttributes_Implementation(TArray<FName>& OutAttributeNames) override;
+	//~ End DMXFixtureComponent interface
 		
 	UPROPERTY(EditAnywhere, BlueprintReadOnly, Category = "DMX Channels")
 	FDMXChannelData DMXChannel;
