DMX - Fix CIS Static Analysis Win64 (MSVC)

diff --git a/Engine/Plugins/VirtualProduction/DMX/DMXEngine/Source/DMXRuntime/Private/MVR/DMXMVRUnzip.cpp b/Engine/Plugins/VirtualProduction/DMX/DMXEngine/Source/DMXRuntime/Private/MVR/DMXMVRUnzip.cpp
--- a/Engine/Plugins/VirtualProduction/DMX/DMXEngine/Source/DMXRuntime/Private/MVR/DMXMVRUnzip.cpp
+++ b/Engine/Plugins/VirtualProduction/DMX/DMXEngine/Source/DMXRuntime/Private/MVR/DMXMVRUnzip.cpp
@@ -141,7 +141,7 @@ bool FDMXMVRUnzip::InitializeFromDataInternal(const uint8* DataPtr, const int64
 	// Step 0: retrieve the trailer magic
 	TArray<uint8> Magic;
 	bool bIndexFound = false;
-	uint64 Index = 0;
+	int64 Index = 0;
 	for (Index = Data.Num() - 1; Index >= 0; Index--)
 	{
 		Magic.Insert(Data[Index], 0);
