DMX - Fix CIS Static Analysis Win64 (PVS-Studio)

diff --git a/Engine/Plugins/VirtualProduction/DMX/DMXEngine/Source/DMXEditor/Private/Factories/DMXLibraryFromMVRFactory.cpp b/Engine/Plugins/VirtualProduction/DMX/DMXEngine/Source/DMXEditor/Private/Factories/DMXLibraryFromMVRFactory.cpp
--- a/Engine/Plugins/VirtualProduction/DMX/DMXEngine/Source/DMXEditor/Private/Factories/DMXLibraryFromMVRFactory.cpp
+++ b/Engine/Plugins/VirtualProduction/DMX/DMXEngine/Source/DMXEditor/Private/Factories/DMXLibraryFromMVRFactory.cpp
@@ -108,7 +108,6 @@ UDMXLibrary* UDMXLibraryFromMVRFactory::CreateDMXLibraryAsset(UObject* Parent, E
 		return nullptr;
 	}
 
-	check(sizeof(this) == sizeof(int64));
 	ImportSubsystem->BroadcastAssetPostImport(this, NewDMXLibrary);
 	
 	return NewDMXLibrary;
diff --git a/Engine/Plugins/VirtualProduction/DMX/DMXEngine/Source/DMXRuntime/Private/MVR/DMXMVRGeneralSceneDescription.cpp b/Engine/Plugins/VirtualProduction/DMX/DMXEngine/Source/DMXRuntime/Private/MVR/DMXMVRGeneralSceneDescription.cpp
--- a/Engine/Plugins/VirtualProduction/DMX/DMXEngine/Source/DMXRuntime/Private/MVR/DMXMVRGeneralSceneDescription.cpp
+++ b/Engine/Plugins/VirtualProduction/DMX/DMXEngine/Source/DMXRuntime/Private/MVR/DMXMVRGeneralSceneDescription.cpp
@@ -262,8 +262,8 @@ namespace UE::DMXRuntime::DMXMVRFixture::Private
 
 		if (AddressNode)
 		{
-			int32 Universe;
-			int32 Address;
+			int32 Universe = -1;
+			int32 Address = -1;
 
 			auto ParseAddressStringLambda = [](const FXmlNode* const InAddressNode, int32& OutUniverse, int32& OutAddress) -> const bool
 			{
