DMX - Fix CIS NonUnity Compile UnrealGame Mac

diff --git a/Engine/Plugins/VirtualProduction/DMX/DMXEngine/Source/DMXRuntime/Private/MVR/DMXMVRSceneActor.cpp b/Engine/Plugins/VirtualProduction/DMX/DMXEngine/Source/DMXRuntime/Private/MVR/DMXMVRSceneActor.cpp
--- a/Engine/Plugins/VirtualProduction/DMX/DMXEngine/Source/DMXRuntime/Private/MVR/DMXMVRSceneActor.cpp
+++ b/Engine/Plugins/VirtualProduction/DMX/DMXEngine/Source/DMXRuntime/Private/MVR/DMXMVRSceneActor.cpp
@@ -12,6 +12,7 @@
 
 #include "DatasmithAssetUserData.h"
 #include "Components/SceneComponent.h"
+#include "Engine/World.h"
 
 #if WITH_EDITOR
 #include "Editor.h"
