DMX - Add an MVR editor to the DMX Library

diff --git a/Engine/Plugins/VirtualProduction/DMX/DMXEngine/Source/DMXEditor/DMXEditor.Build.cs b/Engine/Plugins/VirtualProduction/DMX/DMXEngine/Source/DMXEditor/DMXEditor.Build.cs
--- a/Engine/Plugins/VirtualProduction/DMX/DMXEngine/Source/DMXEditor/DMXEditor.Build.cs
+++ b/Engine/Plugins/VirtualProduction/DMX/DMXEngine/Source/DMXEditor/DMXEditor.Build.cs
@@ -29,6 +29,7 @@ public class DMXEditor : ModuleRules
 				"CoreUObject",
 				"EditorFramework",
 				"EditorStyle",
+				"EditorWidgets",
 				"Engine",
 				"InputCore",
 				"Json",
@@ -44,6 +45,7 @@ public class DMXEditor : ModuleRules
 				"SlateCore",
 				"TakesCore",
 				"TakeRecorder",
+				"ToolWidgets",
 				"TakeTrackRecorders",
 				"UnrealEd",
 				"XmlParser",
diff --git a/Engine/Plugins/VirtualProduction/DMX/DMXEngine/Source/DMXEditor/Private/Customizations/DMXEntityFixturePatchDetails.cpp b/Engine/Plugins/VirtualProduction/DMX/DMXEngine/Source/DMXEditor/Private/Customizations/DMXEntityFixturePatchDetails.cpp
--- a/Engine/Plugins/VirtualProduction/DMX/DMXEngine/Source/DMXEditor/Private/Customizations/DMXEntityFixturePatchDetails.cpp
+++ b/Engine/Plugins/VirtualProduction/DMX/DMXEngine/Source/DMXEditor/Private/Customizations/DMXEntityFixturePatchDetails.cpp
@@ -60,7 +60,7 @@ void FDMXEntityFixturePatchDetails::CustomizeDetails(IDetailLayoutBuilder& Detai
 	GenerateActiveModesSource();
 
 	int32 ActiveMode;
-	if (ensure(ActiveModeHandle->GetValue(ActiveMode) == FPropertyAccess::Success))
+	if (ActiveModeHandle->GetValue(ActiveMode) == FPropertyAccess::Success)
 	{
 		const bool bActiveModeExists = ActiveModesSource.ContainsByPredicate([ActiveMode](TSharedPtr<uint32> Option) 
 			{
diff --git a/Engine/Plugins/VirtualProduction/DMX/DMXEngine/Source/DMXEditor/Private/Customizations/DMXEntityFixtureTypeDetails.cpp b/Engine/Plugins/VirtualProduction/DMX/DMXEngine/Source/DMXEditor/Private/Customizations/DMXEntityFixtureTypeDetails.cpp
--- a/Engine/Plugins/VirtualProduction/DMX/DMXEngine/Source/DMXEditor/Private/Customizations/DMXEntityFixtureTypeDetails.cpp
+++ b/Engine/Plugins/VirtualProduction/DMX/DMXEngine/Source/DMXEditor/Private/Customizations/DMXEntityFixtureTypeDetails.cpp
@@ -24,14 +24,14 @@ void FDMXEntityFixtureTypeDetails::CustomizeDetails(IDetailLayoutBuilder& Detail
 
 	DetailBuilder.HideProperty(GET_MEMBER_NAME_CHECKED(UDMXEntityFixtureType, Modes));
 
-	DMXImportHandle = DetailBuilder.GetProperty(GET_MEMBER_NAME_CHECKED(UDMXEntityFixtureType, DMXImport));
-	DMXImportHandle->SetOnPropertyValueChanged(FSimpleDelegate::CreateSP(this, &FDMXEntityFixtureTypeDetails::OnDMXImportChanged));
+	GDTFHandle = DetailBuilder.GetProperty(GET_MEMBER_NAME_CHECKED(UDMXEntityFixtureType, GDTF));
+	GDTFHandle->SetOnPropertyValueChanged(FSimpleDelegate::CreateSP(this, &FDMXEntityFixtureTypeDetails::OnDMXImportChanged));
 }
 
 void FDMXEntityFixtureTypeDetails::OnDMXImportChanged()
 {
 	UObject* DMXImportObject = nullptr;
-	if (DMXImportHandle->GetValue(DMXImportObject) == FPropertyAccess::Success)
+	if (GDTFHandle->GetValue(DMXImportObject) == FPropertyAccess::Success)
 	{
 		if (UDMXImportGDTF* DMXImportGDTF = Cast<UDMXImportGDTF>(DMXImportObject))
 		{
diff --git a/Engine/Plugins/VirtualProduction/DMX/DMXEngine/Source/DMXEditor/Private/Customizations/DMXEntityFixtureTypeDetails.h b/Engine/Plugins/VirtualProduction/DMX/DMXEngine/Source/DMXEditor/Private/Customizations/DMXEntityFixtureTypeDetails.h
--- a/Engine/Plugins/VirtualProduction/DMX/DMXEngine/Source/DMXEditor/Private/Customizations/DMXEntityFixtureTypeDetails.h
+++ b/Engine/Plugins/VirtualProduction/DMX/DMXEngine/Source/DMXEditor/Private/Customizations/DMXEntityFixtureTypeDetails.h
@@ -27,7 +27,7 @@ private:
 	void OnDMXImportChanged();
 
 	/** Handle to the DMXImport property */
-	TSharedPtr<IPropertyHandle> DMXImportHandle;
+	TSharedPtr<IPropertyHandle> GDTFHandle;
 
 	/** Property utilities for this customization */
 	TSharedPtr<IPropertyUtilities> PropertyUtilities;
diff --git a/Engine/Plugins/VirtualProduction/DMX/DMXEngine/Source/DMXEditor/Private/Customizations/DMXMVRSceneActorDetails.cpp b/Engine/Plugins/VirtualProduction/DMX/DMXEngine/Source/DMXEditor/Private/Customizations/DMXMVRSceneActorDetails.cpp
new file mode 100644
--- /dev/null
+++ b/Engine/Plugins/VirtualProduction/DMX/DMXEngine/Source/DMXEditor/Private/Customizations/DMXMVRSceneActorDetails.cpp
@@ -0,0 +1,111 @@
+// Copyright Epic Games, Inc. All Rights Reserved.
+
+#include "Customizations/DMXMVRSceneActorDetails.h"
+
+#include "MVR/DMXMVRSceneActor.h"
+#include "Library/DMXEntityFixturePatch.h"
+#include "Library/DMXEntityFixtureType.h"
+#include "Library/DMXImportGDTF.h"
+#include "Library/DMXLibrary.h"
+
+#include "DetailCategoryBuilder.h"
+#include "DetailLayoutBuilder.h"
+#include "DetailWidgetRow.h"
+#include "IPropertyUtilities.h"
+
+
+#define LOCTEXT_NAMESPACE "DMXMVRSceneActorDetails"
+
+TSharedRef<IDetailCustomization> FDMXMVRSceneActorDetails::MakeInstance()
+{
+	return MakeShared<FDMXMVRSceneActorDetails>();
+}
+
+void FDMXMVRSceneActorDetails::CustomizeDetails(IDetailLayoutBuilder& DetailBuilder)
+{
+	PropertyUtilities = DetailBuilder.GetPropertyUtilities();
+
+	TArray<TWeakObjectPtr<UObject>> ObjectsBeingCustomized;
+	DetailBuilder.GetObjectsBeingCustomized(ObjectsBeingCustomized);
+	for (TWeakObjectPtr<UObject> Object : ObjectsBeingCustomized)
+	{
+		if (ADMXMVRSceneActor* MVRSceneActor = Cast<ADMXMVRSceneActor>(Object))
+		{
+			OuterSceneActors.Add(MVRSceneActor);
+		}
+	}
+
+	IDetailCategoryBuilder& ActorTypeForGDTFCategory = DetailBuilder.EditCategory("GDTF to Spawned Actor");
+	ActorTypeForGDTFCategory.InitiallyCollapsed(false);
+
+	const TSharedRef<IPropertyHandle> GDTFToDefaultActorClassesHandle = DetailBuilder.GetProperty(ADMXMVRSceneActor::GetGDTFToDefaultActorClassesPropertyNameChecked());
+	GDTFToDefaultActorClassesHandle->MarkHiddenByCustomization();
+
+	const TSharedPtr<IPropertyHandleArray> GDTFToDefaultActorClassHandleArray = GDTFToDefaultActorClassesHandle->AsArray();
+	FSimpleDelegate GDTFToDefaultActorClassArrayChangedDelegate = FSimpleDelegate::CreateSP(this, &FDMXMVRSceneActorDetails::ForceRefresh);
+	GDTFToDefaultActorClassHandleArray->SetOnNumElementsChanged(GDTFToDefaultActorClassArrayChangedDelegate);
+
+	uint32 NumGDTFToDefaultActorClassElements;
+	if (!ensure(GDTFToDefaultActorClassHandleArray->GetNumElements(NumGDTFToDefaultActorClassElements) == FPropertyAccess::Success))
+	{
+		return;
+	}
+
+	for (uint32 GDTFToDefaultActorClassElementIndex = 0; GDTFToDefaultActorClassElementIndex < NumGDTFToDefaultActorClassElements; GDTFToDefaultActorClassElementIndex++)
+	{
+		const TSharedPtr<IPropertyHandle> GDTFToDefaultActorClassHandle = GDTFToDefaultActorClassHandleArray->GetElement(GDTFToDefaultActorClassElementIndex);
+		const TSharedPtr<IPropertyHandle> GDTFHandle = GDTFToDefaultActorClassHandle->GetChildHandle(GET_MEMBER_NAME_CHECKED(FDMXMVRSceneGDTFToActorClassPair, GDTF));
+		DefaultActorClassHandle = GDTFToDefaultActorClassHandle->GetChildHandle(GET_MEMBER_NAME_CHECKED(FDMXMVRSceneGDTFToActorClassPair, ActorClass));
+		DefaultActorClassHandle->SetOnChildPropertyValuePreChange(FSimpleDelegate::CreateSP(this, &FDMXMVRSceneActorDetails::OnPreEditChangeActorClassInGDTFToActorClasses));
+		DefaultActorClassHandle->SetOnChildPropertyValueChanged(FSimpleDelegate::CreateSP(this, &FDMXMVRSceneActorDetails::OnPostEditChangeActorClassInGDTFToActorClasses));
+
+		ActorTypeForGDTFCategory.AddCustomRow(LOCTEXT("GDTFToDefaultActorClassFilter", "GDTF"))
+			.NameContent()
+			[
+				GDTFHandle->CreatePropertyValueWidget()
+			]
+			.ValueContent()
+			[
+				DefaultActorClassHandle->CreatePropertyValueWidget()
+			];
+	}
+}
+
+void FDMXMVRSceneActorDetails::OnFixturePatchChanged(const UDMXEntityFixturePatch* FixturePatch)
+{
+	ForceRefresh();
+}
+
+void FDMXMVRSceneActorDetails::OnFixtureTypeChanged(const UDMXEntityFixtureType* FixtureType)
+{
+	ForceRefresh();
+}
+
+void FDMXMVRSceneActorDetails::OnPreEditChangeActorClassInGDTFToActorClasses()
+{
+	for (TWeakObjectPtr<ADMXMVRSceneActor> WeakMVRSceneActor : OuterSceneActors)
+	{
+		if (ADMXMVRSceneActor* MVRSceneActor = WeakMVRSceneActor.Get())
+		{
+			MVRSceneActor->PreEditChange(FDMXMVRSceneGDTFToActorClassPair::StaticStruct()->FindPropertyByName(GET_MEMBER_NAME_CHECKED(FDMXMVRSceneGDTFToActorClassPair, ActorClass)));
+		}
+	}
+}
+
+void FDMXMVRSceneActorDetails::OnPostEditChangeActorClassInGDTFToActorClasses()
+{
+	for (TWeakObjectPtr<ADMXMVRSceneActor> WeakMVRSceneActor : OuterSceneActors)
+	{
+		if (ADMXMVRSceneActor* MVRSceneActor = WeakMVRSceneActor.Get())
+		{
+			MVRSceneActor->PostEditChange();
+		}
+	}
+}
+
+void FDMXMVRSceneActorDetails::ForceRefresh()
+{
+	PropertyUtilities->ForceRefresh();
+}
+
+#undef LOCTEXT_NAMESPACE
diff --git a/Engine/Plugins/VirtualProduction/DMX/DMXEngine/Source/DMXEditor/Private/Customizations/DMXMVRSceneActorDetails.h b/Engine/Plugins/VirtualProduction/DMX/DMXEngine/Source/DMXEditor/Private/Customizations/DMXMVRSceneActorDetails.h
new file mode 100644
--- /dev/null
+++ b/Engine/Plugins/VirtualProduction/DMX/DMXEngine/Source/DMXEditor/Private/Customizations/DMXMVRSceneActorDetails.h
@@ -0,0 +1,53 @@
+// Copyright Epic Games, Inc. All Rights Reserved.
+
+#pragma once
+
+#include "CoreMinimal.h"
+#include "IDetailCustomization.h"
+
+class ADMXMVRSceneActor;
+class UDMXEntityFixturePatch;
+class UDMXEntityFixtureType;
+
+class IPropertyHandle;
+class IPropertyUtilities;
+
+
+/** Details customization for the 'FixtureType FunctionProperties' details view */
+class FDMXMVRSceneActorDetails
+	: public IDetailCustomization
+{
+public:
+	/** Creates an instance of this details customization */
+	static TSharedRef<IDetailCustomization> MakeInstance();
+
+protected:
+	//~ Begin IDetailCustomization interface
+	virtual void CustomizeDetails(IDetailLayoutBuilder& DetailBuilder) override;
+	//~ End IDetailCustomization interface
+
+private:
+	/** Called when a Fixture Patch changed */
+	void OnFixturePatchChanged(const UDMXEntityFixturePatch* FixturePatch);
+
+	/** Called when a Fixture Type changed */
+	void OnFixtureTypeChanged(const UDMXEntityFixtureType* FixtureType);
+
+	/** Called before an Actor Class changed in the GDTFToActorClasses member of the Actor */
+	void OnPreEditChangeActorClassInGDTFToActorClasses();
+
+	/** Called after an Actor Class changed in the GDTFToActorClasses member of the Actor */
+	void OnPostEditChangeActorClassInGDTFToActorClasses();
+
+	/** Forces this Details Customization to refresh */
+	void ForceRefresh();
+
+	/** Handle to the Actor Class Property in the GDTFToDefaultActorClasses struct */
+	TSharedPtr<IPropertyHandle> DefaultActorClassHandle;
+
+	/** The Actors being customized in this Detais Customization */
+	TArray<TWeakObjectPtr<ADMXMVRSceneActor>> OuterSceneActors;
+
+	/** Property Utilities for this Details Customization */
+	TSharedPtr<IPropertyUtilities> PropertyUtilities;
+};
diff --git a/Engine/Plugins/VirtualProduction/DMX/DMXEngine/Source/DMXEditor/Private/DMXEditorModule.cpp b/Engine/Plugins/VirtualProduction/DMX/DMXEngine/Source/DMXEditor/Private/DMXEditorModule.cpp
--- a/Engine/Plugins/VirtualProduction/DMX/DMXEngine/Source/DMXEditor/Private/DMXEditorModule.cpp
+++ b/Engine/Plugins/VirtualProduction/DMX/DMXEngine/Source/DMXEditor/Private/DMXEditorModule.cpp
@@ -14,6 +14,7 @@
 #include "Customizations/DMXEntityReferenceCustomization.h"
 #include "Customizations/DMXFixtureSignalFormatCustomization.h"
 #include "Customizations/DMXLibraryPortReferencesCustomization.h"
+#include "Customizations/DMXMVRSceneActorDetails.h"
 #include "Customizations/DMXNameListCustomization.h"
 #include "Customizations/DMXPixelMappingDistributionCustomization.h"
 #include "Customizations/TakeRecorderDMXLibrarySourceEditorCustomization.h"
@@ -22,6 +23,7 @@
 #include "Library/DMXEntityFixtureType.h"
 #include "Library/DMXEntityReference.h"
 #include "Library/DMXLibrary.h"
+#include "MVR/DMXMVRSceneActor.h"
 #include "Sequencer/DMXLibraryTrackEditor.h"
 #include "Sequencer/TakeRecorderDMXLibrarySource.h"
 #include "Widgets/Monitors/SDMXActivityMonitor.h"
@@ -250,10 +252,15 @@ void FDMXEditorModule::RegisterClassCustomizations()
 {
 	FPropertyEditorModule& PropertyModule = FModuleManager::GetModuleChecked<FPropertyEditorModule>("PropertyEditor");
 
-	// Property type customization for the UDMXEntityFixtureType class
+	// Details customization for the UDMXEntityFixtureType class
 	PropertyModule.RegisterCustomClassLayout(UDMXEntityFixtureType::StaticClass()->GetFName(),
 		FOnGetDetailCustomizationInstance::CreateStatic(&FDMXEntityFixtureTypeDetails::MakeInstance)
 	);
+
+	// Details customization for the ADMXMVRSceneActor class
+	PropertyModule.RegisterCustomClassLayout(ADMXMVRSceneActor::StaticClass()->GetFName(),
+		FOnGetDetailCustomizationInstance::CreateStatic(&FDMXMVRSceneActorDetails::MakeInstance)
+	);
 }
 
 void FDMXEditorModule::RegisterPropertyTypeCustomizations()
diff --git a/Engine/Plugins/VirtualProduction/DMX/DMXEngine/Source/DMXEditor/Private/DMXEditorStyle.cpp b/Engine/Plugins/VirtualProduction/DMX/DMXEngine/Source/DMXEditor/Private/DMXEditorStyle.cpp
--- a/Engine/Plugins/VirtualProduction/DMX/DMXEngine/Source/DMXEditor/Private/DMXEditorStyle.cpp
+++ b/Engine/Plugins/VirtualProduction/DMX/DMXEngine/Source/DMXEditor/Private/DMXEditorStyle.cpp
@@ -2,10 +2,10 @@
 
 #include "DMXEditorStyle.h"
 
-#include "Styling/AppStyle.h"
 #include "Engine/Font.h"
 #include "Framework/Application/SlateApplication.h"
 #include "Interfaces/IPluginManager.h"
+#include "Styling/AppStyle.h"
 #include "Styling/SlateStyleRegistry.h"
 #include "Styling/SlateStyleMacros.h"
 #include "Styling/SlateTypes.h"
@@ -26,6 +26,8 @@ FDMXEditorStyle::FDMXEditorStyle()
 	static const FVector2D Icon51x30(51.f, 30.f);
 	static const FVector2D Icon51x31(51.f, 31.f);
 
+	const FSlateColor SelectorColor = FAppStyle::GetSlateColor("SelectorColor");
+
 	const TSharedPtr<FSlateStyleSet> Style = MakeShared<FSlateStyleSet>(GetStyleSetName());
 
 	static const TCHAR* DMXEnginePluginName = TEXT("DMXEngine");
@@ -36,13 +38,19 @@ FDMXEditorStyle::FDMXEditorStyle()
 	}
 	SetCoreContentRoot(FPaths::EngineContentDir() / TEXT("Slate"));
 
-	// Default color brushes
+	// Color brushes
 	{
 		Set("DMXEditor.WhiteBrush", new FSlateColorBrush(FLinearColor(1.f, 1.f, 1.f, 1.f)));
 		Set("DMXEditor.BlackBrush", new FSlateColorBrush(FLinearColor(0.f, 0.f, 0.f, 1.f)));
 	}
 
-	// Default Fonts
+	// Border Brushes
+	{
+		Set("DMXEditor.RoundedPropertyBorder", new FSlateRoundedBoxBrush(FLinearColor::Black, 4.f));
+	}
+
+
+	// Fonts
 	{
 		static const TCHAR* FontPathRoboto = TEXT("Font'/Engine/EngineFonts/Roboto.Roboto'");
 		static const UFont* FontRoboto = Cast<UFont>(StaticLoadObject(UFont::StaticClass(), nullptr, FontPathRoboto));
@@ -131,6 +139,16 @@ FDMXEditorStyle::FDMXEditorStyle()
 			.SetArrowsImage(FSlateNoResource()));
 	}
 
+	{
+		Set("MVRFixtureList.Row", FTableRowStyle()
+			.SetEvenRowBackgroundBrush(IMAGE_BRUSH("Common/Selection", Icon8x8, FLinearColor(.95f, 1.f, 1.f, .3f)))
+			.SetOddRowBackgroundBrush(IMAGE_BRUSH("Common/Selection", Icon8x8, FLinearColor(.85f, .9f, 9.f, .3f)))
+			.SetEvenRowBackgroundHoveredBrush(IMAGE_BRUSH("Common/Selection", Icon8x8, FLinearColor(.95f, 1.f, 1.f, 0.6f)))
+			.SetOddRowBackgroundHoveredBrush(IMAGE_BRUSH("Common/Selection", Icon8x8, FLinearColor(.85f, 9.f, 9.f, 0.6f)))
+			.SetSelectorFocusedBrush(BORDER_BRUSH("Common/Selector", FMargin(4.f / 16.f), SelectorColor))
+		);
+	}
+
 	FSlateStyleRegistry::RegisterSlateStyle(*this);
 }
 
diff --git a/Engine/Plugins/VirtualProduction/DMX/DMXEngine/Source/DMXEditor/Private/DMXFixturePatchSharedData.cpp b/Engine/Plugins/VirtualProduction/DMX/DMXEngine/Source/DMXEditor/Private/DMXFixturePatchSharedData.cpp
--- a/Engine/Plugins/VirtualProduction/DMX/DMXEngine/Source/DMXEditor/Private/DMXFixturePatchSharedData.cpp
+++ b/Engine/Plugins/VirtualProduction/DMX/DMXEngine/Source/DMXEditor/Private/DMXFixturePatchSharedData.cpp
@@ -55,8 +55,16 @@ void FDMXFixturePatchSharedData::SelectFixturePatches(const TArray<TWeakObjectPt
 		return;
 	}
 
+	// Make a copy without duplicates
+	TArray<TWeakObjectPtr<UDMXEntityFixturePatch>> UniquePatchesOnlyArray;
+	UniquePatchesOnlyArray.Reserve(Patches.Num());
+	for (TWeakObjectPtr<UDMXEntityFixturePatch> FixturePatch : Patches)
+	{
+		UniquePatchesOnlyArray.AddUnique(FixturePatch);
+	}
+
 	SelectedFixturePatches.Reset();
-	SelectedFixturePatches = Patches;
+	SelectedFixturePatches = UniquePatchesOnlyArray;
 	OnFixturePatchSelectionChanged.Broadcast();
 }
 
diff --git a/Engine/Plugins/VirtualProduction/DMX/DMXEngine/Source/DMXEditor/Private/DMXFixtureTypeSharedData.cpp b/Engine/Plugins/VirtualProduction/DMX/DMXEngine/Source/DMXEditor/Private/DMXFixtureTypeSharedData.cpp
--- a/Engine/Plugins/VirtualProduction/DMX/DMXEngine/Source/DMXEditor/Private/DMXFixtureTypeSharedData.cpp
+++ b/Engine/Plugins/VirtualProduction/DMX/DMXEngine/Source/DMXEditor/Private/DMXFixtureTypeSharedData.cpp
@@ -74,7 +74,15 @@ void FDMXFixtureTypeSharedData::SelectFixtureTypes(const TArray<TWeakObjectPtr<U
 {
 	if (SelectedFixtureTypes != InFixtureTypes)
 	{
-		SelectedFixtureTypes = InFixtureTypes;
+		// Make a copy without duplicates
+		TArray<TWeakObjectPtr<UDMXEntityFixtureType>> UniqueFixtureTypesOnlyArray;
+		UniqueFixtureTypesOnlyArray.Reserve(InFixtureTypes.Num());
+		for (TWeakObjectPtr<UDMXEntityFixtureType> FixtureType : InFixtureTypes)
+		{
+			UniqueFixtureTypesOnlyArray.AddUnique(FixtureType);
+		}
+
+		SelectedFixtureTypes = UniqueFixtureTypesOnlyArray;
 
 		// Selected Modes and Functions turn invalid
 		SelectedModeIndices.Reset();
diff --git a/Engine/Plugins/VirtualProduction/DMX/DMXEngine/Source/DMXEditor/Private/Factories/DMXLibraryFromMVRFactory.cpp b/Engine/Plugins/VirtualProduction/DMX/DMXEngine/Source/DMXEditor/Private/Factories/DMXLibraryFromMVRFactory.cpp
--- a/Engine/Plugins/VirtualProduction/DMX/DMXEngine/Source/DMXEditor/Private/Factories/DMXLibraryFromMVRFactory.cpp
+++ b/Engine/Plugins/VirtualProduction/DMX/DMXEngine/Source/DMXEditor/Private/Factories/DMXLibraryFromMVRFactory.cpp
@@ -62,10 +62,16 @@ UObject* UDMXLibraryFromMVRFactory::FactoryCreateFile(UClass* InClass, UObject*
 	const FUTF8ToTCHAR TCHARData(reinterpret_cast<const ANSICHAR*>(XMLData.GetData()), XMLData.Num());
 
 	const TSharedRef<FXmlFile> GeneralSceneDescriptionXml = MakeShared<FXmlFile>(FString(TCHARData.Length(), TCHARData.Get()), EConstructMethod::ConstructFromBuffer);
-	FDMXMVRGeneralSceneDescription GeneralSceneDescription(GeneralSceneDescriptionXml);
 
-	TArray<UDMXImportGDTF*> GDTFs = CreateGDTFAssets(Parent, Flags, MVRUnzip.ToSharedRef(), GeneralSceneDescription);
+	const FName GeneralSceneDescriptionName = FName(GetName() + TEXT("_MVRGeneralSceneDescription"));
+	UDMXMVRGeneralSceneDescription* GeneralSceneDescription = UDMXMVRGeneralSceneDescription::CreateFromXmlFile(GeneralSceneDescriptionXml, NewDMXLibrary, GeneralSceneDescriptionName);
 
+	if (!ensureAlwaysMsgf(GeneralSceneDescription, TEXT("Failed to create a General Scene Description when importing MVR."), __FUNCTION__))
+	{
+		return nullptr;
+	}
+
+	TArray<UDMXImportGDTF*> GDTFs = CreateGDTFAssets(Parent, Flags, MVRUnzip.ToSharedRef(), *GeneralSceneDescription);
 	InitDMXLibrary(NewDMXLibrary, GDTFs, GeneralSceneDescription);
 
 	return NewDMXLibrary;
@@ -87,22 +93,25 @@ UDMXLibrary* UDMXLibraryFromMVRFactory::CreateDMXLibraryAsset(UObject* Parent, E
 {
 	constexpr bool bRemovePathFromDesiredName = true;
 	const FString BaseFileName = FPaths::GetBaseFilename(InFilename, bRemovePathFromDesiredName);
-	const FString PackageName = Parent->GetName() + TEXT("_MVR/");
-	const FString AssetName = ObjectTools::SanitizeObjectName(BaseFileName);
+	const FString PackageName = Parent->GetName();
+	FString AssetName = ObjectTools::SanitizeObjectName(BaseFileName);
+	if (PackageName.Contains(AssetName))
+	{
+		AssetName = AssetName + TEXT("_DMXLibrary");
+	}
 
-	FString UniquePackageName;
-	FString UniqueAssetName;
+	FString DMXLibraryAssetName;
 	FAssetToolsModule& AssetToolsModule = FModuleManager::LoadModuleChecked<FAssetToolsModule>("AssetTools");
-	AssetToolsModule.Get().CreateUniqueAssetName(PackageName / AssetName, TEXT(""), UniquePackageName, UniqueAssetName);
+	AssetToolsModule.Get().CreateUniqueAssetName(PackageName / AssetName, TEXT(""), DMXLibraryPackageName, DMXLibraryAssetName);
 
-	UPackage* Package = CreatePackage(*UniquePackageName);
+	UPackage* Package = CreatePackage(*DMXLibraryPackageName);
 	check(Package);
 	Package->FullyLoad();
 
 	UImportSubsystem* ImportSubsystem = GEditor->GetEditorSubsystem<UImportSubsystem>();
-	ImportSubsystem->BroadcastAssetPreImport(this, UDMXLibrary::StaticClass(), Parent, *UniqueAssetName, *MVRFileExtension.ToString());
+	ImportSubsystem->BroadcastAssetPreImport(this, UDMXLibrary::StaticClass(), Parent, *DMXLibraryAssetName, *MVRFileExtension.ToString());
 
-	UDMXLibrary* NewDMXLibrary = NewObject<UDMXLibrary>(Package, *UniqueAssetName, Flags | RF_Public);
+	UDMXLibrary* NewDMXLibrary = NewObject<UDMXLibrary>(Package, *DMXLibraryAssetName, Flags | RF_Public);
 	if (!NewDMXLibrary)
 	{
 		return nullptr;
@@ -113,24 +122,20 @@ UDMXLibrary* UDMXLibraryFromMVRFactory::CreateDMXLibraryAsset(UObject* Parent, E
 	return NewDMXLibrary;
 }
 
-TArray<UDMXImportGDTF*> UDMXLibraryFromMVRFactory::CreateGDTFAssets(UObject* Parent, EObjectFlags Flags, const TSharedRef<FDMXMVRUnzip>& MVRUnzip, const FDMXMVRGeneralSceneDescription& GeneralSceneDescription)
+TArray<UDMXImportGDTF*> UDMXLibraryFromMVRFactory::CreateGDTFAssets(UObject* Parent, EObjectFlags Flags, const TSharedRef<FDMXMVRUnzip>& MVRUnzip, const UDMXMVRGeneralSceneDescription& GeneralSceneDescription)
 {
-	TArray<UDMXImportGDTF*> ImportedGDTFs;
-	TArray<FString> ImportedGDTFNames;
-
-	constexpr bool bRemovePathFromDesiredName = true;
-	const FString PackageName = Parent->GetName() + FString(TEXT("_MVR")) / FString(TEXT("GDTFs"));
+	const FString Path = FPaths::GetPath(DMXLibraryPackageName) + TEXT("/GDTFs");
 
 	TArray<FAssetData> ExistingGDTFAssets;
 	FAssetRegistryModule& AssetRegistryModule = FModuleManager::LoadModuleChecked<FAssetRegistryModule>(TEXT("AssetRegistry"));
-	AssetRegistryModule.Get().GetAssetsByPath(FName(PackageName), ExistingGDTFAssets);
+	AssetRegistryModule.Get().GetAssetsByPath(FName(Path), ExistingGDTFAssets);
 
 	TArray<UDMXImportGDTF*> ExistingGDTFs;
 	for (const FAssetData& AssetData : ExistingGDTFAssets)
 	{
 		if (UDMXImportGDTF* GDTFAsset = Cast<UDMXImportGDTF>(AssetData.GetAsset()))
 		{
-			GeneralSceneDescription.MVRFixtures.ContainsByPredicate([GDTFAsset](const FDMXMVRFixture& MVRFixture)
+			GeneralSceneDescription.GetMVRFixtures().ContainsByPredicate([GDTFAsset](const FDMXMVRFixture& MVRFixture)
 				{
 					return GDTFAsset->SourceFilename.Equals(MVRFixture.GDTFSpec, ESearchCase::IgnoreCase);
 				});
@@ -140,6 +145,8 @@ TArray<UDMXImportGDTF*> UDMXLibraryFromMVRFactory::CreateGDTFAssets(UObject* Par
 	}
 
 	UDMXGDTFFactory* GDTFFactory = NewObject<UDMXGDTFFactory>();
+	TArray<UDMXImportGDTF*> ImportedGDTFs;
+	TArray<FString> ImportedGDTFNames;
 	if (ExistingGDTFs.Num() > 0)
 	{
 		const FText MessageText = LOCTEXT("MVRImportReimportsGDTFDialog", "MVR contains existing GDTFs. Do you want to reimport the existing GDTF assets?");
@@ -157,7 +164,7 @@ TArray<UDMXImportGDTF*> UDMXLibraryFromMVRFactory::CreateGDTFAssets(UObject* Par
 	}
 
 	// Import GDTF Assets that aren't yet imported in the reimport procedure
-	for (const FDMXMVRFixture& MVRFixture : GeneralSceneDescription.MVRFixtures)
+	for (const FDMXMVRFixture& MVRFixture : GeneralSceneDescription.GetMVRFixtures())
 	{
 		// Don't import the same GDTF twice
 		if (!ImportedGDTFNames.Contains(MVRFixture.GDTFSpec))
@@ -165,7 +172,18 @@ TArray<UDMXImportGDTF*> UDMXLibraryFromMVRFactory::CreateGDTFAssets(UObject* Par
 			const FDMXMVRUnzip::FDMXScopedUnzipToTempFile ScopedUnzipGDTF(MVRUnzip, MVRFixture.GDTFSpec);
 			if (!ScopedUnzipGDTF.TempFilePathAndName.IsEmpty())
 			{
-				UPackage* Package = FDMXEditorUtils::GetOrCreatePackage(Parent, MVRFixture.GDTFSpec);
+				constexpr bool bRemovePathFromDesiredName = true;
+				const FString BaseFileName = FPaths::GetBaseFilename(ScopedUnzipGDTF.TempFilePathAndName, bRemovePathFromDesiredName);
+				FString AssetName = ObjectTools::SanitizeObjectName(BaseFileName);
+
+				FString GDTFPackageName;
+				FString GDTFAssetName;
+				FAssetToolsModule& AssetToolsModule = FModuleManager::LoadModuleChecked<FAssetToolsModule>("AssetTools");
+				AssetToolsModule.Get().CreateUniqueAssetName(Path / AssetName, TEXT(""), GDTFPackageName, GDTFAssetName);
+
+				UPackage* Package = CreatePackage(*GDTFPackageName);
+				check(Package);
+				Package->FullyLoad();
 
 				bool bCanceled;
 				UObject* NewGDTFObject = GDTFFactory->FactoryCreateFile(UDMXImportGDTF::StaticClass(), Package, FName(*MVRFixture.GDTFSpec), Flags | RF_Public, ScopedUnzipGDTF.TempFilePathAndName, nullptr, GWarn, bCanceled);
@@ -185,68 +203,94 @@ TArray<UDMXImportGDTF*> UDMXLibraryFromMVRFactory::CreateGDTFAssets(UObject* Par
 	return ImportedGDTFs;
 }
 
-void UDMXLibraryFromMVRFactory::InitDMXLibrary(UDMXLibrary* DMXLibrary, const TArray<UDMXImportGDTF*>& GDTFAssets, const FDMXMVRGeneralSceneDescription& GeneralSceneDescription) const
+void UDMXLibraryFromMVRFactory::InitDMXLibrary(UDMXLibrary* DMXLibrary, const TArray<UDMXImportGDTF*>& GDTFAssets, UDMXMVRGeneralSceneDescription* GeneralSceneDescription) const
 {
-	if (DMXLibrary)
+	if (!ensureAlwaysMsgf(DMXLibrary, TEXT("Trying to initialize a DMX Library from MVR, but the DMX Library is not valid.")))
 	{
-		// Create a Fixture Type for each GDTF
-		TMap<FString, UDMXEntityFixtureType*> GDTFSpecToFixtureTypeMap;
-		for (UDMXImportGDTF* GDTF : GDTFAssets)
-		{
-			const FString GDTFFilename = FPaths::GetCleanFilename(GDTF->SourceFilename);
-			FDMXEntityFixtureTypeConstructionParams FixtureTypeConstructionParams;
-			FixtureTypeConstructionParams.DMXCategory = FDMXFixtureCategory(FDMXFixtureCategory::GetFirstValue());
-			FixtureTypeConstructionParams.ParentDMXLibrary = DMXLibrary;
+		return;
+	}
+	if (!ensureAlwaysMsgf(GeneralSceneDescription, TEXT("Trying to initialize a DMX Library '%s' from MVR, but the DMX Library is not valid."), *DMXLibrary->GetName()))
+	{
+		return;
+	}
+	DMXLibrary->SetMVRGeneralSceneDescription(GeneralSceneDescription);
 
-			UDMXEntityFixtureType* FixtureType = UDMXEntityFixtureType::CreateFixtureTypeInLibrary(FixtureTypeConstructionParams, FPaths::GetBaseFilename(GDTFFilename));
-			FixtureType->SetGDTF(GDTF);
+	// Create a Fixture Type for each GDTF
+	TMap<FString, UDMXEntityFixtureType*> GDTFSpecToFixtureTypeMap;
+	TMap<UDMXEntityFixtureType*, FLinearColor> FixtureTypeToColorMap;
+	for (UDMXImportGDTF* GDTF : GDTFAssets)
+	{
+		const FString GDTFFilename = FPaths::GetCleanFilename(GDTF->SourceFilename);
+		FDMXEntityFixtureTypeConstructionParams FixtureTypeConstructionParams;
+		FixtureTypeConstructionParams.DMXCategory = FDMXFixtureCategory(FDMXFixtureCategory::GetFirstValue());
+		FixtureTypeConstructionParams.ParentDMXLibrary = DMXLibrary;
 
-			GDTFSpecToFixtureTypeMap.Add(GDTFFilename, FixtureType);
+		UDMXEntityFixtureType* FixtureType = UDMXEntityFixtureType::CreateFixtureTypeInLibrary(FixtureTypeConstructionParams, FPaths::GetBaseFilename(GDTFFilename));
+		FixtureType->SetGDTF(GDTF);
+
+		GDTFSpecToFixtureTypeMap.Add(GDTFFilename, FixtureType);
+	}
+
+	// Create Fixture Patches for the MVR Fixtures
+	for (const FDMXMVRFixture& MVRFixture : GeneralSceneDescription->GetMVRFixtures())
+	{
+		if (!GDTFSpecToFixtureTypeMap.Contains(MVRFixture.GDTFSpec))
+		{
+			continue;
 		}
 
-		// Create Fixture Patches for the MVR Fixtures
-		for (const FDMXMVRFixture& MVRFixture : GeneralSceneDescription.MVRFixtures)
+		UDMXEntityFixtureType* FixtureType = GDTFSpecToFixtureTypeMap[MVRFixture.GDTFSpec];
+		UDMXEntityFixturePatch* const* FixturePatchPtr = DMXLibrary->GetEntitiesTypeCast<UDMXEntityFixturePatch>().FindByPredicate([MVRFixture, FixtureType](UDMXEntityFixturePatch* FixturePatch)
+			{
+				const FDMXFixtureMode* ActiveModePtr = FixturePatch->GetActiveMode();
+				return
+					FixturePatch->GetFixtureType() == FixtureType &&
+					FixturePatch->GetUniverseID() == MVRFixture.Addresses.Universe &&
+					FixturePatch->GetStartingChannel() == MVRFixture.Addresses.Address &&
+					ActiveModePtr && ActiveModePtr->ModeName == MVRFixture.GDTFMode;
+			});
+
+		if (!FixtureTypeToColorMap.Contains(FixtureType))
 		{
-			if (!GDTFSpecToFixtureTypeMap.Contains(MVRFixture.GDTFSpec))
+			FLinearColor FixtureTypeColor = FLinearColor::MakeRandomColor();
+
+			// Avoid dominant red values for a bit more of a professional feel
+			if (FixtureTypeColor.R > 0.75f)
 			{
-				continue;
+				FixtureTypeColor.R = FMath::Abs(FixtureTypeColor.R - 1.0f);
 			}
 
-			UDMXEntityFixtureType* FixtureType = GDTFSpecToFixtureTypeMap[MVRFixture.GDTFSpec];
-			UDMXEntityFixturePatch* const* FixturePatchPtr = DMXLibrary->GetEntitiesTypeCast<UDMXEntityFixturePatch>().FindByPredicate([MVRFixture, FixtureType](UDMXEntityFixturePatch* FixturePatch)
+			FixtureTypeToColorMap.Add(FixtureType, FixtureTypeColor);
+		}
+
+		if (!FixturePatchPtr)
+		{
+			int32 ActiveModeIndex = FixtureType->Modes.IndexOfByPredicate([MVRFixture](const FDMXFixtureMode& Mode)
 				{
-					const FDMXFixtureMode* ActiveModePtr = FixturePatch->GetActiveMode();
-					return
-						FixturePatch->GetFixtureType() == FixtureType &&
-						FixturePatch->GetUniverseID() == MVRFixture.Addresses.Universe &&
-						FixturePatch->GetStartingChannel() == MVRFixture.Addresses.Address &&
-						ActiveModePtr && ActiveModePtr->ModeName == MVRFixture.GDTFMode;
+					return Mode.ModeName == MVRFixture.GDTFMode;
 				});
-			
-			if (!FixturePatchPtr)
-			{
-				int32 ActiveModeIndex = FixtureType->Modes.IndexOfByPredicate([MVRFixture](const FDMXFixtureMode& Mode)
-					{
-						return Mode.ModeName == MVRFixture.GDTFMode;
-					});
 
-				if (ActiveModeIndex != INDEX_NONE)
-				{
-					FDMXEntityFixturePatchConstructionParams FixturePatchConstructionParams;
-					FixturePatchConstructionParams.ActiveMode = ActiveModeIndex;
-					FixturePatchConstructionParams.FixtureTypeRef = FDMXEntityFixtureTypeRef(FixtureType);
-					FixturePatchConstructionParams.UniverseID = MVRFixture.Addresses.Universe;
-					FixturePatchConstructionParams.StartingAddress = MVRFixture.Addresses.Address;
-
-					UDMXEntityFixturePatch* FixturePatch = UDMXEntityFixturePatch::CreateFixturePatchInLibrary(FixturePatchConstructionParams, MVRFixture.Name);
-					FixturePatchPtr = &FixturePatch;
-				}
+			if (ActiveModeIndex != INDEX_NONE)
+			{
+				FDMXEntityFixturePatchConstructionParams FixturePatchConstructionParams;
+				FixturePatchConstructionParams.ActiveMode = ActiveModeIndex;
+				FixturePatchConstructionParams.FixtureTypeRef = FDMXEntityFixtureTypeRef(FixtureType);
+				FixturePatchConstructionParams.UniverseID = MVRFixture.Addresses.Universe;
+				FixturePatchConstructionParams.StartingAddress = MVRFixture.Addresses.Address;
+				FixturePatchConstructionParams.MVRFixtureUUID = MVRFixture.UUID;
+
+				UDMXEntityFixturePatch* FixturePatch = UDMXEntityFixturePatch::CreateFixturePatchInLibrary(FixturePatchConstructionParams, MVRFixture.Name);
+				FixturePatch->EditorColor = FixtureTypeToColorMap.FindChecked(FixtureType);
+				FixturePatchPtr = &FixturePatch;
+			}
+			else
+			{
+				UE_LOG(LogDMXEditor, Warning, TEXT("Skipped creating a Fixture Patch for '%s', as no valid Mode could be imported when importing %s."), *MVRFixture.Name, *MVRFixture.GDTFSpec)
 			}
-			check(FixturePatchPtr);
-
-			(*FixturePatchPtr)->AddMVRFixtureInstance(MVRFixture);
 		}
 	}
+
+	DMXLibrary->UpdateGeneralSceneDescription();
 }
 
 #undef LOCTEXT_NAMESPACE
diff --git a/Engine/Plugins/VirtualProduction/DMX/DMXEngine/Source/DMXEditor/Private/Widgets/FixturePatch/DMXMVRFixtureListItem.cpp b/Engine/Plugins/VirtualProduction/DMX/DMXEngine/Source/DMXEditor/Private/Widgets/FixturePatch/DMXMVRFixtureListItem.cpp
new file mode 100644
--- /dev/null
+++ b/Engine/Plugins/VirtualProduction/DMX/DMXEngine/Source/DMXEditor/Private/Widgets/FixturePatch/DMXMVRFixtureListItem.cpp
@@ -0,0 +1,482 @@
+// Copyright Epic Games, Inc. All Rights Reserved.
+
+#include "DMXMVRFixtureListItem.h"
+
+#include "DMXEditor.h"
+#include "DMXFixturePatchSharedData.h"
+#include "Library/DMXEntityFixturePatch.h"
+#include "Library/DMXEntityReference.h"
+#include "Library/DMXLibrary.h"
+
+#include "ScopedTransaction.h"
+
+
+#define LOCTEXT_NAMESPACE "DMXMVRFixtureListItem"
+
+FDMXMVRFixtureListItem::FDMXMVRFixtureListItem(TWeakPtr<FDMXEditor> InDMXEditor, const FGuid& InMVRFixtureUUID)
+	: MVRFixtureUUID(InMVRFixtureUUID)
+	, WeakDMXEditor(InDMXEditor)
+{
+	const UDMXLibrary* const DMXLibrary = GetDMXLibrary();
+	if (!DMXLibrary)
+	{
+		return;
+	}
+
+	const TSharedPtr<FDMXEditor> DMXEditor = WeakDMXEditor.Pin();
+	if (!DMXEditor)
+	{
+		return;
+	}
+	FixturePatchSharedData = DMXEditor->GetFixturePatchSharedData();
+
+	const TArray<UDMXEntityFixturePatch*> FixturePatches = DMXLibrary->GetEntitiesTypeCast<UDMXEntityFixturePatch>();
+	UDMXEntityFixturePatch* const* FixturePatchPtr = FixturePatches.FindByPredicate([this](const UDMXEntityFixturePatch* FixturePatch)
+		{
+			return FixturePatch->GetMVRFixtureUUID() == MVRFixtureUUID;
+		});
+
+	if (ensureAlwaysMsgf(FixturePatchPtr, TEXT("Trying to create an MVR Fixture List Item, but there's no corresponding Fixture Patch for the MVR Fixture UUID.")))
+	{
+		WeakFixturePatch = *FixturePatchPtr;
+	}
+	
+	FDMXMVRFixture* MVRFixturePtr = FindMVRFixture();
+	if (MVRFixturePtr)
+	{
+		CachedMVRFixture = *MVRFixturePtr;
+	}
+}
+
+const FGuid& FDMXMVRFixtureListItem::GetMVRUUID() const
+{
+	return MVRFixtureUUID;
+}
+
+FLinearColor FDMXMVRFixtureListItem::GetBackgroundColor() const
+{
+	if (!ErrorStatusText.IsEmpty())
+	{
+		return FLinearColor::Red;
+	}
+
+	if (UDMXEntityFixturePatch* FixturePatch = WeakFixturePatch.Get())
+	{
+		return FixturePatch->EditorColor;
+	}
+
+	return FLinearColor::Red;
+}
+
+FString FDMXMVRFixtureListItem::GetFixturePatchName() const
+{
+	if (const UDMXEntityFixturePatch* FixturePatch = GetFixturePatch())
+	{
+		return FixturePatch->Name;
+	}
+
+	return FString();
+}
+
+void FDMXMVRFixtureListItem::SetFixturePatchName(const FString& InDesiredName, FString& OutNewName)
+{
+	if (UDMXEntityFixturePatch* FixturePatch = WeakFixturePatch.Get())
+	{
+		if (FixturePatch->Name == InDesiredName)
+		{
+			return;
+		}
+
+		const FScopedTransaction SetFixturePatchNameTransaction(LOCTEXT("SetFixturePatchNameTransaction", "Set Fixture Patch Name"));
+		FixturePatch->PreEditChange(UDMXEntityFixturePatch::StaticClass()->FindPropertyByName(GET_MEMBER_NAME_CHECKED(UDMXEntity, Name)));
+
+		FixturePatch->SetName(InDesiredName);
+
+		FixturePatch->PostEditChange();
+	}
+}
+
+int32 FDMXMVRFixtureListItem::GetUnitNumber() const
+{
+	return CachedMVRFixture.UnitNumber;
+}
+
+void FDMXMVRFixtureListItem::SetUnitNumber(int32 InUnitNumber)
+{
+	FDMXMVRFixture* MVRFixturePtr = FindMVRFixture();
+	UDMXLibrary* DMXLibrary = GetDMXLibrary();
+	if (MVRFixturePtr && DMXLibrary)
+	{
+		const FScopedTransaction SetMVRFixtureUnitNumberTransaction(LOCTEXT("SetMVRFixtureUnitNumberTransaction", "Set MVR Fixture UnitNumber"));
+		DMXLibrary->PreEditChange(UDMXLibrary::StaticClass()->FindPropertyByName(UDMXLibrary::GetGeneralSceneDescriptionPropertyName()));
+
+		MVRFixturePtr->UnitNumber = InUnitNumber;
+
+		DMXLibrary->PostEditChange();
+
+		// Update Cache for Getter
+		CachedMVRFixture = *MVRFixturePtr;
+	}
+}	
+
+const FString& FDMXMVRFixtureListItem::GetMVRFixtureName() const
+{
+	return CachedMVRFixture.Name;
+}
+
+bool FDMXMVRFixtureListItem::SetMVRFixtureName(const FString& Name)
+{
+	FDMXMVRFixture* MVRFixturePtr = FindMVRFixture();
+	UDMXLibrary* DMXLibrary = GetDMXLibrary();
+	if (!Name.IsEmpty() && MVRFixturePtr && DMXLibrary)
+	{
+		const FScopedTransaction SetMVRFixtureNameTransaction(LOCTEXT("SetMVRFixtureNameTransaction", "Set MVR Fixture Name"));
+		DMXLibrary->PreEditChange(UDMXLibrary::StaticClass()->FindPropertyByName(UDMXLibrary::GetGeneralSceneDescriptionPropertyName()));
+
+		MVRFixturePtr->Name = Name;
+
+		DMXLibrary->PostEditChange();
+
+		// Update Cache for Getter
+		CachedMVRFixture = *MVRFixturePtr;
+	}
+
+	return false;
+}
+
+UDMXEntityFixtureType* FDMXMVRFixtureListItem::GetFixtureType() const
+{
+	if (UDMXEntityFixturePatch* FixturePatch = WeakFixturePatch.Get())
+	{
+		return FixturePatch->GetFixtureType();
+	}
+
+	return nullptr;
+}
+
+void FDMXMVRFixtureListItem::SetFixtureType(UDMXEntityFixtureType* FixtureType)
+{
+	UDMXLibrary* DMXLibrary = GetDMXLibrary();
+	UDMXEntityFixturePatch* FixturePatch = WeakFixturePatch.Get();
+	FDMXMVRFixture* MVRFixturePtr = FindMVRFixture();
+
+	if (DMXLibrary && FixturePatch && MVRFixturePtr)
+	{
+		if (FixturePatch->GetFixtureType() == FixtureType)
+		{
+			return;
+		}
+
+		const FScopedTransaction SetFixtureTypeTransaction(LOCTEXT("SetFixtureTypeTransaction", "Set Fixture Type of Patch"));
+		FixturePatch->PreEditChange(UDMXEntityFixturePatch::StaticClass()->FindPropertyByName(UDMXEntityFixturePatch::GetParentFixtureTypeTemplatePropertyNameChecked()));
+
+		FixturePatch->SetFixtureType(FixtureType);
+
+		FixturePatch->PostEditChange();
+	}
+}
+
+int32 FDMXMVRFixtureListItem::GetModeIndex() const
+{
+	if (UDMXEntityFixturePatch* FixturePatch = WeakFixturePatch.Get())
+	{
+		return FixturePatch->GetActiveModeIndex();
+	}
+
+	return INDEX_NONE;
+}
+
+void FDMXMVRFixtureListItem::SetModeIndex(int32 ModeIndex)
+{
+	UDMXLibrary* DMXLibrary = GetDMXLibrary();
+	UDMXEntityFixturePatch* FixturePatch = WeakFixturePatch.Get();
+	FDMXMVRFixture* MVRFixturePtr = FindMVRFixture();
+
+	if (DMXLibrary && FixturePatch && MVRFixturePtr)
+	{
+		if (ModeIndex == FixturePatch->GetActiveModeIndex())
+		{
+			return;
+		}
+		
+		// If all should be changed, just change the patch
+		const FScopedTransaction SetModeTransaction(LOCTEXT("SetModeTransaction", "Set Mode of Patch"));
+		FixturePatch->PreEditChange(UDMXEntityFixturePatch::StaticClass()->FindPropertyByName(UDMXEntityFixturePatch::GetParentFixtureTypeTemplatePropertyNameChecked()));
+
+		FixturePatch->SetActiveModeIndex(ModeIndex);
+			
+		FixturePatch->PostEditChange();
+	}
+}
+
+int32 FDMXMVRFixtureListItem::GetUniverse() const
+{
+	if (UDMXEntityFixturePatch* FixturePatch = WeakFixturePatch.Get())
+	{
+		return FixturePatch->GetUniverseID();
+	}
+
+	return -1;
+}
+
+int32 FDMXMVRFixtureListItem::GetAddress() const
+{
+	if (UDMXEntityFixturePatch* FixturePatch = WeakFixturePatch.Get())
+	{
+		return FixturePatch->GetStartingChannel();
+	}
+
+	return -1;
+}
+
+void FDMXMVRFixtureListItem::SetAddresses(int32 Universe, int32 Address)
+{
+	UDMXLibrary* DMXLibrary = GetDMXLibrary();
+	UDMXEntityFixturePatch* FixturePatch = WeakFixturePatch.Get();
+	FDMXMVRFixture* MVRFixturePtr = FindMVRFixture();
+
+	if (DMXLibrary && FixturePatch && MVRFixturePtr)
+	{
+		if (FixturePatch->GetUniverseID() == Universe &&
+			FixturePatch->GetStartingChannel() == Address)
+		{
+			return;
+		}
+
+		// Only valid values
+		const FDMXFixtureMode* ModePtr = FixturePatch->GetActiveMode();
+		const int32 MaxAddress = ModePtr ? DMX_MAX_ADDRESS - ModePtr->ChannelSpan + 1 : DMX_MAX_ADDRESS;
+		if (Universe < 0 || 
+			Universe > DMX_MAX_UNIVERSE || 
+			Address < 1 || 
+			Address > MaxAddress)
+		{
+			return;
+		}
+		
+		const FScopedTransaction SetAddressesTransaction(LOCTEXT("SetAddressesTransaction", "Set Addresses of Patch"));
+		FixturePatch->PreEditChange(UDMXEntityFixturePatch::StaticClass()->FindPropertyByName(UDMXEntityFixturePatch::GetParentFixtureTypeTemplatePropertyNameChecked()));
+
+		FixturePatch->SetUniverseID(Universe);
+		FixturePatch->SetStartingChannel(Address);
+
+		FixturePatch->PostEditChange();
+
+
+		// Select the universe in Fixture Patch Shared Data
+		FixturePatchSharedData->SelectUniverse(Universe);
+	}
+}
+
+int32 FDMXMVRFixtureListItem::GetNumChannels() const
+{
+	if (UDMXEntityFixturePatch* FixturePatch = WeakFixturePatch.Get())
+	{
+		return FixturePatch->GetChannelSpan();
+	}
+
+	return -1;
+}
+
+UDMXEntityFixturePatch* FDMXMVRFixtureListItem::GetFixturePatch() const
+{
+	return WeakFixturePatch.Get();
+}
+
+void FDMXMVRFixtureListItem::PasteItemsOntoItem(TWeakPtr<FDMXEditor> WeakDMXEditor, const TSharedPtr<FDMXMVRFixtureListItem>& PasteOntoItem, const TArray<TSharedPtr<FDMXMVRFixtureListItem>>& ItemsToPaste)
+{
+	TArray<UDMXEntityFixturePatch*> FixturePatchesToPaste;
+	for (const TSharedPtr<FDMXMVRFixtureListItem>& Item : ItemsToPaste)
+	{
+		if (UDMXEntityFixturePatch* FixturePatch = Item->GetFixturePatch())
+		{
+			FixturePatchesToPaste.Add(FixturePatch);
+		}
+	}
+
+	const FText PasteFixturePatchesTransactionText = FText::Format(LOCTEXT("DuplicateFixturePatchesTransaction", "Duplicate Fixture {0}|plural(one=Patch, other=Patches)"), ItemsToPaste.Num() > 1);
+	DuplicateFixturePatchesInternal(WeakDMXEditor, FixturePatchesToPaste, PasteFixturePatchesTransactionText);
+}
+
+void FDMXMVRFixtureListItem::DuplicateItems(TWeakPtr<FDMXEditor> WeakDMXEditor, const TArray<TSharedPtr<FDMXMVRFixtureListItem>>& ItemsToDuplicate)
+{
+	TArray<UDMXEntityFixturePatch*> FixturePatchesToDuplicate;
+	for (const TSharedPtr<FDMXMVRFixtureListItem>& Item : ItemsToDuplicate)
+	{
+		if (UDMXEntityFixturePatch* FixturePatch = Item->GetFixturePatch())
+		{
+			FixturePatchesToDuplicate.Add(FixturePatch);
+		}
+	}
+
+	const FText DuplicateFixturePatchesTransactionText = FText::Format(LOCTEXT("DuplicateFixturePatchesTransaction", "Duplicate Fixture {0}|plural(one=Patch, other=Patches)"), ItemsToDuplicate.Num() > 1);
+	DuplicateFixturePatchesInternal(WeakDMXEditor, FixturePatchesToDuplicate, DuplicateFixturePatchesTransactionText);
+}
+
+void FDMXMVRFixtureListItem::DeleteItems(const TArray<TSharedPtr<FDMXMVRFixtureListItem>>& ItemsToDelete)
+{
+	if (ItemsToDelete.Num() == 0)
+	{
+		return;
+	}
+	
+	// Its safe to assume all patches are in the same Library - A Multi-Library Editor wouldn't make sense.
+	UDMXLibrary* DMXLibrary = ItemsToDelete[0]->GetDMXLibrary();
+	if (!DMXLibrary)
+	{
+		return;
+	}
+	
+	const FText DeleteFixturePatchesTransactionText = FText::Format(LOCTEXT("DeleteFixturePatchesTransaction", "Delete Fixture {0}|plural(one=Patch, other=Patches)"), ItemsToDelete.Num() > 1);
+	const FScopedTransaction DeleteFixturePatchTransaction(DeleteFixturePatchesTransactionText);
+	DMXLibrary->PreEditChange(nullptr);
+
+	for (const TSharedPtr<FDMXMVRFixtureListItem>& Item : ItemsToDelete)
+	{
+		if (UDMXEntityFixturePatch* FixturePatch = Item->GetFixturePatch())
+		{
+			FixturePatch->PreEditChange(UDMXEntityFixturePatch::StaticClass()->FindPropertyByName(UDMXEntityFixturePatch::GetMVRFixtureUUIDPropertyNameChecked()));
+			UDMXEntityFixturePatch::RemoveFixturePatchFromLibrary(FixturePatch);
+			FixturePatch->PostEditChange();
+		}
+	}
+
+	DMXLibrary->PostEditChange();
+}
+
+UDMXLibrary* FDMXMVRFixtureListItem::GetDMXLibrary() const
+{
+	if (const TSharedPtr<FDMXEditor> DMXEditor = WeakDMXEditor.Pin())
+	{
+		UDMXLibrary* DMXLibrary = DMXEditor->GetDMXLibrary();
+		if (DMXLibrary)
+		{
+			return DMXLibrary;
+		}
+	}
+
+	return nullptr;
+}
+
+void FDMXMVRFixtureListItem::DuplicateFixturePatchesInternal(TWeakPtr<FDMXEditor> WeakDMXEditor, const TArray<UDMXEntityFixturePatch*>& FixturePatchesToDuplicate, const FText& TransactionText)
+{
+	const TSharedPtr<FDMXEditor> DMXEditor = WeakDMXEditor.Pin();
+	if (!DMXEditor.IsValid())
+	{
+		return;
+	}
+	const TSharedPtr<FDMXFixturePatchSharedData> FixturePatchSharedData = DMXEditor->GetFixturePatchSharedData();
+
+	UDMXLibrary* DMXLibrary = DMXEditor->GetDMXLibrary();
+	if (!DMXLibrary)
+	{
+		return;
+	}
+
+	// Find the first free Addresses
+	TArray<UDMXEntityFixturePatch*> FixturePatches = DMXLibrary->GetEntitiesTypeCast<UDMXEntityFixturePatch>();
+	if (FixturePatches.Num() == 0)
+	{
+		return;
+	}
+
+	FixturePatches.Sort([](const UDMXEntityFixturePatch& FixturePatchA, const UDMXEntityFixturePatch& FixturePatchB)
+		{
+			const bool bUniverseIsSmaller = FixturePatchA.GetUniverseID() < FixturePatchB.GetUniverseID();
+			const bool bUniverseIsEqual = FixturePatchA.GetUniverseID() == FixturePatchB.GetUniverseID();
+			const bool bAddressIsSmallerOrEqual = FixturePatchA.GetStartingChannel() <= FixturePatchB.GetStartingChannel();
+
+			return bUniverseIsSmaller || (bUniverseIsEqual && bAddressIsSmallerOrEqual);
+		});
+
+	int32 Address = FixturePatches.Last()->GetStartingChannel() + FixturePatches.Last()->GetChannelSpan();
+	int32 Universe = -1;
+	if (Address > DMX_MAX_ADDRESS)
+	{
+		Address = 1;
+		Universe = FixturePatches.Last()->GetUniverseID() + 1;
+	}
+	else
+	{
+		Universe = FixturePatches.Last()->GetUniverseID();
+	}
+
+	// Duplicate
+	const FScopedTransaction DuplicateFixturePatchTransaction(TransactionText);
+	DMXLibrary->PreEditChange(UDMXLibrary::StaticClass()->FindPropertyByName(UDMXLibrary::GetEntitiesPropertyName()));
+
+	TArray<TWeakObjectPtr<UDMXEntityFixturePatch>> NewFixturePatches;
+	for (const UDMXEntityFixturePatch* FixturePatchToDuplicate : FixturePatchesToDuplicate)
+	{
+		// If this is duplicated from one library onto another (e.g. when pasting via duplicate), create a Copy of the Fixture Type in the duplicated to DMX Library.
+		UDMXEntityFixtureType* FixtureTypeOfPatchToDuplicate = FixturePatchToDuplicate->GetFixtureType();
+		if (FixtureTypeOfPatchToDuplicate && FixtureTypeOfPatchToDuplicate->GetParentLibrary() != DMXLibrary)
+		{
+			FDMXEntityFixtureTypeConstructionParams FixtureTypeConstructionParams;
+			FixtureTypeConstructionParams.DMXCategory = FixtureTypeOfPatchToDuplicate->DMXCategory;
+			FixtureTypeConstructionParams.Modes = FixtureTypeOfPatchToDuplicate->Modes;
+			FixtureTypeConstructionParams.ParentDMXLibrary = DMXLibrary;
+
+			UDMXEntityFixtureType::CreateFixtureTypeInLibrary(FixtureTypeConstructionParams, FixtureTypeOfPatchToDuplicate->Name);
+		}
+
+		// Duplicate the Fixture Patch
+		const int32 ChannelSpan = FixturePatchToDuplicate->GetChannelSpan();
+		if (Address + ChannelSpan - 1 > DMX_MAX_ADDRESS)
+		{
+			Address = 1;
+			Universe++;
+		}
+
+		FDMXEntityFixturePatchConstructionParams ConstructionParams;
+		ConstructionParams.FixtureTypeRef = FixturePatchToDuplicate->GetFixtureType();
+		ConstructionParams.ActiveMode = FixturePatchToDuplicate->GetActiveModeIndex();
+		ConstructionParams.UniverseID = Universe;
+		ConstructionParams.StartingAddress = Address;
+
+		UDMXEntityFixturePatch* NewFixturePatch = UDMXEntityFixturePatch::CreateFixturePatchInLibrary(ConstructionParams, FixturePatchToDuplicate->Name);
+		NewFixturePatches.Add(NewFixturePatch);
+
+		// Duplicate the MVR Fixture if it exists
+		UDMXMVRGeneralSceneDescription* GeneralSceneDescription = DMXLibrary->GetLazyGeneralSceneDescription();
+		if (ensureAlwaysMsgf(GeneralSceneDescription, TEXT("Trying to duplicate MVR Fixture in the General Scene Description of '%s', but the General Scene Desription is not valid."), *DMXLibrary->GetName()))
+		{
+			FDMXMVRFixture* SourceMVRFixturePtr = GeneralSceneDescription->FindMVRFixture(FixturePatchToDuplicate->GetMVRFixtureUUID());
+			if (SourceMVRFixturePtr)
+			{
+				FDMXMVRFixture MVRFixture = *SourceMVRFixturePtr;
+				MVRFixture.UUID = NewFixturePatch->GetMVRFixtureUUID();
+				DMXLibrary->GetLazyGeneralSceneDescription()->AddMVRFixture(MVRFixture);
+			}
+		}
+
+		Address += ChannelSpan;
+	}
+	if (NewFixturePatches.Num() > 0)
+	{
+		FixturePatchSharedData->SelectFixturePatches(NewFixturePatches);
+	}
+
+	DMXLibrary->PostEditChange();
+}
+
+FDMXMVRFixture* FDMXMVRFixtureListItem::FindMVRFixture() const
+{
+	UDMXLibrary* DMXLibrary = GetDMXLibrary();
+
+	if (!DMXLibrary)
+	{
+		return nullptr;
+	}
+
+	DMXLibrary->UpdateGeneralSceneDescription();
+	UDMXMVRGeneralSceneDescription* GeneralSceneDescription = DMXLibrary->GetLazyGeneralSceneDescription();
+	if (!ensureMsgf(GeneralSceneDescription, TEXT("Found Library without a General Scene Description. This should never occur.")))
+	{
+		return nullptr;
+	}
+
+	return GeneralSceneDescription->FindMVRFixture(MVRFixtureUUID);
+}
+
+#undef LOCTEXT_NAMESPACE
diff --git a/Engine/Plugins/VirtualProduction/DMX/DMXEngine/Source/DMXEditor/Private/Widgets/FixturePatch/DMXMVRFixtureListItem.h b/Engine/Plugins/VirtualProduction/DMX/DMXEngine/Source/DMXEditor/Private/Widgets/FixturePatch/DMXMVRFixtureListItem.h
new file mode 100644
--- /dev/null
+++ b/Engine/Plugins/VirtualProduction/DMX/DMXEngine/Source/DMXEditor/Private/Widgets/FixturePatch/DMXMVRFixtureListItem.h
@@ -0,0 +1,118 @@
+// Copyright Epic Games, Inc. All Rights Reserved.
+
+#pragma once
+
+#include "MVR/DMXMVRGeneralSceneDescription.h"
+
+#include "CoreMinimal.h"
+#include "UObject/WeakObjectPtr.h"
+
+class FDMXEditor;
+class FDMXFixturePatchSharedData;
+class UDMXEntityFixturePatch;
+class UDMXEntityFixtureType;
+class UDMXLibrary;
+
+
+/** An MVR Fixture as an Item in a List. A Primary Items is the first MVR Fixture in a Patch. Secondary Items are subsequent MVR Fixtures */
+class FDMXMVRFixtureListItem
+	: public TSharedFromThis<FDMXMVRFixtureListItem>
+{
+public:	
+	/** Constructor */
+	FDMXMVRFixtureListItem(TWeakPtr<FDMXEditor> InDMXEditor, const FGuid& InMVRFixtureUUID);
+
+	/** Returns the MVR UUID of the MVR Fixture */
+	const FGuid& GetMVRUUID() const;
+
+	/** Returns the background color of this Item */
+	FLinearColor GetBackgroundColor() const;
+
+	/** Returns the Name of the Fixture Patch */
+	FString GetFixturePatchName() const;
+
+	/** Sets the Name of the Fixture Patch */
+	void SetFixturePatchName(const FString& InDesiredName, FString& OutNewName);
+
+	/** Returns the Unit Number */
+	int32 GetUnitNumber() const;
+
+	/** Sets the Unit Number */
+	void SetUnitNumber(int32 InUnitNumber);
+
+	/** Returns the Name of the MVR Fixture */
+	const FString& GetMVRFixtureName() const;
+
+	/** Sets the Name of the MVR Fixture */
+	bool SetMVRFixtureName(const FString& Name);
+
+	/** Returns the Fixture Type of the MVR Fixture */
+	UDMXEntityFixtureType* GetFixtureType() const;
+	
+	/** Sets the fixture type of the MVR Fixture */
+	void SetFixtureType(UDMXEntityFixtureType* FixtureType);
+
+	/** Returns the Mode Index of the MVR Fixture */
+	int32 GetModeIndex() const;
+
+	/** Sets the Mode Index for the MVR */
+	void SetModeIndex(int32 ModeIndex);
+
+	/** Returns the Universe the MVR Fixture resides in */
+	int32 GetUniverse() const;
+
+	/** Returns the Address of the MVR Fixture */
+	int32 GetAddress() const;
+
+	/** Sets the Addresses of the MVR Fixture */
+	void SetAddresses(int32 Universe, int32 Address);
+
+	/** Returns Num Channels the MVR Fixture spans */
+	int32 GetNumChannels() const;
+
+	/** Pastes Pasted Items onto an item */
+	static void PasteItemsOntoItem(TWeakPtr<FDMXEditor> WeakDMXEditor, const TSharedPtr<FDMXMVRFixtureListItem>& PasteOntoItem, const TArray<TSharedPtr<FDMXMVRFixtureListItem>>& ItemsToPaste);
+
+	/** Duplicates the Fixture Patches of the items */
+	static void DuplicateItems(TWeakPtr<FDMXEditor> WeakDMXEditor, const TArray<TSharedPtr<FDMXMVRFixtureListItem>>& ItemsToDuplicate);
+
+	/** Deletes specified items. Note, if all items of a Fixture Patch is deleted, the Fixture Patch is remvoed from its DMX Library. */
+	static void DeleteItems(const TArray<TSharedPtr<FDMXMVRFixtureListItem>>& ItemsToDelete);
+
+	/** Returns the Fixture Patch. Const on purpose, should be edited via this class's methods */
+	UDMXEntityFixturePatch* GetFixturePatch() const;
+
+	/** Returns the DMX Library in which the MVR Fixture resides */
+	UDMXLibrary* GetDMXLibrary() const;
+
+	/** Warning Status Text of the Item */
+	FText WarningStatusText;
+
+	/** Error Status Text of the Item */
+	FText ErrorStatusText;
+
+private:
+	/** Duplicates Fixture Patches */
+	static void DuplicateFixturePatchesInternal(TWeakPtr<FDMXEditor> WeakDMXEditor, const TArray<UDMXEntityFixturePatch*>& FixturePatchesToDuplicate, const FText& TransactionText);
+
+	/** Returns the live MVR Fixture in the library */
+	FDMXMVRFixture* FindMVRFixture() const;
+
+	/** True if this is in an even group */
+	bool bIsEvenGroup = false;
+
+	/** Cached MVR Fixture to speed up getters. Should not be used by setters. Use RefreshCachedMVRFixture to update. */
+	FDMXMVRFixture CachedMVRFixture;
+
+	/** The MVRFixtureUUID this Item handles */
+	FGuid MVRFixtureUUID;
+
+	/** The fixture patch the MVR Fixture UUID is assigned to */
+	TWeakObjectPtr<UDMXEntityFixturePatch> WeakFixturePatch;
+
+	/** Shared Data for Fixture Patch Editors */
+	TSharedPtr<FDMXFixturePatchSharedData> FixturePatchSharedData;
+
+	/** The DMX Editor that owns this Item */
+	TWeakPtr<FDMXEditor> WeakDMXEditor;
+};
diff --git a/Engine/Plugins/VirtualProduction/DMX/DMXEngine/Source/DMXEditor/Private/Widgets/FixturePatch/SDMXFixturePatchEditor.cpp b/Engine/Plugins/VirtualProduction/DMX/DMXEngine/Source/DMXEditor/Private/Widgets/FixturePatch/SDMXFixturePatchEditor.cpp
--- a/Engine/Plugins/VirtualProduction/DMX/DMXEngine/Source/DMXEditor/Private/Widgets/FixturePatch/SDMXFixturePatchEditor.cpp
+++ b/Engine/Plugins/VirtualProduction/DMX/DMXEngine/Source/DMXEditor/Private/Widgets/FixturePatch/SDMXFixturePatchEditor.cpp
@@ -6,13 +6,13 @@
 #include "DMXFixturePatchSharedData.h"
 #include "DMXSubsystem.h"
 #include "SDMXFixturePatcher.h"
-#include "SDMXFixturePatchTree.h"
 #include "Customizations/DMXEntityFixturePatchDetails.h"
 #include "IO/DMXPortManager.h"
 #include "IO/DMXOutputPort.h"
 #include "IO/DMXOutputPortReference.h"
 #include "Library/DMXLibrary.h"
 #include "Library/DMXEntityFixturePatch.h"
+#include "Widgets/FixturePatch/SDMXMVRFixtureList.h"
 
 #include "PropertyEditorModule.h"
 #include "Modules/ModuleManager.h"
@@ -41,7 +41,13 @@ void SDMXFixturePatchEditor::Construct(const FArguments& InArgs)
 		.Orientation(EOrientation::Orient_Horizontal)
 		.ResizeMode(ESplitterResizeMode::FixedPosition)
 		
-		// Left, entity list & inspector
+		// Left, MVR Fixture List
+		+ SSplitter::Slot()	
+		[
+			SAssignNew(MVRFixtureList, SDMXMVRFixtureList, DMXEditorPtr)
+		]
+
+		// Right, Fixture Patcher and Details
 		+ SSplitter::Slot()	
 		[
 			SNew(SSplitter)
@@ -49,25 +55,18 @@ void SDMXFixturePatchEditor::Construct(const FArguments& InArgs)
 			.ResizeMode(ESplitterResizeMode::FixedPosition)
 
 			+SSplitter::Slot()			
-			.Value(0.38f)
+			.Value(0.6f)
 			[
-				SAssignNew(FixturePatchTree, SDMXFixturePatchTree)
+				SAssignNew(FixturePatcher, SDMXFixturePatcher)
 				.DMXEditor(DMXEditorPtr)
 			]
 	
 			+SSplitter::Slot()
-			.Value(0.62f)
+			.Value(0.4f)
 			[
 				FixturePatchDetailsView.ToSharedRef()
 			]
 		]
-
-		// Right, fixture patcher
-		+ SSplitter::Slot()	
-		[
-			SAssignNew(FixturePatcher, SDMXFixturePatcher)
-			.DMXEditor(DMXEditorPtr)
-		]
 	];
 
 	// Adopt the selection
@@ -79,31 +78,52 @@ void SDMXFixturePatchEditor::Construct(const FArguments& InArgs)
 
 void SDMXFixturePatchEditor::RequestRenameOnNewEntity(const UDMXEntity* InEntity, ESelectInfo::Type SelectionType)
 {
-	check(FixturePatchTree.IsValid());
-
-	FixturePatchTree->SelectItemByEntity(InEntity, SelectionType);
-	FixturePatchTree->UpdateTree();
+	if (MVRFixtureList.IsValid())
+	{
+		MVRFixtureList->EnterFixturePatchNameEditingMode();
+	}
 }
 
 void SDMXFixturePatchEditor::SelectEntity(UDMXEntity* InEntity, ESelectInfo::Type InSelectionType /*= ESelectInfo::Type::Direct*/)
 {
-	check(FixturePatchTree.IsValid());
-
-	FixturePatchTree->SelectItemByEntity(InEntity, InSelectionType);
+	if (TSharedPtr<FDMXEditor> DMXEditor = DMXEditorPtr.Pin())
+	{
+		if (UDMXEntityFixturePatch* FixturePatch = Cast<UDMXEntityFixturePatch>(InEntity))
+		{
+			DMXEditor->GetFixturePatchSharedData()->SelectFixturePatch(FixturePatch);
+		}
+	}
 }
 
 void SDMXFixturePatchEditor::SelectEntities(const TArray<UDMXEntity*>& InEntities, ESelectInfo::Type InSelectionType /*= ESelectInfo::Type::Direct*/)
 {
-	check(FixturePatchTree.IsValid());
-
-	FixturePatchTree->SelectItemsByEntities(InEntities, InSelectionType);
+	if (TSharedPtr<FDMXEditor> DMXEditor = DMXEditorPtr.Pin())
+	{
+		TArray<TWeakObjectPtr<UDMXEntityFixturePatch>> FixturePatches;
+		for (UDMXEntity* Entity : InEntities)
+		{
+			if (UDMXEntityFixturePatch* FixturePatch = Cast<UDMXEntityFixturePatch>(Entity))
+			{
+				FixturePatches.Add(FixturePatch);
+			}
+		}
+		DMXEditor->GetFixturePatchSharedData()->SelectFixturePatches(FixturePatches);
+	}
 }
 
 TArray<UDMXEntity*> SDMXFixturePatchEditor::GetSelectedEntities() const
 {
-	check(FixturePatchTree.IsValid());
+	TArray<TWeakObjectPtr<UDMXEntityFixturePatch>> SelectedFixturePatches = FixturePatchSharedData->GetSelectedFixturePatches();
+	TArray<UDMXEntity*> SelectedEntities;
+	for (TWeakObjectPtr<UDMXEntityFixturePatch> WeakSelectedFixturePatch : SelectedFixturePatches)
+	{
+		if (UDMXEntity* Entity = WeakSelectedFixturePatch.Get())
+		{
+			SelectedEntities.Add(Entity);
+		}
+	}
 
-	return FixturePatchTree->GetSelectedEntities();
+	return SelectedEntities;
 }
 
 void SDMXFixturePatchEditor::SelectUniverse(int32 UniverseID)
@@ -116,15 +136,6 @@ void SDMXFixturePatchEditor::SelectUniverse(int32 UniverseID)
 	}
 }
 
-void SDMXFixturePatchEditor::OnFinishedChangingProperties(const FPropertyChangedEvent& PropertyChangedEvent)
-{
-	check(FixturePatchTree.IsValid());
-	check(FixturePatcher.IsValid());
-
-	FixturePatchTree->UpdateTree();
-	FixturePatcher->NotifyPropertyChanged(PropertyChangedEvent);
-}
-
 void SDMXFixturePatchEditor::OnFixturePatchesSelected()
 {
 	TArray<TWeakObjectPtr<UDMXEntityFixturePatch>> SelectedFixturePatches = FixturePatchSharedData->GetSelectedFixturePatches();
diff --git a/Engine/Plugins/VirtualProduction/DMX/DMXEngine/Source/DMXEditor/Private/Widgets/FixturePatch/SDMXFixturePatchEditor.h b/Engine/Plugins/VirtualProduction/DMX/DMXEngine/Source/DMXEditor/Private/Widgets/FixturePatch/SDMXFixturePatchEditor.h
--- a/Engine/Plugins/VirtualProduction/DMX/DMXEngine/Source/DMXEditor/Private/Widgets/FixturePatch/SDMXFixturePatchEditor.h
+++ b/Engine/Plugins/VirtualProduction/DMX/DMXEngine/Source/DMXEditor/Private/Widgets/FixturePatch/SDMXFixturePatchEditor.h
@@ -10,6 +10,7 @@ class FDMXEditor;
 class FDMXFixturePatchSharedData;
 class SDMXFixturePatcher;
 class SDMXFixturePatchTree;
+class SDMXMVRFixtureList;
 class UDMXEntityFixturePatch;
 
 struct FPropertyChangedEvent;
@@ -45,17 +46,14 @@ private:
 	/** Selects the patch */
 	void SelectUniverse(int32 UniverseID);
 
-	/** Callback for when some property has changed in the inspector */
-	virtual void OnFinishedChangingProperties(const FPropertyChangedEvent& PropertyChangedEvent);
-
 	/** Called whewn Fixture Patches were selected in Fixture Patch Shared Data */
 	void OnFixturePatchesSelected();
 
 	/** Generates a Detail View for the edited Fixture Patch */
 	TSharedRef<IDetailsView> GenerateFixturePatchDetailsView() const;
 
-	/** Tree View of available Fixture Patches */
-	TSharedPtr<SDMXFixturePatchTree> FixturePatchTree;
+	/** List of Fixture Patches as MVR Fixtures */
+	TSharedPtr<SDMXMVRFixtureList> MVRFixtureList;
 
 	/** Details View for the selected Fixture Patches */
 	TSharedPtr<IDetailsView> FixturePatchDetailsView;
diff --git a/Engine/Plugins/VirtualProduction/DMX/DMXEngine/Source/DMXEditor/Private/Widgets/FixturePatch/SDMXFixturePatcher.cpp b/Engine/Plugins/VirtualProduction/DMX/DMXEngine/Source/DMXEditor/Private/Widgets/FixturePatch/SDMXFixturePatcher.cpp
--- a/Engine/Plugins/VirtualProduction/DMX/DMXEngine/Source/DMXEditor/Private/Widgets/FixturePatch/SDMXFixturePatcher.cpp
+++ b/Engine/Plugins/VirtualProduction/DMX/DMXEngine/Source/DMXEditor/Private/Widgets/FixturePatch/SDMXFixturePatcher.cpp
@@ -140,9 +140,6 @@ void SDMXFixturePatcher::Construct(const FArguments& InArgs)
 				]
 			];
 
-		// Bind to fixture patch changes
-		UDMXEntityFixturePatch::GetOnFixturePatchChanged().AddSP(this, &SDMXFixturePatcher::OnFixturePatchChanged);
-
 		// Bind to selection changes
 		SharedData->OnFixturePatchSelectionChanged.AddSP(this, &SDMXFixturePatcher::OnFixturePatchSelectionChanged);
 		SharedData->OnUniverseSelectionChanged.AddSP(this, &SDMXFixturePatcher::OnUniverseSelectionChanged);
@@ -151,6 +148,12 @@ void SDMXFixturePatcher::Construct(const FArguments& InArgs)
 		UDMXLibrary* Library = GetDMXLibrary();
 		check(Library);
 
+		// Bind to object changes
+		UDMXEntityFixturePatch::GetOnFixturePatchChanged().AddSP(this, &SDMXFixturePatcher::OnFixturePatchChanged);
+		UDMXEntityFixtureType::GetOnFixtureTypeChanged().AddSP(this, &SDMXFixturePatcher::OnFixtureTypeChanged);
+		Library->GetOnEntitiesAdded().AddSP(this, &SDMXFixturePatcher::OnEntitiesAddedOrRemoved);
+		Library->GetOnEntitiesRemoved().AddSP(this, &SDMXFixturePatcher::OnEntitiesAddedOrRemoved);
+
 		TArray<UDMXEntityFixturePatch*> Patches = Library->GetEntitiesTypeCast<UDMXEntityFixturePatch>();
 		UDMXEntityFixturePatch** ExistingPatchPtr = Patches.FindByPredicate([&](UDMXEntityFixturePatch* Patch) {
 			return Patch->GetUniverseID() == SharedData->GetSelectedUniverse();
@@ -163,9 +166,6 @@ void SDMXFixturePatcher::Construct(const FArguments& InArgs)
 		// Bind to tabs being switched
 		FGlobalTabmanager::Get()->OnActiveTabChanged_Subscribe(FOnActiveTabChanged::FDelegate::CreateSP(this, &SDMXFixturePatcher::OnActiveTabChanged));
 
-		// Bind to entity changes
-		Library->GetOnEntitiesAdded().AddSP(this, &SDMXFixturePatcher::OnEntitiesAddedOrRemoved);
-		Library->GetOnEntitiesRemoved().AddSP(this, &SDMXFixturePatcher::OnEntitiesAddedOrRemoved);
 
 		GEditor->RegisterForUndo(this);
 
@@ -174,29 +174,6 @@ void SDMXFixturePatcher::Construct(const FArguments& InArgs)
 }
 END_SLATE_FUNCTION_BUILD_OPTIMIZATION
 
-void SDMXFixturePatcher::NotifyPropertyChanged(const FPropertyChangedEvent& PropertyChangedEvent)
-{
-	if (PropertyChangedEvent.GetPropertyName() == UDMXEntityFixturePatch::GetUniverseIDPropertyNameChecked() ||
-		PropertyChangedEvent.GetPropertyName() == UDMXEntityFixturePatch::GetManualStartingAddressPropertyNameChecked())
-	{
-		if (IsUniverseSelectionEnabled() && PropertyChangedEvent.GetNumObjectsBeingEdited() == 1)
-		{
-			const UDMXEntityFixturePatch* FixturePatch = Cast<UDMXEntityFixturePatch>(PropertyChangedEvent.GetObjectBeingEdited(0));
-			check(FixturePatch);
-
-			SelectUniverse(FixturePatch->GetUniverseID());
-		}
-		
-		RefreshFromProperties();
-	}
-	else if (PropertyChangedEvent.GetPropertyName() == UDMXEntityFixturePatch::GetAutoAssignAddressPropertyNameChecked() ||
-		PropertyChangedEvent.GetPropertyName() == GET_MEMBER_NAME_CHECKED(UDMXEntityFixturePatch, EditorColor) ||
-		PropertyChangedEvent.GetPropertyName() == UDMXEntityFixturePatch::GetActiveModePropertyNameChecked())
-	{
-		RefreshFromProperties();
-	}
-}
-
 void SDMXFixturePatcher::RefreshFromProperties()
 {
 	RefreshFixturePatchState = EDMXRefreshFixturePatcherState::RefreshFromProperties;
@@ -555,6 +532,11 @@ void SDMXFixturePatcher::OnFixturePatchChanged(const UDMXEntityFixturePatch* Fix
 	}
 }
 
+void SDMXFixturePatcher::OnFixtureTypeChanged(const UDMXEntityFixtureType* FixtureType)
+{
+	RefreshFromLibrary();
+}
+
 void SDMXFixturePatcher::OnFixturePatchSelectionChanged()
 {
 	check(SharedData.IsValid());
diff --git a/Engine/Plugins/VirtualProduction/DMX/DMXEngine/Source/DMXEditor/Private/Widgets/FixturePatch/SDMXFixturePatcher.h b/Engine/Plugins/VirtualProduction/DMX/DMXEngine/Source/DMXEditor/Private/Widgets/FixturePatch/SDMXFixturePatcher.h
--- a/Engine/Plugins/VirtualProduction/DMX/DMXEngine/Source/DMXEditor/Private/Widgets/FixturePatch/SDMXFixturePatcher.h
+++ b/Engine/Plugins/VirtualProduction/DMX/DMXEngine/Source/DMXEditor/Private/Widgets/FixturePatch/SDMXFixturePatcher.h
@@ -42,9 +42,6 @@ public:
 	/** Constructs this widget */
 	void Construct(const FArguments& InArgs);
 
-	/** Updates the patcher, should be called on property changes */
-	void NotifyPropertyChanged(const FPropertyChangedEvent& PropertyChangedEvent);
-
 	/** Refreshes the whole view from properties, does not consider changes in the library */
 	void RefreshFromProperties();
 
@@ -92,6 +89,9 @@ protected:
 	/** Called when a fixture patch changed */
 	void OnFixturePatchChanged(const UDMXEntityFixturePatch* FixturePatch);
 
+	/** Called when a fixture type changed */
+	void OnFixtureTypeChanged(const UDMXEntityFixtureType* FixtureType);
+
 	/** Called when a fixture patch was selected */
 	void OnFixturePatchSelectionChanged();
 
diff --git a/Engine/Plugins/VirtualProduction/DMX/DMXEngine/Source/DMXEditor/Private/Widgets/FixturePatch/SDMXMVRFixtureList.cpp b/Engine/Plugins/VirtualProduction/DMX/DMXEngine/Source/DMXEditor/Private/Widgets/FixturePatch/SDMXMVRFixtureList.cpp
new file mode 100644
--- /dev/null
+++ b/Engine/Plugins/VirtualProduction/DMX/DMXEngine/Source/DMXEditor/Private/Widgets/FixturePatch/SDMXMVRFixtureList.cpp
@@ -0,0 +1,988 @@
+// Copyright Epic Games, Inc. All Rights Reserved.
+
+#include "SDMXMVRFixtureList.h"
+
+#include "DMXEditor.h"
+#include "DMXEditorSettings.h"
+#include "DMXEditorUtils.h"
+#include "DMXFixturePatchSharedData.h"
+#include "Commands/DMXEditorCommands.h"
+#include "Library/DMXEntityFixturePatch.h"
+#include "Library/DMXEntityFixtureType.h"
+#include "Library/DMXLibrary.h"
+#include "Widgets/FixturePatch/DMXMVRFixtureListItem.h"
+#include "Widgets/FixturePatch/SDMXMVRFixtureListRow.h"
+#include "Widgets/FixturePatch/SDMXMVRFixtureListToolbar.h"
+
+#include "ScopedTransaction.h"
+#include "TimerManager.h"
+#include "Framework/Commands/GenericCommands.h"
+#include "Framework/Commands/UICommandList.h"
+#include "Framework/MultiBox/MultiBoxBuilder.h"
+#include "HAL/PlatformApplicationMisc.h"
+#include "Styling/AppStyle.h"
+#include "Widgets/SBoxPanel.h"
+#include "Widgets/Views/SHeaderRow.h"
+#include "Widgets/Views/SListView.h"
+
+
+#define LOCTEXT_NAMESPACE "SDMXMVRFixtureList"
+
+/** Helper to generate Status Text for MVR Fixture List Items */
+class FDMXMVRFixtureListStatusTextGenerator
+{
+public:
+	FDMXMVRFixtureListStatusTextGenerator(const TArray<TSharedPtr<FDMXMVRFixtureListItem>>& InItems)
+		: Items(InItems)
+	{}
+
+	/** Generates warning texts. Returns a map of those Items that need a warning set along with the Warning Text */
+	TMap<TSharedPtr<FDMXMVRFixtureListItem>, FText> GenerateWarningTexts() const
+	{
+		TMap<TSharedPtr<FDMXMVRFixtureListItem>, FText> WarningTextMap;
+
+		return WarningTextMap;
+	}
+
+	/** Generates warning texts. Returns a map of those Items that need a warning set along with the Error Text */
+	TMap<TSharedPtr<FDMXMVRFixtureListItem>, FText> GenerateErrorTexts() const
+	{
+		TMap<TSharedPtr<FDMXMVRFixtureListItem>, FText> UnitNumberConflicts = GetUnitNumberConflicts();
+		TMap<TSharedPtr<FDMXMVRFixtureListItem>, FText> ChannelConflicts = GetChannelConflicts();
+
+		TMap<TSharedPtr<FDMXMVRFixtureListItem>, FText> AccumulatedConflicts = UnitNumberConflicts;
+		for (const TTuple<TSharedPtr<FDMXMVRFixtureListItem>, FText>& ChannelConflict : ChannelConflicts)
+		{
+			if (AccumulatedConflicts.Contains(ChannelConflict.Key))
+			{
+				const FText LineTerminator = FText::FromString(LINE_TERMINATOR);
+				const FText AccumulatedErrorText = FText::Format(FText::FromString(TEXT("{0}{1}{2}{3}")), AccumulatedConflicts[ChannelConflict.Key], LineTerminator, LineTerminator, ChannelConflict.Value);
+				AccumulatedConflicts[ChannelConflict.Key] = AccumulatedErrorText;
+			}
+			else
+			{
+				AccumulatedConflicts.Add(ChannelConflict.Key, ChannelConflict.Value);
+			}
+		}
+
+		return AccumulatedConflicts;
+	}
+
+private:
+	/** The patch of an item */
+	struct FItemPatch
+	{
+		FItemPatch(const TSharedPtr<FDMXMVRFixtureListItem>& InItem)
+			: Item(InItem)
+		{
+			Universe = Item->GetUniverse();
+			AddressRange = TRange<int32>(Item->GetAddress(), Item->GetNumChannels());
+		}
+
+		FText GetConfictsWith(const FItemPatch& Other) const
+		{
+			// No conflict with self
+			if (Other.Item == Item)
+			{
+				return FText::GetEmpty();
+			}
+
+			// No conflict with the same patch
+			if (Item->GetFixturePatch() == Other.Item->GetFixturePatch())
+			{
+				return FText::GetEmpty();
+			}
+
+			// No conflict if not in the same universe
+			if (Other.Universe != Universe)
+			{
+				return FText::GetEmpty();
+			}
+
+			// No conflict if channels don't overlap
+			if (!AddressRange.Overlaps(Other.AddressRange))
+			{
+				return FText::GetEmpty();
+			}
+
+			// No conflict if patches are functionally equal
+			if (AddressRange.GetLowerBound() == Other.AddressRange.GetLowerBound() &&
+				Item->GetFixtureType() == Other.Item->GetFixtureType() &&
+				Item->GetModeIndex() == Other.Item->GetModeIndex())
+			{
+				return FText::GetEmpty();
+			}
+
+			const FText UnitNumberText = MakeBeautifulItemText(Item);
+			const FText OtherUnitNumberText = MakeBeautifulItemText(Other.Item);
+			if (AddressRange.GetLowerBound() == Other.AddressRange.GetLowerBound() &&
+				Item->GetFixtureType() == Other.Item->GetFixtureType())
+			{
+				// Modes confict
+				check(Item->GetModeIndex() != Other.Item->GetModeIndex());
+				return FText::Format(LOCTEXT("ModeConflict", "Using same Address and Fixture Type as Fixture {1}, but Modes differ."), UnitNumberText, OtherUnitNumberText);
+			}
+			else if (AddressRange.GetLowerBound() == Other.AddressRange.GetLowerBound())
+			{
+				// Fixture Types conflict
+				check(Item->GetFixtureType() != Other.Item->GetFixtureType());
+				return FText::Format(LOCTEXT("FixtureTypeConflict", "Using same Address as Fixture {1}, but Fixture Types differ."), UnitNumberText, OtherUnitNumberText);
+			}
+			else
+			{
+				// Addresses conflict
+				return FText::Format(LOCTEXT("AddressConflict", "Overlapping Addresses with Fixture {1}"), UnitNumberText, OtherUnitNumberText);
+			}
+		}
+
+		FORCEINLINE const TSharedPtr<FDMXMVRFixtureListItem>& GetItem() const { return Item; };
+
+
+	private:
+		int32 Universe = -1;
+		TRange<int32> AddressRange;
+
+		TSharedPtr<FDMXMVRFixtureListItem> Item;
+	};
+
+	/** Returns an Map of Items to Channel conflict Texts */
+	TMap<TSharedPtr<FDMXMVRFixtureListItem>, FText> GetChannelConflicts() const
+	{
+		TArray<FItemPatch> ItemPatches;
+		ItemPatches.Reserve(Items.Num());
+		for (const TSharedPtr<FDMXMVRFixtureListItem>& Item : Items)
+		{
+			FItemPatch ItemPatch(Item);
+			ItemPatches.Add(ItemPatch);
+		}
+
+		TMap<TSharedPtr<FDMXMVRFixtureListItem>, FText> ItemToConflictMap;
+		for (const FItemPatch& ItemPatch : ItemPatches)
+		{
+			for (const FItemPatch& Other : ItemPatches)
+			{
+				const FText ConflictText = ItemPatch.GetConfictsWith(Other);
+				if (!ConflictText.IsEmpty())
+				{
+					if (ItemToConflictMap.Contains(ItemPatch.GetItem()))
+					{
+						FText AppendConflictText = FText::Format(FText::FromString(TEXT("{0}{1}{2}")), ItemToConflictMap[ItemPatch.GetItem()], FText::FromString(FString(LINE_TERMINATOR)), ConflictText);
+						ItemToConflictMap[ItemPatch.GetItem()] = AppendConflictText;
+					}
+					else
+					{
+						ItemToConflictMap.Add(ItemPatch.GetItem(), ConflictText);
+					}
+				}
+			}
+		}
+
+		return ItemToConflictMap;
+	}
+
+	/** Returns an Map of Items to Unit Number conflict Texts */
+	TMap<TSharedPtr<FDMXMVRFixtureListItem>, FText> GetUnitNumberConflicts() const
+	{
+		TMap<int32, TArray<TSharedPtr<FDMXMVRFixtureListItem>>> UnitNumberMap;
+		UnitNumberMap.Reserve(Items.Num());
+		for (const TSharedPtr<FDMXMVRFixtureListItem>& Item : Items)
+		{
+			UnitNumberMap.FindOrAdd(Item->GetUnitNumber()).Add(Item);
+		}
+		TArray<TArray<TSharedPtr<FDMXMVRFixtureListItem>>> UnitNumberConflicts;
+		UnitNumberMap.GenerateValueArray(UnitNumberConflicts);
+		UnitNumberConflicts.RemoveAll([](const TArray<TSharedPtr<FDMXMVRFixtureListItem>>& ConflictingItems)
+			{
+				return ConflictingItems.Num() < 2;
+			});
+		
+		TMap<TSharedPtr<FDMXMVRFixtureListItem>, FText> ItemToConflictMap;
+		for (TArray<TSharedPtr<FDMXMVRFixtureListItem>>& ConflictingItems : UnitNumberConflicts)
+		{
+			ConflictingItems.Sort([](const TSharedPtr<FDMXMVRFixtureListItem>& ItemA, const TSharedPtr<FDMXMVRFixtureListItem>& ItemB)
+				{
+					return ItemA->GetUnitNumber() < ItemB->GetUnitNumber();
+				});
+
+			check(ConflictingItems.Num() > 0);
+			FText ConflictText = FText::Format(LOCTEXT("BaseUnitNumberConflictText", "Ambiguous IDs in {0}"), MakeBeautifulItemText(ConflictingItems[0]));
+			for (int32 ConflictingItemIndex = 1; ConflictingItemIndex < ConflictingItems.Num(); ConflictingItemIndex++)
+			{
+				ConflictText = FText::Format(LOCTEXT("AppendUnitNumberConflictText", "{0}, {1}"), ConflictText, MakeBeautifulItemText(ConflictingItems[ConflictingItemIndex]));
+			}
+			
+			for (const TSharedPtr<FDMXMVRFixtureListItem>& Item : ConflictingItems)
+			{
+				ItemToConflictMap.Add(Item, ConflictText);
+			}
+		}
+
+		return ItemToConflictMap;
+	}
+
+	static FText MakeBeautifulItemText(const TSharedPtr<FDMXMVRFixtureListItem>& Item)
+	{
+		const FString AddressesString = FString::FromInt(Item->GetUniverse()) + TEXT(".") + FString::FromInt(Item->GetAddress());
+		const FString ItemNameString = TEXT("'") + Item->GetMVRFixtureName() + TEXT("'");
+		const FString BeautifulItemString = ItemNameString + TEXT(" (") + AddressesString + TEXT(")");;
+		return FText::FromString(BeautifulItemString);
+	}
+
+	/** The items the class handles */
+	TArray<TSharedPtr<FDMXMVRFixtureListItem>> Items;
+};
+
+
+const FName FDMXMVRFixtureListCollumnIDs::Status = "Status";
+const FName FDMXMVRFixtureListCollumnIDs::FixturePatchName = "FixturePatchName";
+const FName FDMXMVRFixtureListCollumnIDs::UnitNumber = "UnitNumber";
+const FName FDMXMVRFixtureListCollumnIDs::MVRFixtureName = "MVRFixtureName";
+const FName FDMXMVRFixtureListCollumnIDs::FixtureType = "FixtureType";
+const FName FDMXMVRFixtureListCollumnIDs::Mode = "Mode";
+const FName FDMXMVRFixtureListCollumnIDs::Patch = "Patch";
+
+const FString SDMXMVRFixtureList::ClipboardCopyMVRFixtureHeader = FString(TEXT("UE::DMX::SDMXMVRFIXTURELIST::COPY_MVRFIXTUREEVENT"));
+
+SDMXMVRFixtureList::SDMXMVRFixtureList()
+	: SortMode(EColumnSortMode::Ascending)
+	, SortedByColumnID(FDMXMVRFixtureListCollumnIDs::Patch)
+{}
+
+SDMXMVRFixtureList::~SDMXMVRFixtureList()
+{
+	SaveHeaderRowSettings();
+}
+
+void SDMXMVRFixtureList::Construct(const FArguments& InArgs, TWeakPtr<FDMXEditor> InDMXEditor)
+{
+	if (!InDMXEditor.IsValid())
+	{
+		return;
+	}
+	
+	WeakDMXEditor = InDMXEditor;
+	FixturePatchSharedData = InDMXEditor.Pin()->GetFixturePatchSharedData();
+
+	// Handle Entity changes
+	UDMXEntityFixturePatch::GetOnFixturePatchChanged().AddSP(this, &SDMXMVRFixtureList::OnFixturePatchChanged);
+	UDMXEntityFixtureType::GetOnFixtureTypeChanged().AddSP(this, &SDMXMVRFixtureList::OnFixtureTypeChanged);
+
+	// Handle Shared Data selection changes
+	FixturePatchSharedData->OnFixturePatchSelectionChanged.AddSP(this, &SDMXMVRFixtureList::OnFixturePatchSharedDataSelectedFixturePatches);
+	FixturePatchSharedData->OnUniverseSelectionChanged.AddSP(this, &SDMXMVRFixtureList::OnFixturePatchSharedDataSelectedUniverse);
+
+	static const FTableViewStyle TableViewStyle = FAppStyle::Get().GetWidgetStyle<FTableViewStyle>("TreeView");
+
+	ChildSlot
+	[
+		SNew(SVerticalBox)
+
+		+ SVerticalBox::Slot()
+		.HAlign(HAlign_Fill)
+		.AutoHeight()
+		[
+			SAssignNew(Toolbar, SDMXMVRFixtureListToolbar, WeakDMXEditor)
+			.OnSearchChanged(this, &SDMXMVRFixtureList::OnSearchChanged)
+		]
+
+		+ SVerticalBox::Slot()
+		.HAlign(HAlign_Fill)
+		.FillHeight(1.f)
+		[
+			SAssignNew(ListView, FDMXMVRFixtureListType)
+			.ListViewStyle(&TableViewStyle)
+			.HeaderRow(GenerateHeaderRow())
+			.ItemHeight(40.0f)
+			.ListItemsSource(&ListSource)
+			.OnGenerateRow(this, &SDMXMVRFixtureList::OnGenerateRow)
+			.OnSelectionChanged(this, &SDMXMVRFixtureList::OnSelectionChanged)
+			.OnContextMenuOpening(this, &SDMXMVRFixtureList::OnContextMenuOpening)
+			.ReturnFocusToSelection(false)
+		]
+	];
+
+	RegisterCommands();
+	RefreshList();
+
+	// Make an initial selection, as if the user clicked it
+	if (ListSource.Num() > 0)
+	{
+		ListView->SetSelection(ListSource[0], ESelectInfo::OnMouseClick);
+	}
+}
+
+void SDMXMVRFixtureList::RequestListRefresh()
+{
+	GEditor->GetTimerManager()->SetTimerForNextTick(
+		FTimerDelegate::CreateLambda([this]()
+			{
+				RefreshList();
+			})
+	);
+}
+
+void SDMXMVRFixtureList::EnterFixturePatchNameEditingMode()
+{
+	if (ListView.IsValid())
+	{
+		const TArray<TSharedPtr<FDMXMVRFixtureListItem>> SelectedItems = ListView->GetSelectedItems();
+		if (SelectedItems.Num() == 0)
+		{
+			const TSharedPtr<SDMXMVRFixtureListRow>* SelectedRowPtr = Rows.FindByPredicate([&SelectedItems](const TSharedPtr<SDMXMVRFixtureListRow>& Row)
+				{
+					return Row->GetItem() == SelectedItems[0];
+				});
+			if (SelectedRowPtr)
+			{
+				(*SelectedRowPtr)->EnterFixturePatchNameEditingMode();
+			}
+		}
+	}
+}
+
+FReply SDMXMVRFixtureList::OnKeyDown(const FGeometry& MyGeometry, const FKeyEvent& InKeyEvent)
+{
+	if (CommandList->ProcessCommandBindings(InKeyEvent))
+	{
+		return FReply::Handled();
+	}
+
+	return FReply::Unhandled();
+}
+
+void SDMXMVRFixtureList::OnSearchChanged()
+{
+	RequestListRefresh();
+}
+
+void SDMXMVRFixtureList::RefreshList()
+{	
+	SaveHeaderRowSettings();
+
+	// Clear cached data
+	ListSource.Reset();
+	Rows.Reset();
+
+	TSharedPtr<FDMXEditor> DMXEditor = WeakDMXEditor.Pin();
+	if (!DMXEditor.IsValid())
+	{
+		return;
+	}
+
+	UDMXLibrary* DMXLibrary = DMXEditor->GetDMXLibrary();
+	if (!IsValid(DMXLibrary))
+	{
+		return;
+	}
+
+	const TArray<UDMXEntityFixturePatch*> FixturePatches = DMXLibrary->GetEntitiesTypeCast<UDMXEntityFixturePatch>();
+	for (UDMXEntityFixturePatch* FixturePatch : FixturePatches)
+	{
+		const FGuid& MVRFixtureUUID = FixturePatch->GetMVRFixtureUUID();
+		const TSharedRef<FDMXMVRFixtureListItem> NewItem = MakeShared<FDMXMVRFixtureListItem>(DMXEditor.ToSharedRef(), MVRFixtureUUID);
+
+		ListSource.Add(NewItem);
+	}
+
+	// Generate status texts
+	GenereateStatusText();
+
+	// Apply search filters. Relies on up-to-date status to find conflicts.
+	if (Toolbar.IsValid())
+	{
+		ListSource = Toolbar->FilterItems(ListSource);
+	}
+	
+	// Update and sort the list and its widgets
+	ListView->RequestListRefresh();
+
+	SortByColumnID(EColumnSortPriority::Max, FDMXMVRFixtureListCollumnIDs::Patch, EColumnSortMode::Ascending);
+
+	AdoptSelectionFromFixturePatchSharedData();
+}
+
+void SDMXMVRFixtureList::GenereateStatusText()
+{
+	for (const TSharedPtr<FDMXMVRFixtureListItem>& Item : ListSource)
+	{
+		Item->WarningStatusText = FText::GetEmpty();
+		Item->ErrorStatusText = FText::GetEmpty();
+	}
+
+	FDMXMVRFixtureListStatusTextGenerator StatusTextGenerator(ListSource);
+
+	const TMap<TSharedPtr<FDMXMVRFixtureListItem>, FText> WarningTextMap = StatusTextGenerator.GenerateWarningTexts();
+	for (const TTuple<TSharedPtr<FDMXMVRFixtureListItem>, FText>& ItemToWarningTextPair : WarningTextMap)
+	{
+		ItemToWarningTextPair.Key->WarningStatusText = ItemToWarningTextPair.Value;
+	}
+
+	const TMap<TSharedPtr<FDMXMVRFixtureListItem>, FText> ErrorTextMap = StatusTextGenerator.GenerateErrorTexts();
+	for (const TTuple<TSharedPtr<FDMXMVRFixtureListItem>, FText>& ItemToErrorTextPair : ErrorTextMap)
+	{
+		ItemToErrorTextPair.Key->WarningStatusText = ItemToErrorTextPair.Value;
+	}
+}
+
+TSharedRef<ITableRow> SDMXMVRFixtureList::OnGenerateRow(TSharedPtr<FDMXMVRFixtureListItem> InItem, const TSharedRef<STableViewBase>& OwnerTable)
+{
+	const TSharedRef<SDMXMVRFixtureListRow> NewRow =
+		SNew(SDMXMVRFixtureListRow, OwnerTable, InItem.ToSharedRef())
+		.OnRowRequestsListRefresh(this, &SDMXMVRFixtureList::RequestListRefresh)
+		.OnRowRequestsStatusRefresh(this, &SDMXMVRFixtureList::GenereateStatusText)
+		.IsSelected_Lambda([this, InItem]()
+			{
+				return ListView->IsItemSelected(InItem);
+			});
+
+	Rows.Add(NewRow);
+	return NewRow;
+}
+
+void SDMXMVRFixtureList::OnSelectionChanged(TSharedPtr<FDMXMVRFixtureListItem> InItem, ESelectInfo::Type SelectInfo)
+{
+	if (SelectInfo == ESelectInfo::Direct)
+	{
+		return;
+	}
+ 
+	TArray<TWeakObjectPtr<UDMXEntityFixturePatch>> FixturePatches;
+	const TArray<TSharedPtr<FDMXMVRFixtureListItem>> SelectedItems = ListView->GetSelectedItems();
+	for (const TSharedPtr<FDMXMVRFixtureListItem>& Item : SelectedItems)
+	{
+		if (UDMXEntityFixturePatch* FixturePatch = Item->GetFixturePatch())
+		{
+			FixturePatches.AddUnique(FixturePatch);
+		}
+	}
+
+	FixturePatchSharedData->SelectFixturePatches(FixturePatches);
+
+	if (SelectedItems.Num() > 0)
+	{
+		const int32 SelectedUniverse = FixturePatchSharedData->GetSelectedUniverse();
+		const int32 UniverseOfFirstItem = SelectedItems[0]->GetUniverse();
+		if (SelectedUniverse != UniverseOfFirstItem)
+		{
+			FixturePatchSharedData->SelectUniverse(UniverseOfFirstItem);
+		}
+	}
+}
+
+void SDMXMVRFixtureList::OnFixturePatchChanged(const UDMXEntityFixturePatch* FixturePatch)
+{
+	if (!bChangingDMXLibrary)
+	{
+		RequestListRefresh();
+	}
+}
+
+void SDMXMVRFixtureList::OnFixtureTypeChanged(const UDMXEntityFixtureType* FixtureType)
+{
+	if (!bChangingDMXLibrary)
+	{
+		RequestListRefresh();
+	}
+}
+
+void SDMXMVRFixtureList::OnFixturePatchSharedDataSelectedFixturePatches()
+{
+	if (!bChangingDMXLibrary)
+	{
+		TArray<TWeakObjectPtr<UDMXEntityFixturePatch>> SelectedFixturePatches = FixturePatchSharedData->GetSelectedFixturePatches();
+		SelectedFixturePatches.RemoveAll([](TWeakObjectPtr<UDMXEntityFixturePatch> FixturePatch)
+			{
+				return !FixturePatch.IsValid();
+			});
+
+		TArray<TSharedPtr<FDMXMVRFixtureListItem>> NewSelection;
+		for (const TSharedPtr<FDMXMVRFixtureListItem>& Item : ListSource)
+		{
+			if (SelectedFixturePatches.Contains(Item->GetFixturePatch()))
+			{
+				NewSelection.Add(Item);
+			}
+		}
+
+		if (NewSelection.Num() > 0)
+		{
+			ListView->ClearSelection();
+			ListView->SetItemSelection(NewSelection, true, ESelectInfo::Direct);
+
+			if (!ListView->IsItemVisible(NewSelection[0]))
+			{
+				ListView->RequestScrollIntoView(NewSelection[0]);
+			}
+		}
+		else
+		{
+			ListView->ClearSelection();
+		}
+	}
+}
+
+void SDMXMVRFixtureList::OnFixturePatchSharedDataSelectedUniverse()
+{
+	if (!bChangingDMXLibrary)
+	{
+		const int32 SelectedUniverse = FixturePatchSharedData->GetSelectedUniverse();
+		const TArray<TSharedPtr<FDMXMVRFixtureListItem>> SelectedItems = ListView->GetSelectedItems();
+		const bool bUniverseAlreadySelected = SelectedItems.ContainsByPredicate([SelectedUniverse](const TSharedPtr<FDMXMVRFixtureListItem>& Item)
+			{
+				return Item->GetUniverse() == SelectedUniverse;
+			});
+		if (bUniverseAlreadySelected)
+		{
+			return;
+		}
+
+		const TSharedPtr<FDMXMVRFixtureListItem>* ItemPtr = ListSource.FindByPredicate([SelectedUniverse](const TSharedPtr<FDMXMVRFixtureListItem>& Item)
+			{
+				return Item->GetUniverse() == SelectedUniverse;
+			});
+		if (ItemPtr)
+		{
+			ListView->RequestScrollIntoView(*ItemPtr);
+		}
+	}
+}
+
+void SDMXMVRFixtureList::AdoptSelectionFromFixturePatchSharedData()
+{
+	if (!ListView.IsValid())
+	{
+		return;
+	}
+
+	const TArray<TWeakObjectPtr<UDMXEntityFixturePatch>> SelectedFixturePatches = FixturePatchSharedData->GetSelectedFixturePatches();
+
+	TArray<TSharedPtr<FDMXMVRFixtureListItem>> NewSelection;
+	for (const TWeakObjectPtr<UDMXEntityFixturePatch> SelectedFixturePatch : SelectedFixturePatches)
+	{
+		const TSharedPtr<FDMXMVRFixtureListItem>* SelectedItemPtr = ListSource.FindByPredicate([SelectedFixturePatch](const TSharedPtr<FDMXMVRFixtureListItem>& Item)
+			{
+				return SelectedFixturePatch.IsValid() && Item->GetFixturePatch() == SelectedFixturePatch;
+			});
+
+		if (SelectedItemPtr)
+		{
+			NewSelection.Add(*SelectedItemPtr);
+		}
+	}
+
+	if (NewSelection.Num() > 0)
+	{
+		ListView->ClearSelection();
+
+		constexpr bool bSelected = true;
+		ListView->SetItemSelection(NewSelection, bSelected, ESelectInfo::OnMouseClick);
+		ListView->RequestScrollIntoView(NewSelection[0]);
+	}
+}
+
+TSharedRef<SHeaderRow> SDMXMVRFixtureList::GenerateHeaderRow()
+{
+	const float StatusColumnWidth = FMath::Max(FAppStyle::GetBrush("Icons.Warning")->GetImageSize().X + 6.f, FAppStyle::GetBrush("Icons.Error")->GetImageSize().X + 6.f);
+
+	HeaderRow = SNew(SHeaderRow);
+	SHeaderRow::FColumn::FArguments ColumnArgs;
+
+	HeaderRow->AddColumn(
+		SHeaderRow::FColumn::FArguments()
+		.ColumnId(FDMXMVRFixtureListCollumnIDs::FixturePatchName)
+		.DefaultLabel(LOCTEXT("FixturePatchNameColumnLabel", "Fixture Patch"))
+		.FillWidth(0.2f)
+	);
+
+	HeaderRow->AddColumn(
+		SHeaderRow::FColumn::FArguments()
+		.ColumnId(FDMXMVRFixtureListCollumnIDs::Status)
+		.DefaultLabel(LOCTEXT("StatusColumnLabel", ""))
+		.FixedWidth(StatusColumnWidth)
+	);
+
+	HeaderRow->AddColumn(
+		SHeaderRow::FColumn::FArguments()
+		.ColumnId(FDMXMVRFixtureListCollumnIDs::UnitNumber)
+		.SortMode(this, &SDMXMVRFixtureList::GetColumnSortMode, FDMXMVRFixtureListCollumnIDs::UnitNumber)
+		.OnSort(this, &SDMXMVRFixtureList::SortByColumnID)
+		.DefaultLabel(LOCTEXT("UnitNumberColumnLabel", "ID"))
+		.FillWidth(0.1f)
+	);
+
+	HeaderRow->AddColumn(
+		SHeaderRow::FColumn::FArguments()
+		.ColumnId(FDMXMVRFixtureListCollumnIDs::MVRFixtureName)
+		.DefaultLabel(LOCTEXT("NameColumnLabel", "Fixture Name"))
+		.FillWidth(0.2f)
+	);
+
+	HeaderRow->AddColumn(
+		SHeaderRow::FColumn::FArguments()
+		.ColumnId(FDMXMVRFixtureListCollumnIDs::FixtureType)
+		.DefaultLabel(LOCTEXT("FixtureTypeColumnLabel", "FixtureType"))
+		.FillWidth(0.2f)
+	);
+
+	HeaderRow->AddColumn(
+		SHeaderRow::FColumn::FArguments()
+		.ColumnId(FDMXMVRFixtureListCollumnIDs::Mode)
+		.DefaultLabel(LOCTEXT("ModeColumnLabel", "Mode"))
+		.FillWidth(0.2f)
+	);
+
+	HeaderRow->AddColumn(
+		SHeaderRow::FColumn::FArguments()
+		.ColumnId(FDMXMVRFixtureListCollumnIDs::Patch)
+		.SortMode(this, &SDMXMVRFixtureList::GetColumnSortMode, FDMXMVRFixtureListCollumnIDs::Patch)
+		.OnSort(this, &SDMXMVRFixtureList::SortByColumnID)
+		.DefaultLabel(LOCTEXT("PatchColumnLabel", "Patch"))
+		.FillWidth(0.1f)
+	);
+
+	// Restore user settings
+	RestoresHeaderRowSettings();
+
+	return HeaderRow.ToSharedRef();
+}
+
+void SDMXMVRFixtureList::SaveHeaderRowSettings()
+{
+	UDMXEditorSettings* EditorSettings = GetMutableDefault<UDMXEditorSettings>();
+	if (HeaderRow.IsValid() && EditorSettings)
+	{
+		for (const SHeaderRow::FColumn& Column : HeaderRow->GetColumns())
+		{
+			if (Column.ColumnId == FDMXMVRFixtureListCollumnIDs::UnitNumber)
+			{
+				EditorSettings->MVRFixtureListSettings.UnitNumberColumnWidth = Column.Width.Get();
+			}
+			if (Column.ColumnId == FDMXMVRFixtureListCollumnIDs::MVRFixtureName)
+			{
+				EditorSettings->MVRFixtureListSettings.NameColumnWidth = Column.Width.Get();
+			}
+			else if (Column.ColumnId == FDMXMVRFixtureListCollumnIDs::FixtureType)
+			{
+				EditorSettings->MVRFixtureListSettings.FixtureTypeColumnWidth = Column.Width.Get();
+			}
+			else if (Column.ColumnId == FDMXMVRFixtureListCollumnIDs::Mode)
+			{
+				EditorSettings->MVRFixtureListSettings.ModeColumnWidth = Column.Width.Get();
+			}
+			else if (Column.ColumnId == FDMXMVRFixtureListCollumnIDs::Patch)
+			{
+				EditorSettings->MVRFixtureListSettings.PatchColumnWidth = Column.Width.Get();
+			}
+		}
+
+		EditorSettings->SaveConfig();
+	}
+}
+
+void SDMXMVRFixtureList::RestoresHeaderRowSettings()
+{
+	if (const UDMXEditorSettings* EditorSettings = GetDefault<UDMXEditorSettings>())
+	{
+		const float UnitNumberColumnWidth = EditorSettings->MVRFixtureListSettings.UnitNumberColumnWidth;
+		if (UnitNumberColumnWidth > 20.f)
+		{
+			HeaderRow->SetColumnWidth(FDMXMVRFixtureListCollumnIDs::UnitNumber, UnitNumberColumnWidth);
+		}
+
+		const float NameColumnWidth = EditorSettings->MVRFixtureListSettings.NameColumnWidth;
+		if (NameColumnWidth > 20.f)
+		{
+			HeaderRow->SetColumnWidth(FDMXMVRFixtureListCollumnIDs::MVRFixtureName, NameColumnWidth);
+		}
+
+		const float FixtureTypeColumnWidth = EditorSettings->MVRFixtureListSettings.FixtureTypeColumnWidth;
+		if (FixtureTypeColumnWidth > 20.f)
+		{
+			HeaderRow->SetColumnWidth(FDMXMVRFixtureListCollumnIDs::FixtureType, FixtureTypeColumnWidth);
+		}
+
+		const float ModeColumnWidth = EditorSettings->MVRFixtureListSettings.ModeColumnWidth;
+		if (ModeColumnWidth > 20.f)
+		{
+			HeaderRow->SetColumnWidth(FDMXMVRFixtureListCollumnIDs::Mode, ModeColumnWidth);
+		}
+
+		const float PatchColumnWidth = EditorSettings->MVRFixtureListSettings.PatchColumnWidth;
+		if (PatchColumnWidth > 20.f)
+		{
+			HeaderRow->SetColumnWidth(FDMXMVRFixtureListCollumnIDs::Patch, PatchColumnWidth);
+		}
+	}
+}
+
+EColumnSortMode::Type SDMXMVRFixtureList::GetColumnSortMode(const FName ColumnId) const
+{
+	if (SortedByColumnID != ColumnId)
+	{
+		return EColumnSortMode::None;
+	}
+
+	return SortMode;
+}
+
+void SDMXMVRFixtureList::SortByColumnID(const EColumnSortPriority::Type SortPriority, const FName& ColumnId, const EColumnSortMode::Type InSortMode)
+{
+	SortMode = InSortMode;
+	SortedByColumnID = ColumnId;
+
+	if (ColumnId == FDMXMVRFixtureListCollumnIDs::UnitNumber)
+	{
+		if (InSortMode == EColumnSortMode::Descending)
+		{
+			ListSource.Sort([](const TSharedPtr<FDMXMVRFixtureListItem>& ItemA, const TSharedPtr<FDMXMVRFixtureListItem>& ItemB)
+				{
+					return !(ItemA->GetUnitNumber() <= ItemB->GetUnitNumber());
+				});
+		}
+		else if (InSortMode == EColumnSortMode::Ascending)
+		{
+			ListSource.Sort([](const TSharedPtr<FDMXMVRFixtureListItem>& ItemA, const TSharedPtr<FDMXMVRFixtureListItem>& ItemB)
+				{
+					return !(ItemA->GetUnitNumber() >= ItemB->GetUnitNumber());
+				});
+		}
+	}
+	else if(ColumnId == FDMXMVRFixtureListCollumnIDs::Patch)
+	{
+		if (InSortMode == EColumnSortMode::Descending)
+		{
+			ListSource.StableSort([](const TSharedPtr<FDMXMVRFixtureListItem>& ItemA, const TSharedPtr<FDMXMVRFixtureListItem>& ItemB)
+				{
+					UDMXEntityFixturePatch* FixturePatchA = ItemA->GetFixturePatch();
+					UDMXEntityFixturePatch* FixturePatchB = ItemB->GetFixturePatch();
+
+					const bool bIsUniverseIDLarger = ItemA->GetUniverse() > ItemB->GetUniverse();
+					const bool bIsSameUniverse = ItemA->GetUniverse() == ItemB->GetUniverse();
+					const bool bAreAddressesLarger = ItemA->GetAddress() > ItemB->GetAddress();
+					return bIsUniverseIDLarger || (bIsSameUniverse && bAreAddressesLarger);
+				});
+		}
+		else if (InSortMode == EColumnSortMode::Ascending)
+		{
+			ListSource.StableSort([](const TSharedPtr<FDMXMVRFixtureListItem>& ItemA, const TSharedPtr<FDMXMVRFixtureListItem>& ItemB)
+				{
+					UDMXEntityFixturePatch* FixturePatchA = ItemA->GetFixturePatch();
+					UDMXEntityFixturePatch* FixturePatchB = ItemB->GetFixturePatch();
+					const bool bIsUniverseIDSmaller = ItemA->GetUniverse() < ItemB->GetUniverse();
+					const bool bIsSameUniverse = ItemA->GetUniverse() == ItemB->GetUniverse();
+					const bool bAreAddressesSmaller = ItemA->GetAddress() < ItemB->GetAddress();
+					return bIsUniverseIDSmaller || (bIsSameUniverse && bAreAddressesSmaller);
+				});
+		}
+	}
+
+	ListView->RequestListRefresh();
+}
+
+TSharedPtr<SWidget> SDMXMVRFixtureList::OnContextMenuOpening()
+{
+	if (!ListView.IsValid())
+	{
+		return SNullWidget::NullWidget;
+	}
+
+	const bool bCloseWindowAfterMenuSelection = true;
+	FMenuBuilder MenuBuilder(bCloseWindowAfterMenuSelection, CommandList);
+
+	if (ListView->GetNumItemsSelected() > 0)
+	{
+		MenuBuilder.BeginSection("BasicOperations");
+		{
+			MenuBuilder.AddMenuEntry(FGenericCommands::Get().Cut);
+			MenuBuilder.AddMenuEntry(FGenericCommands::Get().Copy);
+			MenuBuilder.AddMenuEntry(FGenericCommands::Get().Paste);
+			MenuBuilder.AddMenuEntry(FGenericCommands::Get().Duplicate);
+			MenuBuilder.AddMenuEntry(FGenericCommands::Get().Delete);
+		}
+		MenuBuilder.EndSection();
+	}
+
+	return MenuBuilder.MakeWidget();
+}
+
+void SDMXMVRFixtureList::RegisterCommands()
+{
+	// listen to common editor shortcuts for copy/paste etc
+	CommandList = MakeShared<FUICommandList>();
+	CommandList->MapAction(FGenericCommands::Get().Cut, 
+		FUIAction(
+			FExecuteAction::CreateSP(this, &SDMXMVRFixtureList::OnCutSelectedItems),
+			FCanExecuteAction::CreateSP(this, &SDMXMVRFixtureList::CanCutItems)
+		)
+	);
+	CommandList->MapAction(FGenericCommands::Get().Copy,
+		FUIAction(
+			FExecuteAction::CreateSP(this, &SDMXMVRFixtureList::OnCopySelectedItems),
+			FCanExecuteAction::CreateSP(this, &SDMXMVRFixtureList::CanCopyItems)
+		)
+	);
+	CommandList->MapAction(FGenericCommands::Get().Paste,
+		FUIAction(
+			FExecuteAction::CreateSP(this, &SDMXMVRFixtureList::OnPasteItems),
+			FCanExecuteAction::CreateSP(this, &SDMXMVRFixtureList::CanPasteItems)
+		)
+	);
+	CommandList->MapAction(FGenericCommands::Get().Duplicate,
+		FUIAction(
+			FExecuteAction::CreateSP(this, &SDMXMVRFixtureList::OnDuplicateItems),
+			FCanExecuteAction::CreateSP(this, &SDMXMVRFixtureList::CanDuplicateItems)
+		)
+	);
+	CommandList->MapAction(FGenericCommands::Get().Delete,
+		FUIAction(
+			FExecuteAction::CreateSP(this, &SDMXMVRFixtureList::OnDeleteItems),
+			FCanExecuteAction::CreateSP(this, &SDMXMVRFixtureList::CanDeleteItems),
+			EUIActionRepeatMode::RepeatEnabled
+		)
+	);
+}
+
+bool SDMXMVRFixtureList::CanCutItems() const
+{
+	return CanCopyItems() && CanDeleteItems();
+}
+
+void SDMXMVRFixtureList::OnCutSelectedItems()
+{
+	const TArray<TSharedPtr<FDMXMVRFixtureListItem>> SelectedItems = ListView->GetSelectedItems();
+
+	const FScopedTransaction Transaction(SelectedItems.Num() > 1 ? LOCTEXT("CutMVRFixtures", "Cut Fixtures") : LOCTEXT("CutMVRFixture", "Cut Fixture"));
+
+	OnCopySelectedItems();
+	OnDeleteItems();
+}
+
+bool SDMXMVRFixtureList::CanCopyItems() const
+{
+	const TArray<TSharedPtr<FDMXMVRFixtureListItem>> SelectedItems = ListView->GetSelectedItems();
+	return SelectedItems.Num() > 0;
+}
+
+void SDMXMVRFixtureList::OnCopySelectedItems()
+{
+	const TArray<TSharedPtr<FDMXMVRFixtureListItem>> SelectedItems = ListView->GetSelectedItems();
+
+	FString SelectedMVRUUIDsString(ClipboardCopyMVRFixtureHeader);
+	for (const TSharedPtr<FDMXMVRFixtureListItem>& Item : SelectedItems)
+	{
+		SelectedMVRUUIDsString += Item->GetMVRUUID().ToString() + FString(LINE_TERMINATOR);
+	}
+
+	FPlatformApplicationMisc::ClipboardCopy(*SelectedMVRUUIDsString);
+}
+
+bool SDMXMVRFixtureList::CanPasteItems() const
+{
+	const TArray<TSharedPtr<FDMXMVRFixtureListItem>> SelectedItems = ListView->GetSelectedItems();
+	return SelectedItems.Num() == 1;
+}
+
+void SDMXMVRFixtureList::OnPasteItems()
+{
+	TGuardValue<bool>(bChangingDMXLibrary, true);
+
+	FString ClipboardString;
+	FPlatformApplicationMisc::ClipboardPaste(ClipboardString);
+
+	if (ClipboardString.RemoveFromStart(ClipboardCopyMVRFixtureHeader, ESearchCase::CaseSensitive))
+	{
+		TArray<FString> UUIDStrings;
+		ClipboardString.ParseIntoArray(UUIDStrings, LINE_TERMINATOR);
+
+		TArray<FGuid> UUIDs;
+		for (const FString& UUIDString : UUIDStrings)
+		{
+			FGuid UUID;
+			if (FGuid::Parse(UUIDString, UUID))
+			{
+				UUIDs.Add(UUID);
+			}
+		}
+
+		TArray<TSharedPtr<FDMXMVRFixtureListItem>> ItemsToPaste;
+		for (const FGuid& UUID : UUIDs)
+		{
+			const TSharedPtr<FDMXMVRFixtureListItem>* ItemPtr = ListSource.FindByPredicate([UUID](const TSharedPtr<FDMXMVRFixtureListItem>& Item)
+				{
+					return UUID == Item->GetMVRUUID();
+				});
+
+			if (ItemPtr)
+			{
+				ItemsToPaste.Add(*ItemPtr);
+			}
+		}
+
+		const TArray<TSharedPtr<FDMXMVRFixtureListItem>> SelectedItems = ListView->GetSelectedItems();
+		if (SelectedItems.Num() == 1) // Should always be one given CanPasteItems
+		{
+			FDMXMVRFixtureListItem::PasteItemsOntoItem(WeakDMXEditor, SelectedItems[0], ItemsToPaste);
+		}
+
+		RequestListRefresh();
+		AdoptSelectionFromFixturePatchSharedData();
+	}
+}
+
+bool SDMXMVRFixtureList::CanDuplicateItems() const
+{
+	const TArray<TSharedPtr<FDMXMVRFixtureListItem>> SelectedItems = ListView->GetSelectedItems();
+	return SelectedItems.Num() > 0;
+}
+
+void SDMXMVRFixtureList::OnDuplicateItems()
+{
+	TGuardValue<bool>(bChangingDMXLibrary, true);
+
+	const TArray<TSharedPtr<FDMXMVRFixtureListItem>> SelectedItems = ListView->GetSelectedItems();
+	FDMXMVRFixtureListItem::DuplicateItems(WeakDMXEditor, SelectedItems);
+
+	RequestListRefresh();
+	AdoptSelectionFromFixturePatchSharedData();
+}
+
+bool SDMXMVRFixtureList::CanDeleteItems() const
+{
+	const TArray<TSharedPtr<FDMXMVRFixtureListItem>> SelectedItems = ListView->GetSelectedItems();
+	return SelectedItems.Num() > 0;
+}
+
+void SDMXMVRFixtureList::OnDeleteItems()
+{
+	TGuardValue<bool>(bChangingDMXLibrary, true);
+
+	const TArray<TSharedPtr<FDMXMVRFixtureListItem>> SelectedItems = ListView->GetSelectedItems();
+
+	// Make a beautiful selection invariant to ordering of the List
+	TSharedPtr<FDMXMVRFixtureListItem> NewSelection;
+	for (int32 ItemIndex = 0; ItemIndex < ListSource.Num(); ItemIndex++)
+	{
+		if (SelectedItems.Contains(ListSource[ItemIndex]))
+		{
+			if (ListSource.IsValidIndex(ItemIndex + 1) && !SelectedItems.Contains(ListSource[ItemIndex + 1]))
+			{
+				NewSelection = ListSource[ItemIndex + 1];
+				break;
+			}
+			else if (ListSource.IsValidIndex(ItemIndex - 1) && !SelectedItems.Contains(ListSource[ItemIndex - 1]))
+			{
+				NewSelection = ListSource[ItemIndex - 1];
+				break;
+			}
+		}
+	}
+	if (NewSelection.IsValid())
+	{
+		ListView->SetSelection(NewSelection, ESelectInfo::OnMouseClick);
+	}
+
+	FDMXMVRFixtureListItem::DeleteItems(SelectedItems);
+	RequestListRefresh();
+}
+
+#undef LOCTEXT_NAMESPACE
diff --git a/Engine/Plugins/VirtualProduction/DMX/DMXEngine/Source/DMXEditor/Private/Widgets/FixturePatch/SDMXMVRFixtureList.h b/Engine/Plugins/VirtualProduction/DMX/DMXEngine/Source/DMXEditor/Private/Widgets/FixturePatch/SDMXMVRFixtureList.h
new file mode 100644
--- /dev/null
+++ b/Engine/Plugins/VirtualProduction/DMX/DMXEngine/Source/DMXEditor/Private/Widgets/FixturePatch/SDMXMVRFixtureList.h
@@ -0,0 +1,171 @@
+// Copyright Epic Games, Inc. All Rights Reserved.
+
+#pragma once
+
+#include "CoreMinimal.h"
+#include "Widgets/SCompoundWidget.h"
+#include "Widgets/Views/SHeaderRow.h"
+
+class FDMXEditor;
+class FDMXMVRFixtureListItem;
+class FDMXFixturePatchSharedData;
+class SDMXMVRFixtureListRow;
+class SDMXMVRFixtureListToolbar;
+class UDMXEntityFixturePatch;
+class UDMXEntityFixtureType;
+
+class FUICommandList;
+class ITableRow;
+class SBorder;
+class SHeaderRow;
+template <typename ItemType> class SListView;
+class STableViewBase;
+
+/** Collumn IDs in the Fixture Patch List */
+struct FDMXMVRFixtureListCollumnIDs
+{
+	static const FName Status;
+	static const FName FixturePatchName;
+	static const FName UnitNumber;
+	static const FName MVRFixtureName;
+	static const FName FixtureType;
+	static const FName Mode;
+	static const FName Patch;
+};
+
+
+using FDMXMVRFixtureListType = SListView<TSharedPtr<FDMXMVRFixtureListItem>>;
+
+/** Sortable, editable List of Fixture Patches in the library */
+class SDMXMVRFixtureList
+	: public SCompoundWidget
+{
+public:
+	SLATE_BEGIN_ARGS(SDMXMVRFixtureList)
+	{}
+
+	SLATE_END_ARGS()
+
+	/** Constructor */
+	SDMXMVRFixtureList();
+
+	/** Destructor */
+	virtual ~SDMXMVRFixtureList();
+	
+	/** Constructs this widget */
+	void Construct(const FArguments& InArgs, TWeakPtr<FDMXEditor> InDMXEditor);
+
+	/** Updates the List on the next tick */
+	void RequestListRefresh();
+
+	/** If a single Row is selected, enters Editing Mode for the Fixture Patch Name of the selected Row */
+	void EnterFixturePatchNameEditingMode();
+
+protected:
+	//~ Begin SWidget Interface
+	virtual FReply OnKeyDown(const FGeometry& MyGeometry, const FKeyEvent& InKeyEvent) override;
+	//~ End SWidget Interface
+
+private:
+	/** Called when the search changed */
+	void OnSearchChanged();
+
+	/** Refreshes the List */
+	void RefreshList();
+
+	/** Generates Status Text for each Item and sets it */
+	void GenereateStatusText();
+
+	/** Called when a row in the List gets generated */
+	TSharedRef<ITableRow> OnGenerateRow(TSharedPtr<FDMXMVRFixtureListItem> InItem, const TSharedRef<STableViewBase>& OwnerTable);
+
+	/** Called when the Selection in the list changed */
+	void OnSelectionChanged(TSharedPtr<FDMXMVRFixtureListItem> InItem, ESelectInfo::Type SelectInfo);
+
+	/** Called when a Fixture Patch changed */
+	void OnFixturePatchChanged(const UDMXEntityFixturePatch* FixturePatch);
+
+	/** Called when a Fixture Type changed */
+	void OnFixtureTypeChanged(const UDMXEntityFixtureType* FixtureType);
+
+	/** Called when Fixture Patch Shared Data selected Fixture Patches */
+	void OnFixturePatchSharedDataSelectedFixturePatches();
+
+	/** Called when Fixture Patch Shared Data selected a Universe */
+	void OnFixturePatchSharedDataSelectedUniverse();
+
+	/** Selects what is selected in Fixture Patch Shared Data */
+	void AdoptSelectionFromFixturePatchSharedData();
+
+	/** Generates the Header Row of the List */
+	TSharedRef<SHeaderRow> GenerateHeaderRow();
+
+	/** Saves how the user customized the header row in DMX Editor Settings */
+	void SaveHeaderRowSettings();
+
+	/** Restores how the user customized the header row in DMX Editor Settings */
+	void RestoresHeaderRowSettings();
+
+	/** Returns the column sort mode for the list*/
+	EColumnSortMode::Type GetColumnSortMode(const FName ColumnId) const;
+
+	/** Sort the list by ColumnId */
+	void SortByColumnID(const EColumnSortPriority::Type SortPriority, const FName& ColumnId, const EColumnSortMode::Type InSortMode);
+
+	/** The current Sort Mode */
+	EColumnSortMode::Type SortMode;
+
+	/** By which column ID the List is sorted */
+	FName SortedByColumnID;
+
+	/** Set to true while changing the DMX Library and only redraw the list when this is false, e.g. for Duplicate, Paste etc. */
+	bool bChangingDMXLibrary = false;
+
+	/** Source array for the Fixture Patch List */
+	TArray<TSharedPtr<FDMXMVRFixtureListItem>> ListSource;
+
+	/** List source when filtered by search */
+	TArray<TSharedPtr<FDMXMVRFixtureListItem>> FilteredListSource;
+
+	/** The Search Bar for the List */
+	TSharedPtr<SDMXMVRFixtureListToolbar> Toolbar;
+
+	/** The the list of Fixture Patches */
+	TSharedPtr<FDMXMVRFixtureListType> ListView;
+
+	/** The Header Row of the List */
+	TSharedPtr<SHeaderRow> HeaderRow;
+
+	/** Rows of Mode widgets in the List */
+	TArray<TSharedPtr<SDMXMVRFixtureListRow>> Rows;
+
+	/** Shared Data for Fixture Patch Editors */
+	TSharedPtr<FDMXFixturePatchSharedData> FixturePatchSharedData;
+
+	/** The Fixture Type Editor that owns this widget */
+	TWeakPtr<FDMXEditor> WeakDMXEditor;
+
+private:
+	///////////////////////////////////////////////////
+	// Context menu Commands related
+
+	TSharedPtr<SWidget> OnContextMenuOpening();
+
+	/** Registers the commands */
+	void RegisterCommands();
+
+	bool CanCutItems() const;
+	void OnCutSelectedItems();
+	bool CanCopyItems() const;
+	void OnCopySelectedItems();
+	bool CanPasteItems() const;
+	void OnPasteItems();
+	bool CanDuplicateItems() const;
+	void OnDuplicateItems();
+	bool CanDeleteItems() const;
+	void OnDeleteItems();
+
+	TSharedPtr<FUICommandList> CommandList;
+
+	static const FString ClipboardCopyMVRFixtureHeader;
+};
diff --git a/Engine/Plugins/VirtualProduction/DMX/DMXEngine/Source/DMXEditor/Private/Widgets/FixturePatch/SDMXMVRFixtureListRow.cpp b/Engine/Plugins/VirtualProduction/DMX/DMXEngine/Source/DMXEditor/Private/Widgets/FixturePatch/SDMXMVRFixtureListRow.cpp
new file mode 100644
--- /dev/null
+++ b/Engine/Plugins/VirtualProduction/DMX/DMXEngine/Source/DMXEditor/Private/Widgets/FixturePatch/SDMXMVRFixtureListRow.cpp
@@ -0,0 +1,743 @@
+// Copyright Epic Games, Inc. All Rights Reserved.
+
+#include "SDMXMVRFixtureListRow.h"
+
+#include "DMXEditor.h"
+#include "DMXEditorStyle.h"
+#include "Library/DMXEntityFixturePatch.h"
+#include "Library/DMXEntityFixtureType.h"
+#include "Library/DMXLibrary.h"
+#include "MVR/DMXMVRGeneralSceneDescription.h"
+#include "Widgets/FixturePatch/SDMXMVRFixtureList.h"
+#include "Widgets/FixturePatch/DMXMVRFixtureListItem.h"
+
+#include "SSearchableComboBox.h"
+#include "Widgets/Text/SInlineEditableTextBlock.h"
+#include "Widgets/Input/STextComboBox.h"
+#include "Styling/AppStyle.h"
+
+#define LOCTEXT_NAMESPACE "SDMXMVRFixtureListRow"
+
+/////////////////////////////////////////////////////
+// SDMXMVRFixtureFixtureTypePicker
+
+/** Widget to pick a fixture type for an MVR Fixture */
+
+class SDMXMVRFixtureFixtureTypePicker
+	: public SCompoundWidget
+{
+	DECLARE_DELEGATE_OneParam(FDMXMVRFixtureListRowOnFixtureTypeSelectedDelegate, UDMXEntityFixtureType* /** Selected Fixture Type */);
+
+public:
+	SLATE_BEGIN_ARGS(SDMXMVRFixtureFixtureTypePicker)
+	{}
+		/** Called when the combo box selection changed */
+		SLATE_EVENT(FDMXMVRFixtureListRowOnFixtureTypeSelectedDelegate, OnFixtureTypeSelected)
+
+	SLATE_END_ARGS()
+
+	/** Constructs this widget */
+	void Construct(const FArguments& InArgs, UDMXLibrary* InDMXLibrary)
+	{
+		if (!InDMXLibrary)
+		{
+			return;
+		}
+
+		WeakDMXLibrary = InDMXLibrary;
+		OnFixtureTypeSelectedDelegate = InArgs._OnFixtureTypeSelected;
+
+
+		InDMXLibrary->GetOnEntitiesAdded().AddSP(this, &SDMXMVRFixtureFixtureTypePicker::OnEntitiesAddedOrRemoved);
+		InDMXLibrary->GetOnEntitiesRemoved().AddSP(this, &SDMXMVRFixtureFixtureTypePicker::OnEntitiesAddedOrRemoved);
+		UDMXEntityFixtureType::GetOnFixtureTypeChanged().AddSP(this, &SDMXMVRFixtureFixtureTypePicker::OnFixtureTypeChanged);
+
+		ChildSlot
+			[
+				SAssignNew(ComboBox, SSearchableComboBox)
+				.OptionsSource(&FixtureTypeNames)
+				.OnGenerateWidget(this, &SDMXMVRFixtureFixtureTypePicker::OnGenerateWidget)
+				.OnSelectionChanged(this, &SDMXMVRFixtureFixtureTypePicker::OnSelectionChanged)
+				.Content()						
+				[
+					SNew(STextBlock)
+					.Text(this, &SDMXMVRFixtureFixtureTypePicker::GetSelectedItemText)
+				]
+			];
+
+		Refresh();
+	}
+
+	/** Sets the currently selected Fixture Type */
+	void SetSelection(UDMXEntityFixtureType* FixtureType)
+	{
+		if (FixtureType)
+		{
+			TSharedPtr<FString> const* FixtureTypeNameToSelectPtr = FixtureTypeNames.FindByPredicate([FixtureType](const TSharedPtr<FString>& FixtureTypeName)
+				{
+					return *FixtureTypeName == FixtureType->Name;
+				});
+			if (ensureAlwaysMsgf(FixtureTypeNameToSelectPtr, TEXT("Trying to select a fixture type but Fixture Type %s is not present."), *FixtureType->Name))
+			{
+				ComboBox->SetSelectedItem(*FixtureTypeNameToSelectPtr);
+			}
+		}
+		else
+		{
+			checkf(FixtureTypeNames.Num() > 0, TEXT("Expected at least a 'None' entry in FitureTypesName, but the array is empty."));
+			ComboBox->SetSelectedItem(FixtureTypeNames[0]);
+		}
+	}
+
+private:
+	/** Called when the combo box generates a widget */
+	TSharedRef<SWidget> OnGenerateWidget(TSharedPtr<FString> InItem)
+	{
+		return
+			SNew(STextBlock)
+			.Text(FText::FromString(*InItem));
+	}
+
+	/** Called when the combo box selection changed */
+	void OnSelectionChanged(TSharedPtr<FString> NewSelection, ESelectInfo::Type SelectInfo)
+	{
+		if (UDMXEntityFixtureType* SelectedFixtureType = GetSelectedFixtureType())
+		{
+			OnFixtureTypeSelectedDelegate.ExecuteIfBound(SelectedFixtureType);
+		}
+	}
+
+	UDMXEntityFixtureType* GetSelectedFixtureType() const
+	{
+		UDMXEntityFixtureType* SelectedFixtureType = nullptr;
+		const TSharedPtr<FString> SelectedItem = ComboBox->GetSelectedItem();
+		if (TWeakObjectPtr<UDMXEntityFixtureType> const* FixtureTypePtr = FixtureTypeNameToFixtureTypeMap.Find(SelectedItem))
+		{
+			SelectedFixtureType = FixtureTypePtr->Get();
+		}
+
+		return SelectedFixtureType;
+	}
+
+	/** Returns the text of the selected item */
+	FText GetSelectedItemText() const
+	{
+		if (ComboBox.IsValid())
+		{
+			if (const TSharedPtr<FString> SelectedItem = ComboBox->GetSelectedItem())
+			{
+				return FText::FromString(*SelectedItem);
+			}
+		}
+
+		return FText::GetEmpty();
+	}
+
+	/** Refreshes the Combo Box */
+	void Refresh()
+	{
+		FixtureTypeNames.Reset();
+
+		FixtureTypeNames.Add(MakeShared<FString>(TEXT("None")));
+		if (const UDMXLibrary* DMXLibrary = WeakDMXLibrary.Get())
+		{
+			const TArray<UDMXEntityFixtureType*> FixtureTypes = DMXLibrary->GetEntitiesTypeCast<UDMXEntityFixtureType>();
+			for (UDMXEntityFixtureType* FixtureType : FixtureTypes)
+			{
+				const TSharedPtr<FString> FixtureTypeName = MakeShared<FString>(FixtureType->Name);
+
+				FixtureTypeNames.Add(FixtureTypeName);
+				FixtureTypeNameToFixtureTypeMap.Add(FixtureTypeName, FixtureType);
+			}
+		}
+
+		ComboBox->RefreshOptions();
+	}
+
+	/** Called when the Entities in the DMX Library were added or removed */
+	void OnEntitiesAddedOrRemoved(UDMXLibrary* DMXLibrary, TArray<UDMXEntity*> AddedOrRemovedEntities)
+	{
+		Refresh();
+	}
+
+	/** Called when a Fixture Type changed */
+	void OnFixtureTypeChanged(const UDMXEntityFixtureType* ChangedFixtureType)
+	{
+		Refresh();
+	}
+
+	/** Names of Fixture Types in the Combo Box */
+	TArray<TSharedPtr<FString>> FixtureTypeNames;
+
+	/** Map from the Fixture Type Name to the actual Fixture Type object */
+	TMap<TSharedPtr<FString>, TWeakObjectPtr<UDMXEntityFixtureType>> FixtureTypeNameToFixtureTypeMap;
+
+	/** The combo box widget in use */
+	TSharedPtr<SSearchableComboBox> ComboBox;
+
+	/** The DMX Library from which to select fixture types */
+	TWeakObjectPtr<UDMXLibrary> WeakDMXLibrary;
+
+	// Slate args
+	FDMXMVRFixtureListRowOnFixtureTypeSelectedDelegate OnFixtureTypeSelectedDelegate;
+};
+
+
+/////////////////////////////////////////////////////
+// SDMXMVRFixtureModePicker
+
+/** Widget to pick a fixture type for an MVR Fixture */
+class SDMXMVRFixtureModePicker
+	: public SCompoundWidget
+{
+	DECLARE_DELEGATE_OneParam(FDMXMVRFixtureListRowOnModeSelectedDelegate, int32 /** Selected Mode Index */);
+
+public:
+	SLATE_BEGIN_ARGS(SDMXMVRFixtureModePicker)
+	{}
+		/** Called when the combo box selection changed */
+		SLATE_EVENT(FDMXMVRFixtureListRowOnModeSelectedDelegate, OnModeSelected)
+
+	SLATE_END_ARGS()
+
+	/** Constructs this widget */
+	void Construct(const FArguments& InArgs)
+	{
+		OnModeSelectedDelegate = InArgs._OnModeSelected;
+
+		UDMXEntityFixtureType::GetOnFixtureTypeChanged().AddSP(this, &SDMXMVRFixtureModePicker::OnFixtureTypeChanged);
+
+		ChildSlot
+			[
+				SNew(SVerticalBox)
+				.IsEnabled_Lambda([this]
+					{
+						return bEnabled;
+					})
+				
+				+ SVerticalBox::Slot()
+				[
+					SAssignNew(ComboBox, STextComboBox)
+					.Visibility_Lambda([this]()
+						{
+							return WeakFixtureType.IsValid() ? EVisibility::Visible : EVisibility::Hidden;
+						})
+					.OptionsSource(&ModeNames)
+					.OnSelectionChanged(this, &SDMXMVRFixtureModePicker::OnSelectionChanged)
+				]
+
+				+ SVerticalBox::Slot()
+				[
+					SNew(STextBlock)
+					.Visibility_Lambda([this]()
+						{
+							return WeakFixtureType.IsValid() ? EVisibility::Collapsed : EVisibility::Visible;
+						})
+					.Text(LOCTEXT("NoModeBecauseNoFixtureTypeSelectedInfo", "No Fixture Type Selected"))
+					.Font(FAppStyle::GetFontStyle(TEXT("PropertyWindow.NormalFont")))
+				]
+			];
+
+		RefreshOptions();
+	}
+
+	/** Sets the selection */
+	void SetSelection(UDMXEntityFixtureType* FixtureType, int32 ModeIndex)
+	{
+		WeakFixtureType = FixtureType;		
+		RefreshOptions();
+
+		if (FixtureType && FixtureType->Modes.IsValidIndex(ModeIndex))
+		{
+			if (ensureAlwaysMsgf(ModeNames.IsValidIndex(ModeIndex), TEXT("Trying to select a Mode but its index is not valid.")))
+			{
+				ComboBox->SetSelectedItem(ModeNames[ModeIndex]);
+			}
+		}
+	}
+
+private:
+	/** Called when the combo box selection changed */
+	void OnSelectionChanged(TSharedPtr<FString> NewSelection, ESelectInfo::Type SelectInfo)
+	{
+		UDMXEntityFixtureType* FixtureType = WeakFixtureType.Get();
+		if (ensureAlwaysMsgf(FixtureType, TEXT("Tried to select Mode, but Fixture Type is no longer valid.")))
+		{
+			const TSharedPtr<FString> SelectedItem = ComboBox->GetSelectedItem();
+			int32 const* ModeIndexPtr = ModeNameToModeIndexMap.Find(SelectedItem);
+			if (ensureAlwaysMsgf(ModeIndexPtr, TEXT("Tried to select Mode, but the Mode Index is no longer valid.")))
+			{
+				OnModeSelectedDelegate.ExecuteIfBound(*ModeIndexPtr);
+			}
+		}
+	}
+
+	/** Returns the text of the selected item */
+	FText GetSelectedItemText() const
+	{
+		const TSharedPtr<FString> SelectedItem = ComboBox->GetSelectedItem();
+		return FText::FromString(*SelectedItem);
+	}
+
+	/** Refreshes the Combo Box */
+	void RefreshOptions()
+	{
+		ModeNames.Reset();
+
+		if (UDMXEntityFixtureType* FixtureType = WeakFixtureType.Get())
+		{
+			for (int32 ModeIndex = 0; ModeIndex < FixtureType->Modes.Num(); ModeIndex++)
+			{
+				const TSharedPtr<FString> ModeName = MakeShared<FString>(FixtureType->Modes[ModeIndex].ModeName);
+
+				ModeNames.Add(ModeName);
+				ModeNameToModeIndexMap.Add(ModeName, ModeIndex);
+			}
+		}
+
+		if (ModeNames.Num() == 0)
+		{
+			const FText NoModeAvailableText = LOCTEXT("NoModeAvailableInModeComboBoxText", "No Mode available");
+			ModeNames.Add(MakeShared<FString>(NoModeAvailableText.ToString()));
+
+			bEnabled = false;
+		}
+		else
+		{
+			bEnabled = true;
+		}
+
+		ComboBox->RefreshOptions();
+	}
+
+	/** Called when the Entities in the DMX Library were added or removed */
+	void OnFixtureTypeChanged(const UDMXEntityFixtureType* ChangedFixtureType)
+	{
+		RefreshOptions();
+	}
+
+	/** True if the Combo Box should be enabled */
+	bool bEnabled = true;
+
+	/** Names of Modes in the Combo Box */
+	TArray<TSharedPtr<FString>> ModeNames;
+
+	/** Map from the Fixture Type Name to the actual Mode Index */
+	TMap<TSharedPtr<FString>, int32> ModeNameToModeIndexMap;
+
+	/** The combo box widget in use */
+	TSharedPtr<STextComboBox> ComboBox;
+
+	/** The Fixture Type for which the Modes are currently displayed */
+	TWeakObjectPtr<UDMXEntityFixtureType> WeakFixtureType;
+
+	// Slate args
+	FDMXMVRFixtureListRowOnModeSelectedDelegate OnModeSelectedDelegate;
+};
+
+
+/////////////////////////////////////////////////////
+// SDMXMVRFixtureListRow
+
+void SDMXMVRFixtureListRow::Construct(const FArguments& InArgs, const TSharedRef<STableViewBase>& InOwnerTable, const TSharedRef<FDMXMVRFixtureListItem>& InItem)
+{
+	Item = InItem;
+	OnRowRequestsStatusRefresh = InArgs._OnRowRequestsStatusRefresh;
+	OnRowRequestsListRefresh = InArgs._OnRowRequestsListRefresh;
+	IsSelected = InArgs._IsSelected;
+	
+	SetBorderBackgroundColor(TAttribute<FSlateColor>::CreateLambda([this]()
+		{
+			if (Item.IsValid())
+			{
+				return Item->GetBackgroundColor();
+			}
+			return FLinearColor::Black;
+		}));
+
+
+	SMultiColumnTableRow<TSharedPtr<FDMXMVRFixtureListItem>>::Construct(
+		FSuperRowType::FArguments()
+		.Style(&FDMXEditorStyle::Get().GetWidgetStyle<FTableRowStyle>("MVRFixtureList.Row")), 
+		InOwnerTable);
+}
+
+void SDMXMVRFixtureListRow::EnterFixturePatchNameEditingMode()
+{
+	if (FixturePatchNameTextBlock.IsValid())
+	{
+		FixturePatchNameTextBlock->EnterEditingMode();
+	}
+}
+
+TSharedRef<SWidget> SDMXMVRFixtureListRow::GenerateWidgetForColumn(const FName& ColumnName)
+{
+	if (ColumnName == FDMXMVRFixtureListCollumnIDs::FixturePatchName)
+	{
+		return GenerateFixturePatchNameRow();
+	}
+	else if (ColumnName == FDMXMVRFixtureListCollumnIDs::Status)
+	{
+		return GenerateStatusRow();
+	}
+	else if (ColumnName == FDMXMVRFixtureListCollumnIDs::UnitNumber)
+	{
+		return GenerateUnitNumberRow();
+	}
+	else if (ColumnName == FDMXMVRFixtureListCollumnIDs::MVRFixtureName)
+	{
+		return GenerateMVRFixtureNameRow();
+	}
+	else if (ColumnName == FDMXMVRFixtureListCollumnIDs::FixtureType)
+	{
+		return GenerateFixtureTypeRow();
+	}
+	else if (ColumnName == FDMXMVRFixtureListCollumnIDs::Mode)
+	{
+		return GenerateModeRow();
+	}
+	else if (ColumnName == FDMXMVRFixtureListCollumnIDs::Patch)
+	{
+		return GeneratePatchRow();
+	}
+
+	return SNullWidget::NullWidget;
+}
+
+TSharedRef<SWidget> SDMXMVRFixtureListRow::GenerateFixturePatchNameRow()
+{
+	return
+		SAssignNew(FixturePatchNameBorder, SBorder)
+		.HAlign(HAlign_Fill)
+		.VAlign(VAlign_Center)
+		.Padding(4.f)
+		.BorderImage(FAppStyle::GetBrush("NoBorder"))
+		.OnMouseDoubleClick(this, &SDMXMVRFixtureListRow::OnFixturePatchNameBorderDoubleClicked)
+		[
+			SNew(SBorder)
+			.HAlign(HAlign_Fill)
+			.VAlign(VAlign_Center)
+			.BorderImage(FDMXEditorStyle::Get().GetBrush("DMXEditor.RoundedPropertyBorder"))
+			[
+				SAssignNew(FixturePatchNameTextBlock, SInlineEditableTextBlock)
+				.Text_Lambda([this]()
+				{
+					const FString FixturePatchName = Item->GetFixturePatchName();
+					return FText::FromString(FixturePatchName);
+				})
+				.Font(FCoreStyle::GetDefaultFontStyle("Regular", 10))
+				.OnTextCommitted(this, &SDMXMVRFixtureListRow::OnFixturePatchNameCommitted)
+				.IsSelected(IsSelected)
+			]
+		];
+}
+
+FReply SDMXMVRFixtureListRow::OnFixturePatchNameBorderDoubleClicked(const FGeometry& InMyGeometry, const FPointerEvent& InMouseEvent)
+{
+	if (InMouseEvent.GetEffectingButton() == EKeys::LeftMouseButton)
+	{
+		if (FixturePatchNameTextBlock.IsValid())
+		{
+			FixturePatchNameTextBlock->EnterEditingMode();
+		}
+	}
+
+	return FReply::Handled();
+}
+
+void SDMXMVRFixtureListRow::OnFixturePatchNameCommitted(const FText& InNewText, ETextCommit::Type InTextCommit)
+{
+	if (InNewText.IsEmpty())
+	{
+		return;
+	}
+
+	FString ResultingName;
+	Item->SetFixturePatchName(InNewText.ToString(), ResultingName);
+	FixturePatchNameTextBlock->SetText(FText::FromString(ResultingName));
+}
+
+TSharedRef<SWidget> SDMXMVRFixtureListRow::GenerateStatusRow()
+{
+	return
+		SNew(SBox)
+		.HAlign(HAlign_Center)
+		.VAlign(VAlign_Center)
+		[
+			SNew(SImage)
+			.Image_Lambda([this]()
+				{
+					if (!Item->ErrorStatusText.IsEmpty())
+					{
+						return FAppStyle::GetBrush("Icons.Error");
+					}
+
+					if (!Item->WarningStatusText.IsEmpty())
+					{
+						return FAppStyle::GetBrush("Icons.Warning");
+					}
+
+					static const FSlateBrush EmptyBrush = FSlateNoResource();
+					return &EmptyBrush;
+				})
+			.ToolTipText_Lambda([this]()
+				{
+					if (!Item->ErrorStatusText.IsEmpty())
+					{
+						return Item->ErrorStatusText;
+					}
+					else if (!Item->WarningStatusText.IsEmpty())
+					{
+						return Item->WarningStatusText;
+					}
+
+					return FText::GetEmpty();
+				})
+		];
+}
+
+TSharedRef<SWidget> SDMXMVRFixtureListRow::GenerateUnitNumberRow()
+{
+	return
+		SNew(SBorder)
+		.HAlign(HAlign_Fill)
+		.VAlign(VAlign_Center)
+		.Padding(4.f)
+		.BorderImage(FAppStyle::GetBrush("NoBorder"))
+		.OnMouseDoubleClick(this, &SDMXMVRFixtureListRow::OnUnitNumberBorderDoubleClicked)
+		[
+			SNew(SBorder)
+			.HAlign(HAlign_Fill)
+			.VAlign(VAlign_Center)
+			.BorderImage(FDMXEditorStyle::Get().GetBrush("DMXEditor.RoundedPropertyBorder"))
+			[
+				SAssignNew(UnitNumberTextBlocK, SInlineEditableTextBlock)
+				.Text_Lambda([this]()
+				{
+					const int32 UnitNumber = Item->GetUnitNumber();
+					return FText::AsNumber(UnitNumber);
+				})
+				.Font(FCoreStyle::GetDefaultFontStyle("Regular", 10))
+				.OnTextCommitted(this, &SDMXMVRFixtureListRow::OnUnitNumberCommitted)
+				.IsSelected(IsSelected)
+			]
+		];
+}
+
+FReply SDMXMVRFixtureListRow::OnUnitNumberBorderDoubleClicked(const FGeometry& InMyGeometry, const FPointerEvent& InMouseEvent)
+{
+	if (InMouseEvent.GetEffectingButton() == EKeys::LeftMouseButton)
+	{
+		if (UnitNumberTextBlocK.IsValid())
+		{
+			UnitNumberTextBlocK->EnterEditingMode();
+		}
+	}
+
+	return FReply::Handled();
+}
+
+void SDMXMVRFixtureListRow::OnUnitNumberCommitted(const FText& InNewText, ETextCommit::Type InTextCommit)
+{
+	const FString StringValue = InNewText.ToString();
+	int32 NewUnitNumber;
+	if (LexTryParseString<int32>(NewUnitNumber, *StringValue))
+	{
+		Item->SetUnitNumber(NewUnitNumber);
+
+		const FString ParsedUnitNumberString = FString::FromInt(NewUnitNumber);
+		UnitNumberTextBlocK->SetText(FText::FromString(ParsedUnitNumberString));
+
+		OnRowRequestsStatusRefresh.ExecuteIfBound();
+	}
+}
+
+TSharedRef<SWidget> SDMXMVRFixtureListRow::GenerateMVRFixtureNameRow()
+{
+	return
+		SNew(SBorder)
+		.HAlign(HAlign_Fill)
+		.VAlign(VAlign_Center)
+		.Padding(4.f)
+		.BorderImage(FAppStyle::GetBrush("NoBorder"))
+		.OnMouseDoubleClick(this, &SDMXMVRFixtureListRow::OnMVRFixtureNameBorderDoubleClicked)
+		[
+			SNew(SBorder)
+			.HAlign(HAlign_Fill)
+			.VAlign(VAlign_Center)
+			.IsEnabled(true)
+			.BorderImage(FDMXEditorStyle::Get().GetBrush("DMXEditor.RoundedPropertyBorder"))
+			[
+				SAssignNew(NameTextBlock, SInlineEditableTextBlock)
+				.Text_Lambda([this]()
+				{
+					const FString& Name = Item->GetMVRFixtureName();
+					return FText::FromString(Name);
+				})
+				.Font(FCoreStyle::GetDefaultFontStyle("Regular", 10))
+				.OnTextCommitted(this, &SDMXMVRFixtureListRow::OnMVRFixtureNameCommitted)
+				.IsSelected(IsSelected)
+			]
+		];
+
+}
+
+FReply SDMXMVRFixtureListRow::OnMVRFixtureNameBorderDoubleClicked(const FGeometry& InMyGeometry, const FPointerEvent& InMouseEvent)
+{
+	if (InMouseEvent.GetEffectingButton() == EKeys::LeftMouseButton)
+	{
+		if (NameTextBlock.IsValid())
+		{
+			NameTextBlock->EnterEditingMode();
+		}
+	}
+
+	return FReply::Handled();
+}
+
+void SDMXMVRFixtureListRow::OnMVRFixtureNameCommitted(const FText& InNewText, ETextCommit::Type InTextCommit)
+{
+	const FString NameString = InNewText.ToString();
+	if (Item->SetMVRFixtureName(NameString))
+	{
+		NameTextBlock->SetText(FText::FromString(NameString));
+	}
+}
+
+TSharedRef<SWidget> SDMXMVRFixtureListRow::GenerateFixtureTypeRow()
+{
+	UDMXLibrary* DMXLibrary = Item->GetDMXLibrary();
+	if (!ensureAlwaysMsgf(DMXLibrary, TEXT("Tried to set fixture type for MVR Fixture, but fixture type is invalid.")))
+	{
+		return SNullWidget::NullWidget;
+	}
+
+	const TSharedRef<SDMXMVRFixtureFixtureTypePicker> FixtureTypePicker =
+		SNew(SDMXMVRFixtureFixtureTypePicker, DMXLibrary)
+		.OnFixtureTypeSelected(this, &SDMXMVRFixtureListRow::OnFixtureTypeSelected);
+
+	UDMXEntityFixtureType* SelectedFixtureType = Item->GetFixtureType();
+	FixtureTypePicker->SetSelection(SelectedFixtureType);
+		
+	return 
+		SNew(SBorder)
+		.HAlign(HAlign_Fill)
+		.VAlign(VAlign_Center)
+		.Padding(4.f)
+		.BorderImage(FAppStyle::GetBrush("NoBorder"))
+		[
+			FixtureTypePicker
+		];
+}
+
+void SDMXMVRFixtureListRow::OnFixtureTypeSelected(UDMXEntityFixtureType* SelectedFixtureType)
+{
+	Item->SetFixtureType(SelectedFixtureType);
+}
+
+TSharedRef<SWidget> SDMXMVRFixtureListRow::GenerateModeRow()
+{
+	const TSharedRef<SDMXMVRFixtureModePicker> ModePicker =
+		SNew(SDMXMVRFixtureModePicker)
+		.OnModeSelected(this, &SDMXMVRFixtureListRow::OnModeSelected);
+
+	UDMXEntityFixtureType* SelectedFixtureType = Item->GetFixtureType();
+	const int32 SelectedModeIndex = Item->GetModeIndex();
+
+	ModePicker->SetSelection(SelectedFixtureType, SelectedModeIndex);
+
+	return
+		SNew(SBorder)
+		.HAlign(HAlign_Fill)
+		.VAlign(VAlign_Center)
+		.Padding(4.f)
+		.BorderImage(FAppStyle::GetBrush("NoBorder"))
+		[
+			ModePicker
+		];
+}
+
+void SDMXMVRFixtureListRow::OnModeSelected(int32 SelectedModeIndex)
+{
+	Item->SetModeIndex(SelectedModeIndex);
+}
+
+TSharedRef<SWidget> SDMXMVRFixtureListRow::GeneratePatchRow()
+{
+	return
+		SNew(SBorder)
+		.HAlign(HAlign_Fill)
+		.VAlign(VAlign_Center)
+		.Padding(4.f)
+		.BorderImage(FAppStyle::GetBrush("NoBorder"))
+		.OnMouseDoubleClick(this, &SDMXMVRFixtureListRow::OnPatchBorderDoubleClicked)
+		[
+			SNew(SBorder)
+			.HAlign(HAlign_Fill)
+			.VAlign(VAlign_Center)
+			.BorderImage(FDMXEditorStyle::Get().GetBrush("DMXEditor.RoundedPropertyBorder"))
+			[
+				SAssignNew(PatchTextBlock, SInlineEditableTextBlock)
+				.Text_Lambda([this]()
+					{
+						const int32 UniverseID = Item->GetUniverse();
+						const int32 StartingAddress = Item->GetAddress();
+						return FText::Format(LOCTEXT("AddressesText", "{0}.{1}"), UniverseID, StartingAddress);
+					})
+				.OnTextCommitted(this, &SDMXMVRFixtureListRow::OnPatchCommitted)
+				.IsSelected(IsSelected)
+			]
+		];
+}
+
+FReply SDMXMVRFixtureListRow::OnPatchBorderDoubleClicked(const FGeometry& InMyGeometry, const FPointerEvent& InMouseEvent)
+{
+	if (InMouseEvent.GetEffectingButton() == EKeys::LeftMouseButton)
+	{
+		if (PatchTextBlock.IsValid())
+		{
+			PatchTextBlock->EnterEditingMode();
+		}
+	}
+
+	return FReply::Handled();
+}
+
+void SDMXMVRFixtureListRow::OnPatchCommitted(const FText& InNewText, ETextCommit::Type InTextCommit)
+{
+	const FString PatchString = InNewText.ToString();
+	static const TCHAR* ParamDelimiters[] =
+	{
+		TEXT("."),
+		TEXT(","),
+		TEXT(":"),
+		TEXT(";")
+	};
+
+	TArray<FString> ValueStringArray;
+	PatchString.ParseIntoArray(ValueStringArray, ParamDelimiters, 4);
+	if (ValueStringArray.Num() == 2)
+	{
+		int32 Universe;
+		if (!LexTryParseString(Universe, *ValueStringArray[0]))
+		{
+			return;
+		}
+		
+		int32 Address;
+		if (!LexTryParseString(Address, *ValueStringArray[1]))
+		{
+			return;
+		}
+
+		Item->SetAddresses(Universe, Address);
+
+		const FString UniverseString = FString::FromInt(Universe);
+		const FString AddressString = FString::FromInt(Address);
+		PatchTextBlock->SetText(FText::Format(LOCTEXT("UniverseDotAddressText", "{0}.{1}"), FText::FromString(UniverseString), FText::FromString(AddressString)));
+
+		OnRowRequestsListRefresh.ExecuteIfBound();
+	}
+}
+
+#undef LOCTEXT_NAMESPACE
diff --git a/Engine/Plugins/VirtualProduction/DMX/DMXEngine/Source/DMXEditor/Private/Widgets/FixturePatch/SDMXMVRFixtureListRow.h b/Engine/Plugins/VirtualProduction/DMX/DMXEngine/Source/DMXEditor/Private/Widgets/FixturePatch/SDMXMVRFixtureListRow.h
new file mode 100644
--- /dev/null
+++ b/Engine/Plugins/VirtualProduction/DMX/DMXEngine/Source/DMXEditor/Private/Widgets/FixturePatch/SDMXMVRFixtureListRow.h
@@ -0,0 +1,126 @@
+// Copyright Epic Games, Inc. All Rights Reserved.
+
+#pragma once
+
+#include "Library/DMXEntityReference.h"
+
+#include "CoreMinimal.h"
+#include "Widgets/Views/STableRow.h"
+
+enum class EDMXMVRFixtureListEditMode : uint8;
+class UDMXEntityFixturePatch;
+class FDMXMVRFixtureListItem;
+
+class SInlineEditableTextBlock;
+
+
+/** MVR Fixture view as a row in a list */
+class SDMXMVRFixtureListRow
+	: public SMultiColumnTableRow<TSharedPtr<FDMXMVRFixtureListItem>>
+{
+public:
+	SLATE_BEGIN_ARGS(SDMXMVRFixtureListRow)
+	{}
+		/** Delegate executed when the row requests to refresh the statuses */
+		SLATE_EVENT(FSimpleDelegate, OnRowRequestsStatusRefresh)
+
+		/** Delegate executed when the row requests to refresh the whole list */
+		SLATE_EVENT(FSimpleDelegate, OnRowRequestsListRefresh)
+
+		/** Callback to check if the row is selected (should be hooked up if a parent widget is handling selection or focus) */
+		SLATE_EVENT(FIsSelected, IsSelected)
+
+	SLATE_END_ARGS()
+			
+	/** Constructs this widget */
+	void Construct(const FArguments& InArgs, const TSharedRef<STableViewBase>& InOwnerTable, const TSharedRef<FDMXMVRFixtureListItem>& InItem);
+
+	/** Enters editing mode for the Fixture Patch Name */
+	void EnterFixturePatchNameEditingMode();
+
+	/** Returns the Item of this row */
+	TSharedPtr<FDMXMVRFixtureListItem> GetItem() const { return Item; };
+
+protected:
+	//~ Begin SMultiColumnTableRow interface
+	virtual TSharedRef<SWidget> GenerateWidgetForColumn(const FName& ColumnName) override;
+	//~ End SMultiColumnTableRow interface
+
+private:
+	/** Generates the row that displays the Fixture Patch Name */
+	TSharedRef<SWidget> GenerateFixturePatchNameRow();
+
+	/** Called when the Fixture Patch Name Border was double-clicked */
+	FReply OnFixturePatchNameBorderDoubleClicked(const FGeometry& InMyGeometry, const FPointerEvent& InMouseEvent);
+
+	/** Called when a Fixture Patch Name was committed */
+	void OnFixturePatchNameCommitted(const FText& InNewText, ETextCommit::Type InTextCommit);
+
+	/** Generates the row that displays the Status */
+	TSharedRef<SWidget> GenerateStatusRow();
+
+	/** Generates the row that displays the Unit Number */
+	TSharedRef<SWidget> GenerateUnitNumberRow();
+
+	/** Called when the Unit Number Border was double-clicked */
+	FReply OnUnitNumberBorderDoubleClicked(const FGeometry& InMyGeometry, const FPointerEvent& InMouseEvent);
+
+	/** Called when a Unit Number was committed */
+	void OnUnitNumberCommitted(const FText& InNewText, ETextCommit::Type InTextCommit);
+
+	/** Generates the row that displays the MVR Fixture Name */
+	TSharedRef<SWidget> GenerateMVRFixtureNameRow();
+
+	/** Called when the Name Border was double-clicked */
+	FReply OnMVRFixtureNameBorderDoubleClicked(const FGeometry& InMyGeometry, const FPointerEvent& InMouseEvent);
+
+	/** Called when the MVR Fixture Name was committed */
+	void OnMVRFixtureNameCommitted(const FText& InNewText, ETextCommit::Type InTextCommit);
+
+	/** Generates the row that displays the Fixture Type */
+	TSharedRef<SWidget> GenerateFixtureTypeRow();
+
+	/** Called when a Fixture Type was selected */
+	void OnFixtureTypeSelected(UDMXEntityFixtureType* SelectedFixtureType);
+
+	/** Generates the row that displays the Mode */
+	TSharedRef<SWidget> GenerateModeRow();
+
+	/** Called when a Mode was selected */
+	void OnModeSelected(int32 SelectedModeIndex);
+
+	/** Generates the row that displays the Patch */
+	TSharedRef<SWidget> GeneratePatchRow();
+
+	/** Called when the Patch Border was double-clicked */
+	FReply OnPatchBorderDoubleClicked(const FGeometry& InMyGeometry, const FPointerEvent& InMouseEvent);
+
+	/** Called when a Patch was committed */
+	void OnPatchCommitted(const FText& InNewText, ETextCommit::Type InTextCommit);
+
+	/** The Edit Mode the Widget should present */
+	EDMXMVRFixtureListEditMode EditMode;
+
+	/** The outermost border around the the Fixture Patch Name Column */
+	TSharedPtr<SBorder> FixturePatchNameBorder;
+
+	/** The text block to edit the Fixture Patch Name */
+	TSharedPtr<SInlineEditableTextBlock> FixturePatchNameTextBlock;
+
+	/** The text block to edit the Unit Number */
+	TSharedPtr<SInlineEditableTextBlock> UnitNumberTextBlocK;
+
+	/** The text block to edit the Name */
+	TSharedPtr<SInlineEditableTextBlock> NameTextBlock;
+
+	/** The text block to edit the Name */
+	TSharedPtr<SInlineEditableTextBlock> PatchTextBlock;
+
+	/** The MVR Fixture List Item this row displays */
+	TSharedPtr<FDMXMVRFixtureListItem> Item;
+
+	// Slate arguments
+	FSimpleDelegate OnRowRequestsStatusRefresh;
+	FSimpleDelegate OnRowRequestsListRefresh;
+	FIsSelected IsSelected;
+};
diff --git a/Engine/Plugins/VirtualProduction/DMX/DMXEngine/Source/DMXEditor/Private/Widgets/FixturePatch/SDMXMVRFixtureListToolbar.cpp b/Engine/Plugins/VirtualProduction/DMX/DMXEngine/Source/DMXEditor/Private/Widgets/FixturePatch/SDMXMVRFixtureListToolbar.cpp
new file mode 100644
--- /dev/null
+++ b/Engine/Plugins/VirtualProduction/DMX/DMXEngine/Source/DMXEditor/Private/Widgets/FixturePatch/SDMXMVRFixtureListToolbar.cpp
@@ -0,0 +1,376 @@
+// Copyright Epic Games, Inc. All Rights Reserved.
+
+#include "SDMXMVRFixtureListToolbar.h"
+
+#include "DMXEditor.h"
+#include "DMXMVRFixtureListItem.h"
+#include "Library/DMXEntityFixturePatch.h"
+#include "Library/DMXEntityFixtureType.h"
+#include "Commands/DMXEditorCommands.h"
+#include "Widgets/SDMXEntityDropdownMenu.h"
+
+#include "EditorStyleSet.h"
+#include "Internationalization/Regex.h"
+#include "Widgets/SBoxPanel.h"
+#include "Widgets/Input/SCheckBox.h"
+#include "Widgets/Input/SSearchBox.h"
+#include "Widgets/Layout/SSeparator.h"
+#include "Widgets/Layout/SWrapBox.h"
+
+
+#define LOCTEXT_NAMESPACE "SDMXMVRFixtureListToolbar"
+
+namespace UE::DMXRuntime::SDMXMVRFixtureListToolbar::Private
+{
+	UE_NODISCARD bool ParseUniverse(const FString& InputString, int32& OutUniverse)
+	{
+		// Try to match addresses formating, e.g. '1.', '1:' etc.
+		static const TCHAR* ParamDelimiters[] =
+		{
+			TEXT("."),
+			TEXT(","),
+			TEXT(":"),
+			TEXT(";")
+		};
+
+		TArray<FString> ValueStringArray;
+		constexpr bool bParseEmpty = true;
+		const FString InputStringWithSpace = InputString + TEXT(" "); // So ValueStringArray will be lenght of 2 if Address is empty, e.g. '0.'
+		InputStringWithSpace.ParseIntoArray(ValueStringArray, ParamDelimiters, 4, bParseEmpty);
+		if (ValueStringArray.Num() == 2)
+		{
+			if (LexTryParseString<int32>(OutUniverse, *ValueStringArray[0]))
+			{
+				return true;
+			}
+		}
+
+		// Try to match strings starting with Uni, e.g. 'Uni 1', 'Universe 1', 'Universe1'
+		if (InputString.StartsWith(TEXT("Uni")))
+		{
+			const FRegexPattern SequenceOfDigitsPattern(TEXT("^[^\\d]*(\\d+)"));
+			FRegexMatcher Regex(SequenceOfDigitsPattern, *InputString);
+			if (Regex.FindNext())
+			{
+				const FString UniverseString = Regex.GetCaptureGroup(1);
+				if (LexTryParseString<int32>(OutUniverse, *UniverseString))
+				{
+					return true;
+				}
+			}
+		}
+
+		OutUniverse = -1;
+		return false;
+	}
+
+	UE_NODISCARD bool ParseAddress(const FString& InputString, int32& OutAddress)
+	{
+		// Try to match addresses formating, e.g. '1.1', '1:1' etc.
+		static const TCHAR* ParamDelimiters[] =
+		{
+			TEXT("."),
+			TEXT(","),
+			TEXT(":"),
+			TEXT(";")
+		};
+
+		TArray<FString> ValueStringArray;
+		constexpr bool bParseEmpty = false;
+		InputString.ParseIntoArray(ValueStringArray, ParamDelimiters, 4, bParseEmpty);
+
+		if (ValueStringArray.Num() == 2)
+		{
+			if (LexTryParseString<int32>(OutAddress, *ValueStringArray[1]))
+			{
+				return true;
+			}
+		}
+
+		// Try to match strings starting with Uni Ad, e.g. 'Uni 1 Ad 1', 'Universe 1 Address 1', 'Universe1Address1'
+		if (InputString.StartsWith(TEXT("Uni")) &&
+			InputString.Contains(TEXT("Ad")))
+		{
+			const FRegexPattern SequenceOfDigitsPattern(TEXT("^[^\\d]*(\\d+)[^\\d]*(\\d+)"));
+			FRegexMatcher Regex(SequenceOfDigitsPattern, *InputString);
+			if (Regex.FindNext())
+			{
+				const FString AddressString = Regex.GetCaptureGroup(2);
+				if (LexTryParseString<int32>(OutAddress, *AddressString))
+				{
+					return true;
+				}
+			}
+		}
+
+		OutAddress = -1;
+		return false;
+	}
+
+	UE_NODISCARD bool ParseUnitNumber(const FString& InputString, int32& OutUnitNumber)
+	{
+		if (LexTryParseString<int32>(OutUnitNumber, *InputString))
+		{
+			return true;
+		}
+
+		OutUnitNumber = -1;
+		return false;
+	}
+}
+
+void SDMXMVRFixtureListToolbar::Construct(const FArguments& InArgs, TWeakPtr<FDMXEditor> InDMXEditor)
+{
+	WeakDMXEditor = InDMXEditor;
+	OnSearchChanged = InArgs._OnSearchChanged;
+
+	ChildSlot
+		.HAlign(HAlign_Fill)
+		.VAlign(VAlign_Fill)
+		[
+			SNew(SBorder)
+			.Padding(FMargin(8.f, 8.f))
+			.BorderBackgroundColor(FLinearColor(0.6f, 0.6f, 0.6f))
+			.BorderImage(FAppStyle::GetBrush("DetailsView.CategoryTop"))
+			[
+				SNew(SWrapBox)
+				.InnerSlotPadding(FVector2D(14.f, 8.f))
+				.UseAllottedWidth(true)
+
+				// Add Fixture Button
+				+ SWrapBox::Slot()
+				.HAlign(HAlign_Left)
+				.VAlign(VAlign_Center)
+				[
+					GenerateFixtureTypeDropdownMenu()
+				]
+										
+				+ SWrapBox::Slot()
+				.HAlign(HAlign_Left)
+				.VAlign(VAlign_Center)
+				[
+					SNew(SSeparator)
+					.Orientation(EOrientation::Orient_Vertical)
+				]
+
+				// Search
+				+ SWrapBox::Slot()
+				.HAlign(HAlign_Left)
+				.VAlign(VAlign_Center)
+				[
+					SNew(SSearchBox)
+					.MinDesiredWidth(400.f)
+					.OnTextChanged(this, &SDMXMVRFixtureListToolbar::OnSearchTextChanged)
+				]
+									
+				+ SWrapBox::Slot()
+				.HAlign(HAlign_Left)
+				.VAlign(VAlign_Center)
+				[
+					SNew(SSeparator)
+					.Orientation(EOrientation::Orient_Vertical)
+				]
+
+				// Show Conflicts Only option
+				+ SWrapBox::Slot()
+				.HAlign(HAlign_Left)
+				.VAlign(VAlign_Center)
+				[
+					SNew(STextBlock)
+					.Text(LOCTEXT("ShowConflictsOnlyCheckBoxLabel", "Show Conflicts only"))
+					.Font(FAppStyle::GetFontStyle(TEXT("PropertyWindow.NormalFont")))
+				]
+
+				+ SWrapBox::Slot()
+				.HAlign(HAlign_Left)
+				.VAlign(VAlign_Center)
+				[
+					SNew(SCheckBox)
+					.IsChecked(false)
+					.OnCheckStateChanged(this, &SDMXMVRFixtureListToolbar::OnShowConflictsOnlyCheckStateChanged)
+				]
+			]
+		];
+}
+
+TArray<TSharedPtr<FDMXMVRFixtureListItem>> SDMXMVRFixtureListToolbar::FilterItems(const TArray<TSharedPtr<FDMXMVRFixtureListItem>>& Items)
+{
+	using namespace UE::DMXRuntime::SDMXMVRFixtureListToolbar::Private;
+
+	// Apply 'conflicts only' if enabled
+	TArray<TSharedPtr<FDMXMVRFixtureListItem>> Result = Items;
+	if (bShowConfictsOnly)
+	{
+		Result.RemoveAll([](const TSharedPtr<FDMXMVRFixtureListItem>& Item)
+			{
+				return
+					Item->ErrorStatusText.IsEmpty() &&
+					Item->WarningStatusText.IsEmpty();
+			});
+	}
+
+	// Filter and return in order of precendence
+	if (SearchString.IsEmpty())
+	{
+		return Result;
+	}
+
+	int32 Universe;
+	if (ParseUniverse(SearchString, Universe))
+	{
+		Result.RemoveAll([Universe](const TSharedPtr<FDMXMVRFixtureListItem>& Item)
+			{
+				return Item->GetUniverse() != Universe;
+			});
+
+		int32 Address;
+		if (ParseAddress(SearchString, Address))
+		{
+			Result.RemoveAll([Address](const TSharedPtr<FDMXMVRFixtureListItem>& Item)
+				{
+					return Item->GetAddress() != Address;
+				});
+		}
+
+		return Result;
+	}
+	
+	int32 UnitNumber;
+	if (ParseUnitNumber(SearchString, UnitNumber))
+	{
+		TArray<TSharedPtr<FDMXMVRFixtureListItem>> UnitNumbersOnlyResult = Result;
+		UnitNumbersOnlyResult.RemoveAll([UnitNumber](const TSharedPtr<FDMXMVRFixtureListItem>& Item)
+			{
+				return Item->GetUnitNumber() != UnitNumber;
+			});
+
+		if (UnitNumbersOnlyResult.Num() > 0)
+		{
+			return UnitNumbersOnlyResult;
+		}
+	}
+
+	Result.RemoveAll([this](const TSharedPtr<FDMXMVRFixtureListItem>& Item)
+		{
+			return
+				!Item->GetFixturePatchName().Contains(SearchString) &&
+				!Item->GetMVRFixtureName().Contains(SearchString);
+		});
+
+	return Result;
+}
+
+TSharedRef<SWidget> SDMXMVRFixtureListToolbar::GenerateFixtureTypeDropdownMenu()
+{
+	FText AddButtonLabel = FDMXEditorCommands::Get().AddNewEntityFixturePatch->GetLabel();
+	FText AddButtonToolTip = FDMXEditorCommands::Get().AddNewEntityFixturePatch->GetDescription();
+
+	TSharedRef<SComboButton> AddComboButton = 
+		SNew(SComboButton)
+			.ButtonContent()
+			[
+				SNew(SHorizontalBox)
+				+ SHorizontalBox::Slot()
+				.AutoWidth()
+				.Padding(FMargin(0, 1))
+				[
+					SNew(SImage)
+					.Image(FAppStyle::GetBrush("Plus"))
+				]
+				+ SHorizontalBox::Slot()
+				.VAlign(VAlign_Center)
+				.AutoWidth()
+				.Padding(FMargin(2, 0, 2, 0))
+				[
+					SNew(STextBlock)
+					.Text(AddButtonLabel)
+				]
+			]
+			.MenuContent()
+			[
+				SAssignNew(FixtureTypeDropdownMenu, SDMXEntityDropdownMenu<UDMXEntityFixtureType>)
+				.DMXEditor(WeakDMXEditor)
+				.OnEntitySelected(this, &SDMXMVRFixtureListToolbar::OnAddNewMVRFixtureClicked)
+			]
+			.IsFocusable(true)
+			.ContentPadding(FMargin(5.0f, 1.0f))
+			.ComboButtonStyle(FAppStyle::Get(), "ToolbarComboButton")
+			.ButtonStyle(FAppStyle::Get(), "FlatButton.Success")
+			.ForegroundColor(FLinearColor::White)
+			.ToolTipText(AddButtonToolTip)
+			.OnComboBoxOpened(FOnComboBoxOpened::CreateLambda([this]() 
+				{ 
+					FixtureTypeDropdownMenu->RefreshEntitiesList();
+				}));
+
+	FixtureTypeDropdownMenu->SetComboButton(AddComboButton);
+
+	return AddComboButton;
+}
+
+void SDMXMVRFixtureListToolbar::OnAddNewMVRFixtureClicked(UDMXEntity* InSelectedFixtureType)
+{
+	TSharedPtr<FDMXEditor> DMXEditor = WeakDMXEditor.Pin();
+	if (!DMXEditor.IsValid())
+	{
+		return;
+	}
+	UDMXLibrary* DMXLibrary = DMXEditor->GetDMXLibrary();
+	if (!DMXLibrary)
+	{
+		return;
+	}
+	UDMXEntityFixtureType* FixtureType = CastChecked<UDMXEntityFixtureType>(InSelectedFixtureType);
+
+	// Find a universe and address
+	TArray<UDMXEntityFixturePatch*> FixturePatches = DMXLibrary->GetEntitiesTypeCast<UDMXEntityFixturePatch>();
+	FixturePatches.Sort([](const UDMXEntityFixturePatch& FixturePatchA, const UDMXEntityFixturePatch& FixturePatchB)
+		{
+			const bool bUniverseIsSmaller = FixturePatchA.GetUniverseID() < FixturePatchB.GetUniverseID();
+			const bool bUniverseIsEqual = FixturePatchA.GetUniverseID() == FixturePatchB.GetUniverseID();
+			const bool bAddressIsSmaller = FixturePatchA.GetStartingChannel() < FixturePatchB.GetStartingChannel();
+			
+			return bUniverseIsSmaller || (bUniverseIsEqual && bAddressIsSmaller);
+		});
+
+	int32 Universe = 1; 
+	int32 Address = 1;
+	if (FixturePatches.Num() > 0)
+	{
+		const int32 ChannelSpan = FixtureType->Modes.Num() > 0 ? FixtureType->Modes[0].ChannelSpan : 1;
+		if (FixturePatches.Last()->GetEndingChannel() + 1 + ChannelSpan <= DMX_MAX_ADDRESS)
+		{
+			Universe = FixturePatches.Last()->GetUniverseID() + 1;
+			Address = 1;
+		}
+		else
+		{
+			Universe = FixturePatches.Last()->GetUniverseID();
+			Address = FixturePatches.Last()->GetEndingChannel() + 1;
+		}
+	}
+
+	// Create a new fixture patch
+	FDMXEntityFixturePatchConstructionParams FixturePatchConstructionParams;
+	FixturePatchConstructionParams.FixtureTypeRef = FDMXEntityFixtureTypeRef(FixtureType);
+	FixturePatchConstructionParams.ActiveMode = 0;
+	FixturePatchConstructionParams.UniverseID = Universe;
+	FixturePatchConstructionParams.StartingAddress = Address;
+
+	constexpr bool bMarkLibraryDirty = true;
+	UDMXEntityFixturePatch* NewFixturePatch = UDMXEntityFixturePatch::CreateFixturePatchInLibrary(FixturePatchConstructionParams, FixtureType->Name, bMarkLibraryDirty);
+}
+
+void SDMXMVRFixtureListToolbar::OnSearchTextChanged(const FText& SearchText)
+{
+	SearchString = SearchText.ToString();
+	OnSearchChanged.ExecuteIfBound();
+}
+
+void SDMXMVRFixtureListToolbar::OnShowConflictsOnlyCheckStateChanged(const ECheckBoxState NewCheckState)
+{
+	bShowConfictsOnly = NewCheckState == ECheckBoxState::Checked;
+	OnSearchChanged.ExecuteIfBound();
+}
+
+#undef LOCTEXT_NAMESPACE
diff --git a/Engine/Plugins/VirtualProduction/DMX/DMXEngine/Source/DMXEditor/Private/Widgets/FixturePatch/SDMXMVRFixtureListToolbar.h b/Engine/Plugins/VirtualProduction/DMX/DMXEngine/Source/DMXEditor/Private/Widgets/FixturePatch/SDMXMVRFixtureListToolbar.h
new file mode 100644
--- /dev/null
+++ b/Engine/Plugins/VirtualProduction/DMX/DMXEngine/Source/DMXEditor/Private/Widgets/FixturePatch/SDMXMVRFixtureListToolbar.h
@@ -0,0 +1,63 @@
+// Copyright Epic Games, Inc. All Rights Reserved.
+
+#pragma once
+
+#include "CoreMinimal.h"
+#include "Widgets/SCompoundWidget.h"
+
+class FDMXEditor;
+class FDMXMVRFixtureListItem;
+template<typename TEntityType> class SDMXEntityDropdownMenu;
+class UDMXEntity;
+class UDMXEntityFixtureType;
+
+enum class ECheckBoxState : uint8;
+
+
+/** Search bar for the MVR Fixture List */
+class SDMXMVRFixtureListToolbar
+	: public SCompoundWidget
+{
+public:
+	SLATE_BEGIN_ARGS(SDMXMVRFixtureListToolbar)
+	{}
+
+		/** Executed when the search changed */
+		SLATE_EVENT(FSimpleDelegate, OnSearchChanged)
+
+	SLATE_END_ARGS()
+
+	/** Constructs this widget */
+	void Construct(const FArguments& InArgs, TWeakPtr<FDMXEditor> InDMXEditor);
+
+	/** Filters items according to the state of the search bar */
+	UE_NODISCARD TArray<TSharedPtr<FDMXMVRFixtureListItem>> FilterItems(const TArray<TSharedPtr<FDMXMVRFixtureListItem>>& Items);
+
+private:
+	/** Generates the 'Add MVR Fixture' Dropdown Menu */
+	TSharedRef<SWidget> GenerateFixtureTypeDropdownMenu();
+
+	/** Called when a new Fixture Patch was selected in the FixtureTypeDropdownMenu */
+	void OnAddNewMVRFixtureClicked(UDMXEntity* InSelectedFixtureType);
+
+	/** Called when the Search Text changed */
+	void OnSearchTextChanged(const FText& SearchText);
+
+	/** Called when the Check State of the 'Show Conflicts Only' Check Box changed */
+	void OnShowConflictsOnlyCheckStateChanged(const ECheckBoxState NewCheckState);
+
+	/** Current Search String */
+	FString SearchString;
+
+	/** True if only conflicts are searched for */
+	bool bShowConfictsOnly = false;
+
+	/** Dropdown menu of Fixture Types */
+	TSharedPtr<SDMXEntityDropdownMenu<UDMXEntityFixtureType>> FixtureTypeDropdownMenu;
+
+	/** The DMX Editor that owns this widget */
+	TWeakPtr<FDMXEditor> WeakDMXEditor;
+
+	// Slate Args
+	FSimpleDelegate OnSearchChanged;
+};
diff --git a/Engine/Plugins/VirtualProduction/DMX/DMXEngine/Source/DMXEditor/Private/Widgets/FixtureType/SDMXFixtureTypeFunctionsEditor.cpp b/Engine/Plugins/VirtualProduction/DMX/DMXEngine/Source/DMXEditor/Private/Widgets/FixtureType/SDMXFixtureTypeFunctionsEditor.cpp
--- a/Engine/Plugins/VirtualProduction/DMX/DMXEngine/Source/DMXEditor/Private/Widgets/FixtureType/SDMXFixtureTypeFunctionsEditor.cpp
+++ b/Engine/Plugins/VirtualProduction/DMX/DMXEngine/Source/DMXEditor/Private/Widgets/FixtureType/SDMXFixtureTypeFunctionsEditor.cpp
@@ -368,7 +368,7 @@ void SDMXFixtureTypeFunctionsEditor::OnFixtureTypeChanged(const UDMXEntityFixtur
 
 TSharedRef<SHeaderRow> SDMXFixtureTypeFunctionsEditor::GenerateHeaderRow()
 {
-	const float StatusColumnWidth = FMath::Max(FAppStyle::GetBrush("Icons.Warning")->GetImageSize().X + 6.f, FAppStyle::GetBrush("Icons.Error")->GetImageSize().X + 6.f);
+	const float StatusColumnWidth = FMath::Max(FAppStyle::Get().GetBrush("Icons.Warning")->GetImageSize().X + 6.f, FAppStyle::Get().GetBrush("Icons.Error")->GetImageSize().X + 6.f);
 
 	HeaderRow = SNew(SHeaderRow);
 	SHeaderRow::FColumn::FArguments ColumnArgs;
@@ -405,15 +405,15 @@ TSharedRef<SHeaderRow> SDMXFixtureTypeFunctionsEditor::GenerateHeaderRow()
 	if (const UDMXEditorSettings* EditorSettings = GetDefault<UDMXEditorSettings>())
 	{
 		const float NameColumnWidth = EditorSettings->FixtureTypeFunctionsEditorSettings.NameColumnWidth;
-		if (NameColumnWidth > 0.f)
+		if (NameColumnWidth > 10.f)
 		{
 			HeaderRow->SetColumnWidth("Name", NameColumnWidth);
 		}
 
 		const float AttributeColumnWidth = EditorSettings->FixtureTypeFunctionsEditorSettings.AttributeColumnWidth;
-		if (AttributeColumnWidth > 0.f)
+		if (AttributeColumnWidth > 10.f)
 		{
-			HeaderRow->SetColumnWidth("Attribute", AttributeColumnWidth);
+			HeaderRow->SetColumnWidth(FDMXFixtureTypeFunctionsEditorCollumnIDs::Attribute, AttributeColumnWidth);
 		}
 	}
 
@@ -427,12 +427,12 @@ void SDMXFixtureTypeFunctionsEditor::SaveHeaderRowSettings()
 	{
 		for (const SHeaderRow::FColumn& Column : HeaderRow->GetColumns())
 		{
-			if (Column.ColumnId == "Name")
+			if (Column.ColumnId == FDMXFixtureTypeFunctionsEditorCollumnIDs::Name)
 			{
 				EditorSettings->FixtureTypeFunctionsEditorSettings.NameColumnWidth = Column.Width.Get();
 
 			}
-			else if (Column.ColumnId == "Attribute")
+			else if (Column.ColumnId == FDMXFixtureTypeFunctionsEditorCollumnIDs::Attribute)
 			{
 				EditorSettings->FixtureTypeFunctionsEditorSettings.AttributeColumnWidth = Column.Width.Get();
 			}
diff --git a/Engine/Plugins/VirtualProduction/DMX/DMXEngine/Source/DMXEditor/Public/DMXEditorSettings.h b/Engine/Plugins/VirtualProduction/DMX/DMXEngine/Source/DMXEditor/Public/DMXEditorSettings.h
--- a/Engine/Plugins/VirtualProduction/DMX/DMXEngine/Source/DMXEditor/Public/DMXEditorSettings.h
+++ b/Engine/Plugins/VirtualProduction/DMX/DMXEngine/Source/DMXEditor/Public/DMXEditorSettings.h
@@ -9,6 +9,33 @@
 #include "DMXEditorSettings.generated.h"
 
 
+/** Settings for the Fixture Patch List */
+USTRUCT()
+struct FDMXMVRFixtureListSettings
+{
+	GENERATED_BODY()
+
+	/** Width of the Unit Number column */
+	UPROPERTY()
+	float UnitNumberColumnWidth = 0.f;
+
+	/** Width of the Name column */
+	UPROPERTY()
+	float NameColumnWidth = 0.f;
+
+	/** With of the Fixture Type column */
+	UPROPERTY()
+	float FixtureTypeColumnWidth = 0.f;
+
+	/** With of the Mode column */
+	UPROPERTY()
+	float ModeColumnWidth = 0.f;
+
+	/** With of the Patch column */
+	UPROPERTY()
+	float PatchColumnWidth = 0.f;
+};
+
 /** Settings for the Fixture Type Functions Editor */
 USTRUCT()
 struct FDMXFixtureTypeFunctionsEditorSettings
@@ -98,6 +125,13 @@ class UDMXEditorSettings : public UObject
 {
 	GENERATED_BODY()
 
+
+	// MVR Fixture List
+public:
+	UPROPERTY(Config)
+	FDMXMVRFixtureListSettings MVRFixtureListSettings;
+
+
 	// Fixture Type Functions Editor
 public:
 	UPROPERTY(Config)
diff --git a/Engine/Plugins/VirtualProduction/DMX/DMXEngine/Source/DMXEditor/Public/Factories/DMXLibraryFromMVRFactory.h b/Engine/Plugins/VirtualProduction/DMX/DMXEngine/Source/DMXEditor/Public/Factories/DMXLibraryFromMVRFactory.h
--- a/Engine/Plugins/VirtualProduction/DMX/DMXEngine/Source/DMXEditor/Public/Factories/DMXLibraryFromMVRFactory.h
+++ b/Engine/Plugins/VirtualProduction/DMX/DMXEngine/Source/DMXEditor/Public/Factories/DMXLibraryFromMVRFactory.h
@@ -6,11 +6,10 @@
 
 #include "DMXLibraryFromMVRFactory.generated.h"
 
-struct FDMXMVRGeneralSceneDescription;
 class FDMXMVRUnzip;
 class UDMXImportGDTF;
 class UDMXLibrary;
-
+class UDMXMVRGeneralSceneDescription;
 
 
 UCLASS()
@@ -38,9 +37,11 @@ private:
 	UDMXLibrary* CreateDMXLibraryAsset(UObject* Parent, EObjectFlags Flags, const FString& InFilename);
 
 	/** Creates GDTF assets from the MVR */
-	TArray<UDMXImportGDTF*> CreateGDTFAssets(UObject* Parent, EObjectFlags Flags, const TSharedRef<FDMXMVRUnzip>& MVRUnzip, const FDMXMVRGeneralSceneDescription& GeneralSceneDescription);
+	TArray<UDMXImportGDTF*> CreateGDTFAssets(UObject* Parent, EObjectFlags Flags, const TSharedRef<FDMXMVRUnzip>& MVRUnzip, const UDMXMVRGeneralSceneDescription& GeneralSceneDescription);
 
 	/** Initializes the DMX Library from the General Scene Description and GDTF assets */
-	void InitDMXLibrary(UDMXLibrary* DMXLibrary, const TArray<UDMXImportGDTF*>& GDTFAssets, const FDMXMVRGeneralSceneDescription& GeneralSceneDescription) const;
+	void InitDMXLibrary(UDMXLibrary* DMXLibrary, const TArray<UDMXImportGDTF*>& GDTFAssets, UDMXMVRGeneralSceneDescription* GeneralSceneDescription) const;
 
+	/** The DMX Library Package Name, initialized when the DMX Library Asset is created */
+	FString DMXLibraryPackageName;
 };
diff --git a/Engine/Plugins/VirtualProduction/DMX/DMXEngine/Source/DMXRuntime/Private/DMXRuntimeMainStreamObjectVersion.cpp b/Engine/Plugins/VirtualProduction/DMX/DMXEngine/Source/DMXRuntime/Private/DMXRuntimeMainStreamObjectVersion.cpp
new file mode 100644
--- /dev/null
+++ b/Engine/Plugins/VirtualProduction/DMX/DMXEngine/Source/DMXRuntime/Private/DMXRuntimeMainStreamObjectVersion.cpp
@@ -0,0 +1,9 @@
+// Copyright Epic Games, Inc. All Rights Reserved.
+
+#include "DMXRuntimeMainStreamObjectVersion.h"
+#include "Serialization/CustomVersion.h"
+
+const FGuid FDMXRuntimeMainStreamObjectVersion::GUID(0x75A7B683, 0x4F4C7E28, 0xDDE54885, 0x85323939);
+
+// Register the custom version with core
+FCustomVersionRegistration GRegisterDMXRuntimeMainStreamCustomVersion(FDMXRuntimeMainStreamObjectVersion::GUID, FDMXRuntimeMainStreamObjectVersion::LatestVersion, TEXT("DMXRuntimeMainStreamObjectVersion"));
diff --git a/Engine/Plugins/VirtualProduction/DMX/DMXEngine/Source/DMXRuntime/Private/DMXRuntimeMainStreamObjectVersion.h b/Engine/Plugins/VirtualProduction/DMX/DMXEngine/Source/DMXRuntime/Private/DMXRuntimeMainStreamObjectVersion.h
new file mode 100644
--- /dev/null
+++ b/Engine/Plugins/VirtualProduction/DMX/DMXEngine/Source/DMXRuntime/Private/DMXRuntimeMainStreamObjectVersion.h
@@ -0,0 +1,45 @@
+// Copyright Epic Games, Inc. All Rights Reserved.
+#pragma once
+
+#include "CoreTypes.h"
+#include "Misc/Guid.h"
+
+// Custom serialization version for changes to DMX Runtime Objects in the //Main Stream
+struct FDMXRuntimeMainStreamObjectVersion
+{
+	enum Type
+	{
+		// Roughly corresponds to 4.26
+		BeforeCustomVersionWasAdded = 0,
+
+		// Update to DMX Library Section using normalized values by default
+		DefaultToNormalizedValuesInDMXLibrarySection,
+
+		// Update to DMX Library Section using normalized values by default
+		ReplaceWeakWithStrongFixturePatchReferncesInLibrarySection,
+
+		// Update DMX Entity Fixture Type to allow the Matrix to be enabled/disabled per Mode
+		DMXFixtureTypeAllowMatrixInEachFixtureMode,
+
+		// DMX Controllers are removed from the pugin
+		DMXControllersAreRemoved,
+		
+		// Update DMX Library to contain an MVR General Scene Description
+		DMXLibraryContainsMVRGeneralSceneDescription,
+
+		// Update Fixture Patch to use an MVR UUID
+		DMXFixturePatchHasMVRUUID,
+
+		// Update DMX Entity Fixture Type to use the correct GDTF type instead of its base class
+		DMXFixtureTypeUsesCorrectGDTFType,
+
+		VersionPlusOne,
+		LatestVersion = VersionPlusOne - 1
+	};
+
+	// The GUID for this custom version number
+	const static FGuid GUID;
+
+private:
+	FDMXRuntimeMainStreamObjectVersion() {}
+};
diff --git a/Engine/Plugins/VirtualProduction/DMX/DMXEngine/Source/DMXRuntime/Private/DMXRuntimeObjectVersion.cpp b/Engine/Plugins/VirtualProduction/DMX/DMXEngine/Source/DMXRuntime/Private/DMXRuntimeObjectVersion.cpp
deleted file mode 100644
diff --git a/Engine/Plugins/VirtualProduction/DMX/DMXEngine/Source/DMXRuntime/Private/DMXRuntimeObjectVersion.h b/Engine/Plugins/VirtualProduction/DMX/DMXEngine/Source/DMXRuntime/Private/DMXRuntimeObjectVersion.h
deleted file mode 100644
diff --git a/Engine/Plugins/VirtualProduction/DMX/DMXEngine/Source/DMXRuntime/Private/Library/DMXEntity.cpp b/Engine/Plugins/VirtualProduction/DMX/DMXEngine/Source/DMXRuntime/Private/Library/DMXEntity.cpp
--- a/Engine/Plugins/VirtualProduction/DMX/DMXEngine/Source/DMXRuntime/Private/Library/DMXEntity.cpp
+++ b/Engine/Plugins/VirtualProduction/DMX/DMXEngine/Source/DMXRuntime/Private/Library/DMXEntity.cpp
@@ -46,7 +46,9 @@ FString UDMXEntity::GetDisplayName() const
 
 void UDMXEntity::SetName(const FString& InNewName)
 {
-	Name = InNewName;
+	const FName UniqueName = MakeUniqueObjectName(GetOuter(), UDMXEntity::StaticClass(), *InNewName);
+	Rename(*UniqueName.ToString(), GetOuter());
+	Name = UniqueName.ToString();
 }
 
 void UDMXEntity::SetParentLibrary(UDMXLibrary* InParent)
diff --git a/Engine/Plugins/VirtualProduction/DMX/DMXEngine/Source/DMXRuntime/Private/Library/DMXEntityFixturePatch.cpp b/Engine/Plugins/VirtualProduction/DMX/DMXEngine/Source/DMXRuntime/Private/Library/DMXEntityFixturePatch.cpp
--- a/Engine/Plugins/VirtualProduction/DMX/DMXEngine/Source/DMXRuntime/Private/Library/DMXEntityFixturePatch.cpp
+++ b/Engine/Plugins/VirtualProduction/DMX/DMXEngine/Source/DMXRuntime/Private/Library/DMXEntityFixturePatch.cpp
@@ -5,6 +5,7 @@
 #include "DMXConversions.h"
 #include "DMXProtocolConstants.h"
 #include "DMXRuntimeLog.h"
+#include "DMXRuntimeMainStreamObjectVersion.h"
 #include "DMXRuntimeUtils.h"
 #include "DMXStats.h"
 #include "DMXTypes.h"
@@ -17,6 +18,8 @@
 #include "Library/DMXLibrary.h"
 #include "Modulators/DMXModulator.h"
 
+#include "UObject/UObjectGlobals.h"
+
 DECLARE_LOG_CATEGORY_CLASS(DMXEntityFixturePatchLog, Log, All);
 
 DECLARE_CYCLE_STAT(TEXT("FixturePatch receive DMX"), STAT_DMXFixturePatchReceiveDMX, STATGROUP_DMX);
@@ -25,32 +28,6 @@ DECLARE_CYCLE_STAT(TEXT("FixturePatch cache values"), STAT_DMXFixturePatchCacheV
 
 #define LOCTEXT_NAMESPACE "DMXEntityFixturePatch"
 
-
-UDMXMVRFixtureInstance* UDMXMVRFixtureInstance::CreateFromMVRFixture(const FDMXMVRFixture& InMVRFixture)
-{
-	UDMXMVRFixtureInstance* NewMVRFixtureInstance = NewObject<UDMXMVRFixtureInstance>();
-	NewMVRFixtureInstance->MVRFixture = InMVRFixture;
-
-	return NewMVRFixtureInstance;
-}
-
-void UDMXMVRFixtureInstance::Serialize(FArchive& Ar)
-{
-	Super::Serialize(Ar);
-
-	Ar << MVRFixture;
-}
-
-void UDMXMVRFixtureInstance::SetMVRFixture(const FDMXMVRFixture& InMVRFixture)
-{
-	if (!ensureAlwaysMsgf(InMVRFixture.UUID == MVRFixture.UUID, TEXT("Cannot change the MVR UUID of an existing MVR Fixture Instance")))
-	{
-		return;
-	}
-
-	MVRFixture = InMVRFixture;
-}
-
 FDMXOnFixturePatchChangedDelegate UDMXEntityFixturePatch::OnFixturePatchChangedDelegate;
 
 UDMXEntityFixturePatch::UDMXEntityFixturePatch()
@@ -74,29 +51,75 @@ UDMXEntityFixturePatch* UDMXEntityFixturePatch::CreateFixturePatchInLibrary(FDMX
 		UDMXLibrary* DMXLibrary = FixtureType->GetParentLibrary();
 		if (ensureMsgf(DMXLibrary, TEXT("Cannot create Fixture Patch when Fixture Type's DMX Library is invalid.")))
 		{
+			const FString EntityName = FDMXRuntimeUtils::FindUniqueEntityName(DMXLibrary, UDMXEntityFixturePatch::StaticClass(), DesiredName);
+			UDMXEntityFixturePatch* NewFixturePatch = NewObject<UDMXEntityFixturePatch>(DMXLibrary, UDMXEntityFixturePatch::StaticClass(), NAME_None, RF_Transactional);
+
 #if WITH_EDITOR
 			if (bMarkDMXLibraryDirty)
 			{
 				DMXLibrary->Modify();
 				DMXLibrary->PreEditChange(UDMXLibrary::StaticClass()->FindPropertyByName(UDMXLibrary::GetEntitiesPropertyName()));
+				NewFixturePatch->PreEditChange(nullptr);
 			}
 #endif
-
-			const FString EntityName = FDMXRuntimeUtils::FindUniqueEntityName(DMXLibrary, UDMXEntityFixturePatch::StaticClass(), DesiredName);
-
-			UDMXEntityFixturePatch* NewFixturePatch = NewObject<UDMXEntityFixturePatch>(DMXLibrary, UDMXEntityFixturePatch::StaticClass(), NAME_None, RF_Transactional);
+			
 			NewFixturePatch->SetName(EntityName);
 			NewFixturePatch->SetFixtureType(FixtureType);
 			NewFixturePatch->SetUniverseID(ConstructionParams.UniverseID);
 			NewFixturePatch->SetStartingChannel(ConstructionParams.StartingAddress);
 			NewFixturePatch->SetActiveModeIndex(ConstructionParams.ActiveMode);
-			
+
+			if (ConstructionParams.MVRFixtureUUID.IsValid())
+			{
+				// Make sure the MVR UUID is truly unique across the DNX Library
+				const TArray<UDMXEntityFixturePatch*> FixturePatches = DMXLibrary->GetEntitiesTypeCast<UDMXEntityFixturePatch>();
+				const UDMXEntityFixturePatch* const* bOtherFixturePatchWithMVRFixtureUUIDPtr = FixturePatches.FindByPredicate([NewFixturePatch, &ConstructionParams](const UDMXEntityFixturePatch* FixturePatch)
+					{
+						return
+							FixturePatch != NewFixturePatch &&
+							FixturePatch->GetMVRFixtureUUID() == ConstructionParams.MVRFixtureUUID;
+					});
+
+				if (bOtherFixturePatchWithMVRFixtureUUIDPtr)
+				{
+					UE_LOG(LogDMXRuntime, Error, TEXT("Trying to create Fixture Patch '%s' with Unique MVR Fixture ID '%s' in DMX Library '%s'."), *DesiredName, *ConstructionParams.MVRFixtureUUID.ToString(), *DMXLibrary->GetName());
+					UE_LOG(LogDMXRuntime, Error, TEXT("However Fixture Patch '%s' in the same Library already uses this Unique MVR Fixture ID."), *(*bOtherFixturePatchWithMVRFixtureUUIDPtr)->Name);
+					UE_LOG(LogDMXRuntime, Error, TEXT("Generating a new Unique MVR Fixture ID for the new Fixture Patch instead"));
+					ensureAlways(0);
+
+					ConstructionParams.MVRFixtureUUID = FGuid::NewGuid();
+				}
+			}
+			else
+			{
+				// No MVR Fixture UUID specified, generate one.
+				ConstructionParams.MVRFixtureUUID = FGuid::NewGuid();
+			}
+			NewFixturePatch->MVRFixtureUUID = ConstructionParams.MVRFixtureUUID;
+
+#if WITH_EDITORONLY_DATA
+			// Make a nice Editor Color
+			const FLinearColor EditorColor = [DMXLibrary, &ConstructionParams]()
+			{
+				const TArray<UDMXEntityFixturePatch*> FixturePatches = DMXLibrary->GetEntitiesTypeCast<UDMXEntityFixturePatch>();
+				const UDMXEntityFixturePatch* const* FixturePatchOfSameTypePtr = FixturePatches.FindByPredicate([&ConstructionParams](const UDMXEntityFixturePatch* FixturePatch)
+					{
+						return FixturePatch->GetFixtureType() == ConstructionParams.FixtureTypeRef.GetFixtureType();
+					});
+				return FixturePatchOfSameTypePtr ? (*FixturePatchOfSameTypePtr)->EditorColor : FLinearColor::MakeRandomColor();
+			}();
+
+			NewFixturePatch->EditorColor = EditorColor;
+#endif
+
 #if WITH_EDITOR
 			if (bMarkDMXLibraryDirty)
 			{
 				DMXLibrary->PostEditChange();
+				NewFixturePatch->PostEditChange();
 			}
 #endif
+
 			OnFixturePatchChangedDelegate.Broadcast(NewFixturePatch);
 			return NewFixturePatch;
 		}
@@ -132,11 +155,25 @@ bool UDMXEntityFixturePatch::Modify(bool bAlwaysMarkDirty)
 }
 #endif // WITH_EDITOR
 
+void UDMXEntityFixturePatch::Serialize(FArchive& Ar)
+{
+	Super::Serialize(Ar);
+
+	Ar.UsingCustomVersion(FDMXRuntimeMainStreamObjectVersion::GUID);
+	if (Ar.IsLoading())
+	{
+		if (Ar.CustomVer(FDMXRuntimeMainStreamObjectVersion::GUID) < FDMXRuntimeMainStreamObjectVersion::DMXFixturePatchHasMVRUUID)
+		{
+			MVRFixtureUUID = FGuid::NewGuid();
+		}
+	}
+}
+
 void UDMXEntityFixturePatch::PostInitProperties()
 {
 	Super::PostInitProperties();
 
-	if (!HasAnyFlags(RF_ClassDefaultObject))
+	if (!HasAnyFlags(RF_ClassDefaultObject | RF_ArchetypeObject))
 	{
 		UDMXEntityFixtureType::GetOnFixtureTypeChanged().AddUObject(this, &UDMXEntityFixturePatch::OnFixtureTypeChanged);
 		RebuildCache();
@@ -147,7 +184,10 @@ void UDMXEntityFixturePatch::PostLoad()
 {
 	Super::PostLoad();
 
-	RebuildCache();
+	if (!HasAnyFlags(RF_ClassDefaultObject | RF_ArchetypeObject))
+	{
+		RebuildCache();
+	}
 }
 
 #if WITH_EDITOR
@@ -526,44 +566,6 @@ bool UDMXEntityFixturePatch::SetActiveModeIndex(int32 NewActiveModeIndex)
 	return false;
 }
 
-void UDMXEntityFixturePatch::AddMVRFixtureInstance(const FDMXMVRFixture& MVRFixtureToAdd)
-{
-	if (!ensureMsgf(ParentFixtureTypeTemplate, TEXT("Trying to create an MVR Fixture Instance in Fixture Patch '%s', but it has not Fixture Type."), *Name))
-	{
-		return;
-	}
-	if (!ensureMsgf(ParentFixtureTypeTemplate->DMXImport, TEXT("Trying to create an MVR Fixture Instance in Fixture Patch '%s', but its Fixture Type '%s' has no GDTF set."), *Name, *ParentFixtureTypeTemplate->Name))
-	{
-		return;
-	}
-	if (!ensureMsgf(MVRFixtureToAdd.UUID.IsValid(), TEXT("Trying to create an MVR Fixture Instance in Fixture Patch '%s', the UUID is not valid."), *Name, *ParentFixtureTypeTemplate->Name))
-	{
-		return;
-	}
-	const bool bMVRUUIDAlreadyContained = MVRFixtureInstances.ContainsByPredicate([MVRFixtureToAdd](const UDMXMVRFixtureInstance* MVRFixtureInstance)
-		{
-			return MVRFixtureInstance->GetMVRFixture().UUID == MVRFixtureToAdd.UUID;
-		});
-	if (!ensureMsgf(!bMVRUUIDAlreadyContained, TEXT("Trying to add an MVR Fixture to a Fixture Patch, but its UUID is already contained.")))
-	{
-		return;
-	}	
-
-	UDMXMVRFixtureInstance* NewMVRFixtureInstance = UDMXMVRFixtureInstance::CreateFromMVRFixture(MVRFixtureToAdd);
-	MVRFixtureInstances.Add(NewMVRFixtureInstance);
-}
-
-TArray<FDMXMVRFixture> UDMXEntityFixturePatch::MakeMVRFixtures() const
-{
-	TArray<FDMXMVRFixture> Result;
-	for (UDMXMVRFixtureInstance* MVRFixtureInstance : MVRFixtureInstances)
-	{
-		Result.Add(MVRFixtureInstance->GetMVRFixture());
-	}
-
-	return Result;
-}
-
 int32 UDMXEntityFixturePatch::GetRemoteUniverse() const
 {	
 	/** DEPRECATED 4.27 */
diff --git a/Engine/Plugins/VirtualProduction/DMX/DMXEngine/Source/DMXRuntime/Private/Library/DMXEntityFixtureType.cpp b/Engine/Plugins/VirtualProduction/DMX/DMXEngine/Source/DMXRuntime/Private/Library/DMXEntityFixtureType.cpp
--- a/Engine/Plugins/VirtualProduction/DMX/DMXEngine/Source/DMXRuntime/Private/Library/DMXEntityFixtureType.cpp
+++ b/Engine/Plugins/VirtualProduction/DMX/DMXEngine/Source/DMXRuntime/Private/Library/DMXEntityFixtureType.cpp
@@ -5,7 +5,7 @@
 #include "DMXConversions.h"
 #include "DMXProtocolSettings.h"
 #include "DMXRuntimeLog.h"
-#include "DMXRuntimeObjectVersion.h"
+#include "DMXRuntimeMainStreamObjectVersion.h"
 #include "DMXRuntimeUtils.h"
 #include "Library/DMXEntityFixturePatch.h"
 #include "Library/DMXImport.h"
@@ -225,10 +225,10 @@ void UDMXEntityFixtureType::Serialize(FArchive& Ar)
 {
 	Super::Serialize(Ar);
 
-	Ar.UsingCustomVersion(FDMXRuntimeObjectVersion::GUID);
+	Ar.UsingCustomVersion(FDMXRuntimeMainStreamObjectVersion::GUID);
 	if (Ar.IsLoading())
 	{
-		if (Ar.CustomVer(FDMXRuntimeObjectVersion::GUID) < FDMXRuntimeObjectVersion::DMXFixtureTypeAllowMatrixInEachFixtureMode)
+		if (Ar.CustomVer(FDMXRuntimeMainStreamObjectVersion::GUID) < FDMXRuntimeMainStreamObjectVersion::DMXFixtureTypeAllowMatrixInEachFixtureMode)
 		{
 			// For assets that were created before each mode could enable or disable the matrix, copy the deprecated bFixtureMatrixEnabled property to each mode
 			for (FDMXFixtureMode& Mode : Modes)
@@ -238,6 +238,18 @@ void UDMXEntityFixtureType::Serialize(FArchive& Ar)
 				PRAGMA_ENABLE_DEPRECATION_WARNINGS
 			}
 		}
+
+		if (Ar.CustomVer(FDMXRuntimeMainStreamObjectVersion::GUID) < FDMXRuntimeMainStreamObjectVersion::DMXFixtureTypeUsesCorrectGDTFType)
+		{
+			// Upgrade the GDTF type from the base DMXImport to the actual DMXImportGDTF type which was always used.
+			// Name the new variable 'GDTF' to make the purpose of the variable more clear, namely holding a GDTF asset.
+			PRAGMA_DISABLE_DEPRECATION_WARNINGS
+			if (UDMXImportGDTF* OldGDTF = Cast<UDMXImportGDTF>(DMXImport))
+			{
+				GDTF = OldGDTF;
+			}
+			PRAGMA_ENABLE_DEPRECATION_WARNINGS
+		}
 	}
 }
 
@@ -307,16 +319,11 @@ void UDMXEntityFixtureType::SetGDTF(UDMXImportGDTF* GDTFAsset)
 		return;
 	}
 
-	DMXImport = GDTFAsset;
-	if (UDMXImportGDTFDMXModes* GDTFDMXModes = Cast<UDMXImportGDTFDMXModes>(DMXImport->DMXModes))
+	GDTF = GDTFAsset;
+	if (UDMXImportGDTFDMXModes* GDTFDMXModes = Cast<UDMXImportGDTFDMXModes>(GDTF->DMXModes))
 	{
 		// Clear existing modes
-		Modes.Empty(DMXImport->DMXModes != nullptr ? GDTFDMXModes->DMXModes.Num() : 0);
-
-		if (DMXImport->DMXModes == nullptr)
-		{
-			return;
-		}
+		Modes.Empty(GDTF->DMXModes != nullptr ? GDTFDMXModes->DMXModes.Num() : 0);
 
 		// Used to map Functions to Attributes
 		const UDMXProtocolSettings* ProtocolSettings = GetDefault<UDMXProtocolSettings>();
@@ -452,7 +459,7 @@ FDMXOnFixtureTypeChangedDelegate& UDMXEntityFixtureType::GetOnFixtureTypeChanged
 	return OnFixtureTypeChangedDelegate;
 }
 
-int32 UDMXEntityFixtureType::AddMode()
+int32 UDMXEntityFixtureType::AddMode(FString BaseModeName)
 {
 	FDMXFixtureMode NewMode;
 
@@ -462,7 +469,7 @@ int32 UDMXEntityFixtureType::AddMode()
 	{
 		ModeNames.Add(Mode.ModeName);
 	}
-	NewMode.ModeName = FDMXRuntimeUtils::GenerateUniqueNameFromExisting(ModeNames, LOCTEXT("DefaultModeName", "Mode").ToString());
+	NewMode.ModeName = FDMXRuntimeUtils::GenerateUniqueNameFromExisting(ModeNames, BaseModeName);
 
 	const int32 NewModeIndex = Modes.Add(NewMode);
 
diff --git a/Engine/Plugins/VirtualProduction/DMX/DMXEngine/Source/DMXRuntime/Private/Library/DMXLibrary.cpp b/Engine/Plugins/VirtualProduction/DMX/DMXEngine/Source/DMXRuntime/Private/Library/DMXLibrary.cpp
--- a/Engine/Plugins/VirtualProduction/DMX/DMXEngine/Source/DMXRuntime/Private/Library/DMXLibrary.cpp
+++ b/Engine/Plugins/VirtualProduction/DMX/DMXEngine/Source/DMXRuntime/Private/Library/DMXLibrary.cpp
@@ -2,6 +2,7 @@
 #include "Library/DMXLibrary.h"
 
 #include "DMXRuntimeLog.h"
+#include "DMXRuntimeMainStreamObjectVersion.h"
 #include "DMXProtocolSettings.h"
 #include "DMXProtocolUtils.h"
 #include "Interfaces/IDMXProtocol.h"
@@ -11,14 +12,50 @@
 #include "Library/DMXEntity.h"
 #include "Library/DMXEntityController.h"
 #include "Library/DMXEntityFixturePatch.h"
+#include "Library/DMXImportGDTF.h"
+#include "MVR/DMXMVRGeneralSceneDescription.h"
 
+#include "Misc/Paths.h"
 
 #define LOCTEXT_NAMESPACE "DMXLibrary"
 
+
 FDMXOnEntityArrayChangedDelegate UDMXLibrary::OnEntitiesAddedDelegate;
 
 FDMXOnEntityArrayChangedDelegate UDMXLibrary::OnEntitiesRemovedDelegate;
 
+UDMXLibrary::UDMXLibrary()
+{
+	const FName GeneralSceneDescriptionName = FName(GetName() + TEXT("_MVRGeneralSceneDescription"));
+	GeneralSceneDescription = NewObject<UDMXMVRGeneralSceneDescription>(this, GeneralSceneDescriptionName);
+}
+
+void UDMXLibrary::Serialize(FArchive& Ar)
+{
+	if (Ar.IsSaving())
+	{
+		// Update the General Scene Description before saving it
+		UpdateGeneralSceneDescription();
+	}
+
+	Super::Serialize(Ar);
+
+	Ar.UsingCustomVersion(FDMXRuntimeMainStreamObjectVersion::GUID);
+	if (Ar.IsLoading())
+	{
+		if (Ar.CustomVer(FDMXRuntimeMainStreamObjectVersion::GUID) < FDMXRuntimeMainStreamObjectVersion::DMXLibraryContainsMVRGeneralSceneDescription)
+		{
+			if (!GeneralSceneDescription)
+			{
+				const FName GeneralSceneDescriptionName = FName(GetName() + TEXT("_MVRGeneralSceneDescription"));
+				GeneralSceneDescription = UDMXMVRGeneralSceneDescription::CreateFromDMXLibrary(*this, this, GeneralSceneDescriptionName);
+			}
+
+			GeneralSceneDescription->WriteDMXLibraryToGeneralSceneDescription(*this);
+		}
+	}
+}
+
 void UDMXLibrary::PostInitProperties()
 {
 	Super::PostInitProperties();
@@ -119,15 +156,20 @@ void UDMXLibrary::PostEditChangeProperty(FPropertyChangedEvent& PropertyChangedE
 
 void UDMXLibrary::RegisterEntity(UDMXEntity* Entity)
 {
-	if (ensureMsgf(IsValid(Entity), TEXT("Trying to register Entity with DMX Library, but DMX Entity is not valid.")))
+	if (!ensureAlwaysMsgf(IsValid(Entity), TEXT("Trying to register Entity with DMX Library, but DMX Entity is not valid.")))
 	{
-		if (!Entities.Contains(Entity))
-		{
-			Entities.Add(Entity);
-			LastAddedEntity = Entity;
-			OnEntitiesAddedDelegate.Broadcast(this, TArray<UDMXEntity*>({ Entity }));
-		}	
+		return;
 	}
+
+	if (!ensureAlwaysMsgf(!Entities.Contains(Entity), TEXT("Trying to register Entity %s with DMX Library, but Entity was already added."), *Entity->Name))
+	{
+		return;
+	}
+
+	Entities.Add(Entity);
+	LastAddedEntity = Entity;
+
+	OnEntitiesAddedDelegate.Broadcast(this, TArray<UDMXEntity*>({ Entity }));
 }
 
 void UDMXLibrary::UnregisterEntity(UDMXEntity* Entity)
@@ -530,17 +572,30 @@ void UDMXLibrary::UpdatePorts()
 	}
 }
 
-FDMXMVRGeneralSceneDescription UDMXLibrary::MakeGeneralSceneDescription() const
+void UDMXLibrary::SetMVRGeneralSceneDescription(UDMXMVRGeneralSceneDescription* NewGeneralSceneDescription)
 {
-	FDMXMVRGeneralSceneDescription GeneralSceneDescription;
-	const TArray<UDMXEntityFixturePatch*> FixturePatches = GetEntitiesTypeCast<UDMXEntityFixturePatch>();
-	for (UDMXEntityFixturePatch* FixturePatch : FixturePatches)
+	// Remove all Fixture Patches that don't exist in the General Scene Description
+	Entities.RemoveAll([NewGeneralSceneDescription](UDMXEntity* Entity)
+		{
+			if (UDMXEntityFixturePatch* FixturePatch = Cast<UDMXEntityFixturePatch>(Entity))
+			{
+				return NewGeneralSceneDescription->FindMVRFixture(FixturePatch->GetMVRFixtureUUID()) == nullptr;
+			}
+			return false;
+		});
+
+	GeneralSceneDescription = NewGeneralSceneDescription;
+}
+
+UDMXMVRGeneralSceneDescription* UDMXLibrary::UpdateGeneralSceneDescription()
+{
+	if (ensureAlwaysMsgf(GeneralSceneDescription, TEXT("Trying to update General Scene Description of %s, but the General Scene Description is not valid."), *GetName()))
 	{
-		const TArray<FDMXMVRFixture> MVRFixtures = FixturePatch->MakeMVRFixtures();
-		GeneralSceneDescription.MVRFixtures.Append(MVRFixtures);
+		GeneralSceneDescription->WriteDMXLibraryToGeneralSceneDescription(*this);
+		return GeneralSceneDescription;
 	}
 
-	return GeneralSceneDescription;
+	return nullptr;
 }
 
 #if WITH_EDITOR
diff --git a/Engine/Plugins/VirtualProduction/DMX/DMXEngine/Source/DMXRuntime/Private/MVR/DMXMVRAssetUserData.cpp b/Engine/Plugins/VirtualProduction/DMX/DMXEngine/Source/DMXRuntime/Private/MVR/DMXMVRAssetUserData.cpp
--- a/Engine/Plugins/VirtualProduction/DMX/DMXEngine/Source/DMXRuntime/Private/MVR/DMXMVRAssetUserData.cpp
+++ b/Engine/Plugins/VirtualProduction/DMX/DMXEngine/Source/DMXRuntime/Private/MVR/DMXMVRAssetUserData.cpp
@@ -3,14 +3,15 @@
 #include "MVR/DMXMVRAssetUserData.h"
 
 #include "GameFramework/Actor.h"
+#include "Interfaces/Interface_AssetUserData.h"
 
+const FName UDMXMVRAssetUserData::MVRFixtureUUIDMetaDataKey = TEXT("MVRFixtureUUID");
+const FName UDMXMVRAssetUserData::FixturePatchMetaDataKey = TEXT("FixturePatch");
 
-const FName UDMXMVRAssetUserData::MVRUUIDMetaDataKey = TEXT("MVRUUID");
 
-
-UDMXMVRAssetUserData* UDMXMVRAssetUserData::GetMVRAssetUserData(AActor* Actor)
+UDMXMVRAssetUserData* UDMXMVRAssetUserData::GetMVRAssetUserData(const AActor& Actor)
 {
-	if (USceneComponent* RootComponent = Actor->GetRootComponent())
+	if (UActorComponent* RootComponent = Actor.GetRootComponent())
 	{
 		return RootComponent->GetAssetUserData<UDMXMVRAssetUserData>();
 	}
@@ -18,7 +19,7 @@ UDMXMVRAssetUserData* UDMXMVRAssetUserData::GetMVRAssetUserData(AActor* Actor)
 	return nullptr;
 }
 
-FString UDMXMVRAssetUserData::GetMVRAssetUserDataValueForkey(AActor* Actor, const FName Key)
+FString UDMXMVRAssetUserData::GetMVRAssetUserDataValueForKey(const AActor& Actor, const FName Key)
 {
 	if (UDMXMVRAssetUserData* MVRAssetUserData = GetMVRAssetUserData(Actor))
 	{
@@ -29,27 +30,23 @@ FString UDMXMVRAssetUserData::GetMVRAssetUserDataValueForkey(AActor* Actor, cons
 	return FString();
 }
 
-bool UDMXMVRAssetUserData::SetMVRAssetUserDataValueForKey(AActor* Actor, FName Key, const FString& Value)
+bool UDMXMVRAssetUserData::SetMVRAssetUserDataValueForKey(AActor& InOutActor, FName InKey, const FString& InValue)
 {
 	// For AActor, the interface is actually implemented by the ActorComponent
-	if (USceneComponent* RootComponent = Actor->GetRootComponent())
+	if (UActorComponent* RootComponent = InOutActor.GetRootComponent())
 	{
-		if (IInterface_AssetUserData* AssetUserData = Cast<IInterface_AssetUserData>(RootComponent))
+		UDMXMVRAssetUserData* MVRUserData = RootComponent->GetAssetUserData<UDMXMVRAssetUserData>();
+		if (!MVRUserData)
 		{
-			UDMXMVRAssetUserData* MVRUserData = AssetUserData->GetAssetUserData<UDMXMVRAssetUserData>();
-
-			if (!MVRUserData)
-			{
-				MVRUserData = NewObject<UDMXMVRAssetUserData>(RootComponent, NAME_None, RF_Public | RF_Transactional);
-				AssetUserData->AddAssetUserData(MVRUserData);
-			}
+			MVRUserData = NewObject<UDMXMVRAssetUserData>(&InOutActor, NAME_None, RF_Public | RF_Transactional);
+			RootComponent->AddAssetUserData(MVRUserData);
+		}
 
-			// Add Datasmith meta data
-			MVRUserData->MetaData.Add(Key, Value);
-			MVRUserData->MetaData.KeySort(FNameLexicalLess());
+		// Add Datasmith meta data
+		MVRUserData->MetaData.Add(InKey, InValue);
+		MVRUserData->MetaData.KeySort(FNameLexicalLess());
 
-			return true;
-		}
+		return true;
 	}
 
 	return false;
diff --git a/Engine/Plugins/VirtualProduction/DMX/DMXEngine/Source/DMXRuntime/Private/MVR/DMXMVRFixture.cpp b/Engine/Plugins/VirtualProduction/DMX/DMXEngine/Source/DMXRuntime/Private/MVR/DMXMVRFixture.cpp
new file mode 100644
--- /dev/null
+++ b/Engine/Plugins/VirtualProduction/DMX/DMXEngine/Source/DMXRuntime/Private/MVR/DMXMVRFixture.cpp
@@ -0,0 +1,599 @@
+// Copyright Epic Games, Inc. All Rights Reserved.
+
+#include "MVR/DMXMVRFixture.h"
+
+#include "DMXProtocolCommon.h"
+#include "DMXRuntimeLog.h"
+
+#include "XmlFile.h"
+#include "XmlNode.h"
+
+
+namespace UE::DMXRuntime::DMXMVRFixture::Private
+{
+	/** Parses the Name from the MVR Fixture Node, returns true on success. */
+	bool ParseName(const FXmlNode* const InFixtureNode, FString& OutName)
+	{
+		check(InFixtureNode);
+
+		static const FString InFixtureNodeAttribute_Name = TEXT("Name");
+		OutName = InFixtureNode->GetAttribute(InFixtureNodeAttribute_Name);
+		if (!OutName.IsEmpty())
+		{
+			return true;
+		}
+
+		return false;
+	}
+
+	/** Parses the UUID from the MVR Fixture Node, returns true on success. */
+	bool ParseUUID(const FXmlNode* const InFixtureNode, FGuid& OutUUID)
+	{
+		check(InFixtureNode);
+
+		static const FString InFixtureNodeAttribute_UUID = TEXT("UUID");
+		const FString UUIDString = InFixtureNode->GetAttribute(InFixtureNodeAttribute_UUID);
+		if (FGuid::Parse(UUIDString, OutUUID))
+		{
+			return true;
+		}
+		return false;
+	}
+
+	/** Parses the Matrix from the MVR Fixture Node, optional, always succeeds. */
+	void ParseMatrix(const FXmlNode* const InFixtureNode, TOptional<FTransform>& OutTransform)
+	{
+		check(InFixtureNode);
+		OutTransform.Reset();
+
+		static const FString ChildNodeName_Matrix = TEXT("Matrix");
+		const FXmlNode* const MatrixNode = InFixtureNode->FindChildNode(ChildNodeName_Matrix);
+
+		if (!MatrixNode)
+		{
+			return;
+		}
+
+		auto GetVectorStringsFromMatrixStringLambda([](const FString& MatrixString) -> TArray<FString>
+			{
+				TArray<FString> Result;
+				int32 SubstringIndex = 0;
+				for (int32 SubStringNumber = 0; SubStringNumber < 4; SubStringNumber++)
+				{
+					SubstringIndex = MatrixString.Find(TEXT("{"), ESearchCase::IgnoreCase, ESearchDir::FromStart, SubstringIndex) + 1;
+					const int32 EndIndex = MatrixString.Find(TEXT("}"), ESearchCase::IgnoreCase, ESearchDir::FromStart, SubstringIndex);
+
+					const FString SubString = MatrixString.Mid(SubstringIndex, EndIndex - SubstringIndex);
+					Result.Add(SubString);
+				}
+
+				return Result;
+			});
+
+		auto GetVectorFromVectorStringLambda([](const FString& VectorString) -> FVector
+			{
+				const int32 XIndex = 0;
+				const int32 YIndex = VectorString.Find(TEXT(","), ESearchCase::IgnoreCase, ESearchDir::FromStart, 0) + 1;
+				const int32 ZIndex = VectorString.Find(TEXT(","), ESearchCase::IgnoreCase, ESearchDir::FromStart, YIndex) + 1;
+
+				const FString XString = "X=" + VectorString.Mid(XIndex, YIndex - XIndex);
+				const FString YString = "Y=" + VectorString.Mid(YIndex, ZIndex - YIndex);
+				const FString ZString = "Z=" + VectorString.Mid(ZIndex, VectorString.Len() - ZIndex);
+
+				FVector Vector;
+				if (Vector.InitFromString(XString + YString + ZString))
+				{
+					return Vector;
+				}
+
+				return FVector::ZeroVector;
+			});
+
+		const TArray<FString> SubStrings = GetVectorStringsFromMatrixStringLambda(MatrixNode->GetContent());
+		if (SubStrings.Num() < 4)
+		{
+			return;
+		}
+
+		const FVector U = GetVectorFromVectorStringLambda(SubStrings[0]) / 10.f;
+		const FVector V = GetVectorFromVectorStringLambda(SubStrings[1]) / 10.f;
+		const FVector W = GetVectorFromVectorStringLambda(SubStrings[2]) / 10.f;
+		const FVector O = GetVectorFromVectorStringLambda(SubStrings[3]) / 10.f;
+
+		FMatrix Matrix;
+		Matrix.M[0][0] = U.X;
+		Matrix.M[0][1] = V.X;
+		Matrix.M[0][2] = W.X;
+		Matrix.M[0][3] = O.X;
+
+		Matrix.M[1][0] = U.Y;
+		Matrix.M[1][1] = V.Y;
+		Matrix.M[1][2] = W.Y;
+		Matrix.M[1][3] = -O.Y; // From MVR's right hand to UE's left hand coordinate system
+
+		Matrix.M[2][0] = U.Z;
+		Matrix.M[2][1] = V.Z;
+		Matrix.M[2][2] = W.Z;
+		Matrix.M[2][3] = O.Z;
+
+		Matrix.M[3][0] = 0.0;
+		Matrix.M[3][1] = 0.0;
+		Matrix.M[3][2] = 0.0;
+		Matrix.M[3][3] = 1.0;
+
+		OutTransform = FTransform(Matrix.GetTransposed());
+
+		// Scale from milimeters to centimeters
+		OutTransform->SetScale3D(OutTransform->GetScale3D() * 10.f);
+
+		// GDTFs are facing down on the Z-Axis, but UE Actors are facing up 
+		const FQuat InvertUpRotationQuaternion = FQuat(OutTransform->GetRotation().GetAxisY(), PI);
+		OutTransform->SetRotation(OutTransform->GetRotation() * InvertUpRotationQuaternion);
+	}
+	
+	/** Parses the GDTFSpec from the MVR Fixture Node, returns true on success. */
+	bool ParseGDTFSpec(const FXmlNode* const InFixtureNode, FString& OutGDTFSpec)
+	{
+		check(InFixtureNode);
+
+		static const FString ChildNodeName_GDTFSpec = TEXT("GDTFSpec");
+		const FXmlNode* const GDTFSpecNode = InFixtureNode->FindChildNode(ChildNodeName_GDTFSpec);
+		if (GDTFSpecNode)
+		{
+			OutGDTFSpec = GDTFSpecNode->GetContent();
+			return true;
+		}
+
+		return false;
+	}
+
+	/** Parses the GDTFMode from the MVR Fixture Node, returns true on success. */
+	bool ParseGDTFMode(const FXmlNode* const InFixtureNode, FString& OutGDTFMode)
+	{
+		check(InFixtureNode);
+
+		static const FString ChildNodeName_GDTFMode = TEXT("GDTFMode");
+		const FXmlNode* const GDTFModeNode = InFixtureNode->FindChildNode(ChildNodeName_GDTFMode);
+		if (GDTFModeNode)
+		{
+			OutGDTFMode = GDTFModeNode->GetContent();
+			return true;
+		}
+		return true;
+	}
+
+	/** Parses Focus from the MVR Fixture Node, optional, always succeeds. */
+	void ParseFocus(const FXmlNode* const InFixtureNode, TOptional<FGuid>& OutFocus)
+	{
+		check(InFixtureNode);
+		OutFocus.Reset();
+
+		static const FString ChildNodeName_Focus = TEXT("Focus");
+		const FXmlNode* const FocusNode = InFixtureNode->FindChildNode(ChildNodeName_Focus);
+		if (FocusNode)
+		{
+			const FString FocusString = InFixtureNode->GetContent();
+			FGuid FocusGuid;
+			if (FGuid::Parse(FocusString, FocusGuid))
+			{
+				OutFocus = FocusGuid;
+			}
+		};
+	}
+
+	/** Parses Cast Shadows from the MVR Fixture Node, optional, always succeeds. */
+	void ParseCastShadow(const FXmlNode* const InFixtureNode, TOptional<bool>& OutCastShadow)
+	{
+		check(InFixtureNode);
+		OutCastShadow.Reset();
+
+		static const FString ChidNodeName_CastShadow = TEXT("CastShadow");
+		const FXmlNode* const CastShadowNode = InFixtureNode->FindChildNode(ChidNodeName_CastShadow);
+		if (CastShadowNode)
+		{
+			OutCastShadow = !(CastShadowNode->GetContent() == TEXT("0")) || !CastShadowNode->GetContent().Equals(TEXT("true"), ESearchCase::IgnoreCase);
+		}
+	}
+
+	/** Parses Position of from MVR Fixture Node, optional, always succeeds. */
+	void ParsePosition(const FXmlNode* const InFixtureNode, TOptional<FGuid>& OutPosition)
+	{
+		check(InFixtureNode);
+		OutPosition.Reset();
+
+		static const FString ChildNodenName_Position = TEXT("Position");
+		const FXmlNode* const PositionNode = InFixtureNode->FindChildNode(ChildNodenName_Position);
+		if (PositionNode)
+		{
+			FGuid PositionGuid;
+			if (FGuid::Parse(PositionNode->GetContent(), PositionGuid))
+			{
+				OutPosition = PositionGuid;
+			}
+		}
+	}
+
+	/** Parses the Fixture Id from the MVR Fixture Node, returns true on success. */
+	bool ParseFixtureId(const FXmlNode* const InFixtureNode, FString& OutFixtureId)
+	{
+		check(InFixtureNode);
+		OutFixtureId.Reset();
+
+		static const FString ChildNodenName_FixtureId = TEXT("FixtureId");
+		const FXmlNode* const FixtureIdNode = InFixtureNode->FindChildNode(ChildNodenName_FixtureId);
+		if (FixtureIdNode)
+		{
+			OutFixtureId = FixtureIdNode->GetContent();
+			return true;
+		}
+
+		return false;
+	}
+
+	/** Parses the Unit Number from the MVR Fixture Node, returns true on success. */
+	bool ParseUnitNumber(const FXmlNode* const InFixtureNode, int32& OutUnitNumber)
+	{
+		check(InFixtureNode);
+
+		static const FString ChildNodenName_UnitNumber = TEXT("FixtureId");
+		const FXmlNode* const UnitNumberNode = InFixtureNode->FindChildNode(ChildNodenName_UnitNumber);
+		if (UnitNumberNode)
+		{
+			int32 UnitNumber;
+			if (LexTryParseString(UnitNumber, *UnitNumberNode->GetContent()))
+			{
+				OutUnitNumber = UnitNumber;
+				return true;
+			}
+		}
+		return false;
+	}
+
+	/** Parses the Addresses from the MVR Fixture Node, defaulted, always succeeds */
+	bool ParseAddresses(const FXmlNode* const InFixtureNode, FDMXMVRFixtureAddresses& OutAddresses)
+	{
+		check(InFixtureNode);
+
+		static const FString ChildNodenName_Addresses = TEXT("Addresses");
+		const FXmlNode* const AddressesNode = InFixtureNode->FindChildNode(ChildNodenName_Addresses);
+
+		static const FString ChildNodenName_Address = TEXT("Address");
+		const FXmlNode* const AddressNode = AddressesNode ? AddressesNode->FindChildNode(ChildNodenName_Address) : nullptr;
+
+		if (AddressNode)
+		{
+			int32 Universe = -1;
+			int32 Address = -1;
+
+			auto ParseAddressStringLambda = [](const FXmlNode* const InAddressNode, int32& OutUniverse, int32& OutAddress) -> const bool
+			{
+				const FString InAddressString = InAddressNode->GetContent();
+
+				int32 SeparatorIndex = INDEX_NONE;
+				if (InAddressString.FindChar('.', SeparatorIndex))
+				{
+					if (SeparatorIndex > 0 && SeparatorIndex < InAddressString.Len() - 2)
+					{
+						const FString UniverseString = InAddressString.Left(SeparatorIndex);
+						const FString AddressString = InAddressString.RightChop(SeparatorIndex + 1);
+
+						const bool bValidUniverse = LexTryParseString(OutUniverse, *UniverseString) && OutUniverse >= 0 && OutUniverse <= DMX_MAX_UNIVERSE;
+						const bool bValidAddress = LexTryParseString(OutAddress, *UniverseString) && OutAddress >= 0 && OutAddress <= DMX_MAX_ADDRESS;
+
+						return bValidUniverse && bValidAddress;
+					}
+				}
+				else if (LexTryParseString(OutAddress, *InAddressString))
+				{
+					static const int32 UniverseSize = 512;
+
+					OutUniverse = OutAddress / UniverseSize;
+					OutAddress = OutAddress % UniverseSize;
+
+					return 
+						OutUniverse >= 0 && 
+						OutUniverse <= DMX_MAX_UNIVERSE &&
+						OutAddress >= 0 && 
+						OutAddress <= DMX_MAX_ADDRESS;
+				}
+
+				return false;
+			}(AddressNode, Universe, Address);
+			if (ParseAddressStringLambda)
+			{
+				OutAddresses.Universe = Universe;
+				OutAddresses.Address = Address;
+
+				return true;
+			}
+		}
+		return false;
+	}
+
+	/** Parses the CID Color from the MVR Fixture Node, optional, always succeeds. */
+	void ParseCIEColor(const FXmlNode* const InFixtureNode, TOptional<FDMXMVRColorCIE>& OutCIEColor)
+	{
+		check(InFixtureNode);
+
+		static const FString ChildNodenName_CIEColor = TEXT("CIEColor");
+		const FXmlNode* const CIEColorNode = InFixtureNode->FindChildNode(ChildNodenName_CIEColor);
+		if (CIEColorNode)
+		{
+			const FString CIEColorString = CIEColorNode->GetContent();
+			if (CIEColorString.IsEmpty())
+			{
+				return;
+			}
+
+			TArray<FString> CIEColorArray;
+			CIEColorString.ParseIntoArray(CIEColorArray, TEXT(","));
+
+			if (CIEColorArray.Num() != 3	)
+			{
+				return;
+			}
+
+			FDMXMVRColorCIE Result;
+			bool bSuccess = LexTryParseString<float>(Result.X, *CIEColorArray[0]);
+			bSuccess = bSuccess && LexTryParseString<float>(Result.Y, *CIEColorArray[1]);
+			bSuccess = bSuccess && LexTryParseString<uint8>(Result.YY, *CIEColorArray[2]);
+
+			if (bSuccess)
+			{
+				OutCIEColor = Result;
+			}
+		}
+	}
+
+	/** Parses the Fixture Type Id from MVR Fixture Node, optional, always succeeds. */
+	void ParseFixtureTypeId(const FXmlNode* const InFixtureNode, TOptional<int32>& OutFixtureTypeId)
+	{
+		check(InFixtureNode);
+		OutFixtureTypeId.Reset();
+
+		static const FString ChildNodenName_FixtureTypeId = TEXT("FixtureTypeId");
+		const FXmlNode* const FixtureTypeIdNode = InFixtureNode->FindChildNode(ChildNodenName_FixtureTypeId);
+		if (FixtureTypeIdNode)
+		{
+			int32 FixtureTypeId;
+			if (LexTryParseString(FixtureTypeId, *FixtureTypeIdNode->GetContent()))
+			{
+				OutFixtureTypeId = FixtureTypeId;
+			}
+		}
+	}
+
+	/** Parses the Custom Id from MVR Fixture Node, optional, always succeeds. */
+	void ParseCustomId(const FXmlNode* const InFixtureNode, TOptional<int32>& OutCustomId)
+	{
+		check(InFixtureNode);
+		OutCustomId.Reset();
+
+		static const FString ChildNodenName_Position = TEXT("Position");
+		const FXmlNode* const PositionNode = InFixtureNode->FindChildNode(ChildNodenName_Position);
+		if (PositionNode)
+		{
+			int32 FixtureTypeId;
+			if (LexTryParseString(FixtureTypeId, *PositionNode->GetContent()))
+			{
+				OutCustomId = FixtureTypeId;
+			}
+		}
+	}
+
+	/** Parses the Mapping from MVR Fixture Node, optional, always succeeds. */
+	void ParseMapping(const FXmlNode* const InFixtureNode, TOptional<FDMXMVRFixtureMapping>& OutMapping)
+	{
+		check(InFixtureNode);
+		OutMapping.Reset();
+
+		static const FString ChildNodenName_Mapping = TEXT("Position");
+		const FXmlNode* const MappingNode = InFixtureNode->FindChildNode(ChildNodenName_Mapping);
+		if (MappingNode)
+		{
+			int32 FixtureTypeId;
+			if (LexTryParseString(FixtureTypeId, *MappingNode->GetContent()))
+			{
+				static const FString ChildNodenName_LinkDef = TEXT("LinkDef");
+				const FXmlNode* const LinkDefNOde = MappingNode->FindChildNode(ChildNodenName_LinkDef);
+				FDMXMVRFixtureMapping Mapping;
+				if (LinkDefNOde && !FGuid::Parse(LinkDefNOde->GetContent(), Mapping.LinkDef))
+				{
+					return;
+				}
+
+				static const FString ChildNodenName_UX = TEXT("UX");
+				const FXmlNode* const UXNode = InFixtureNode->FindChildNode(ChildNodenName_UX);
+				int32 UX;
+				if (UXNode && !LexTryParseString(UX, *UXNode->GetContent()))
+				{
+					Mapping.UX = UX;
+					return;
+				}
+
+				static const FString ChildNodenName_UY = TEXT("UY");
+				const FXmlNode* const UYNode = InFixtureNode->FindChildNode(ChildNodenName_UY);
+				int32 UY;
+				if (UYNode && !LexTryParseString(UY, *UYNode->GetContent()))
+				{
+					return;
+				}
+				Mapping.UY = UY;
+
+				static const FString ChildNodenName_OX = TEXT("OX");
+				const FXmlNode* const OXNode = InFixtureNode->FindChildNode(ChildNodenName_OX);
+				int32 OX;
+				if (OXNode && !LexTryParseString(OX, *OXNode->GetContent()))
+				{
+					return;
+				}
+				Mapping.OX = OX;
+
+				static const FString ChildNodenName_OY = TEXT("OY");
+				const FXmlNode* const OYNode = InFixtureNode->FindChildNode(ChildNodenName_OY);
+				int32 OY;
+				if (OYNode && !LexTryParseString(OY, *OYNode->GetContent()))
+				{
+					return;
+				}
+				Mapping.OY = OY;
+
+				static const FString ChildNodenName_RZ = TEXT("RZ");
+				const FXmlNode* const RZNode = InFixtureNode->FindChildNode(ChildNodenName_RZ);
+				int32 RZ;
+				if (RZNode && !LexTryParseString(RZ, *RZNode->GetContent()))
+				{
+					return;
+				}
+				Mapping.RZ = RZ;
+
+				OutMapping = Mapping;
+			}
+		}
+	}
+
+	/** Parses the Gobo from MVR Fixture Node, optional, always succeeds. */
+	void ParseGobo(const FXmlNode* const InFixtureNode, TOptional<FDMXMVRFixtureGobo>& OutGobo)
+	{
+		check(InFixtureNode);
+		OutGobo.Reset();
+
+		static const FString ChildNodenName_Gobo = TEXT("Position");
+		const FXmlNode* const GoboNode = InFixtureNode->FindChildNode(ChildNodenName_Gobo);
+		if (GoboNode)
+		{
+			const FString& GoboResource = GoboNode->GetContent();
+			OutGobo->Value = GoboResource;
+
+			static const FString ChildNodeName_Gobo = TEXT("Rotation");
+			const FXmlNode* const GoboRotationNode = GoboNode->FindChildNode(ChildNodeName_Gobo);
+			float GoboRotation;
+			if (GoboRotationNode && LexTryParseString(GoboRotation, *GoboRotationNode->GetContent()))
+			{
+				OutGobo->Rotation = GoboRotation;
+			}
+		}
+	}
+}
+
+FDMXMVRFixture::FDMXMVRFixture()
+{}
+
+FDMXMVRFixture::FDMXMVRFixture(const FXmlNode* const FixtureNode)
+{
+	// Ensure a valid node
+	if (!ensureAlwaysMsgf(FixtureNode, TEXT("Trying to parse MVR Fixture Node, but node is null.")))
+	{
+		return;
+	}
+
+	static const FString NodeName_Fixture = TEXT("Fixture");
+	if (!ensureAlwaysMsgf(FixtureNode->GetTag() == NodeName_Fixture, TEXT("Trying to parse MVR Fixture Node, but node name doesn't match.")))
+	{
+		return;
+	}
+
+	// Parse nodes. Those that don't fail are either defaulted or optional.
+	using namespace UE::DMXRuntime::DMXMVRFixture::Private;
+	if (!ParseName(FixtureNode, Name))
+	{
+		// Non-optional
+		UE_LOG(LogDMXRuntime, Warning, TEXT("Cannot parse Fixture Node as MVR Fixture Name. Failed to parse MVR Node %s."), *FixtureNode->GetTag());
+		return;
+	}
+
+	if(!ParseUUID(FixtureNode, UUID))
+	{
+		// Non-optional
+		UE_LOG(LogDMXRuntime, Warning, TEXT("Cannot parse Fixture '%s' from MVR. Failed to parse MVR UUID."), *Name);
+		return;
+	}
+
+	ParseMatrix(FixtureNode, Transform);
+	if (Transform.IsSet() && !Transform.GetValue().IsValid())
+	{
+		// Mend invalid transforms
+		UE_LOG(LogDMXRuntime, Warning, TEXT("Imported Fixture '%s' from MVR but its Transform is not valid, using Identity instead."), *Name);
+		Transform = FTransform::Identity;
+	}
+
+	if (!ParseGDTFSpec(FixtureNode, GDTFSpec))
+	{
+		// Non-optional
+		UE_LOG(LogDMXRuntime, Warning, TEXT("Cannot parse Fixture '%s' from MVR. No valid GDTF Spec specified."), *Name);
+		return;
+	}
+
+	if (!ParseGDTFMode(FixtureNode, GDTFMode))
+	{
+		// Non-optional
+		UE_LOG(LogDMXRuntime, Warning, TEXT("Cannot parse Fixture '%s' from MVR. No valid GDTF Mode specified."), *Name);
+		return;
+	}
+
+	ParseFocus(FixtureNode, Focus);
+	ParseCastShadow(FixtureNode, bCastShadows);
+	ParsePosition(FixtureNode, Position);
+	
+	if (!ParseFixtureId(FixtureNode, FixtureId))
+	{
+		// Non-optional, mend with empty string
+		UE_LOG(LogDMXRuntime, Warning, TEXT("Imported Fixture '%s' from MVR. No valid Fixture Id specified. Fixture Id set to Fixture Name."), *Name);
+		FixtureId = TEXT("");
+	}
+
+	if (!ParseUnitNumber(FixtureNode, UnitNumber))
+	{
+		// Non-optional, but defaulted by standard
+		UE_LOG(LogDMXRuntime, Warning, TEXT("Imported Fixture '%s' from MVR. No valid Unit Number specified. Unit Number set to ''0'."), *Name);
+		UnitNumber = 0;
+	}
+
+	if (!ParseAddresses(FixtureNode, Addresses))
+	{
+		// Non-optional
+		UE_LOG(LogDMXRuntime, Warning, TEXT("Cannot parse Fixture '%s' from MVR. No valid Addresses specified."), *Name);
+		return;
+	}
+
+	ParseCIEColor(FixtureNode, CIEColor);
+	ParseFixtureTypeId(FixtureNode, FixtureTypeId);
+	ParseCustomId(FixtureNode, CustomId);
+	ParseMapping(FixtureNode, Mapping);
+	ParseGobo(FixtureNode, Gobo);
+}
+
+bool FDMXMVRFixture::IsValid() const
+{
+	if (!UUID.IsValid())
+	{
+		return false;
+	}
+	
+	if (Name.IsEmpty())
+	{
+		return false;
+	}
+
+	if (GDTFSpec.IsEmpty())
+	{
+		return false;
+	}
+
+	if (GDTFMode.IsEmpty())
+	{
+		return false;
+	}
+
+	if (Addresses.Universe < 0 || Addresses.Universe > DMX_MAX_UNIVERSE)
+	{
+		return false;
+	}
+
+	if (Addresses.Address < 1 || Addresses.Address > DMX_MAX_ADDRESS)
+	{
+		return false;
+	}
+
+	return true;
+}
diff --git a/Engine/Plugins/VirtualProduction/DMX/DMXEngine/Source/DMXRuntime/Private/MVR/DMXMVRGeneralSceneDescription.cpp b/Engine/Plugins/VirtualProduction/DMX/DMXEngine/Source/DMXRuntime/Private/MVR/DMXMVRGeneralSceneDescription.cpp
--- a/Engine/Plugins/VirtualProduction/DMX/DMXEngine/Source/DMXRuntime/Private/MVR/DMXMVRGeneralSceneDescription.cpp
+++ b/Engine/Plugins/VirtualProduction/DMX/DMXEngine/Source/DMXRuntime/Private/MVR/DMXMVRGeneralSceneDescription.cpp
@@ -4,609 +4,160 @@
 
 #include "DMXProtocolCommon.h"
 #include "DMXRuntimeLog.h"
+#include "Library/DMXEntityFixturePatch.h"
+#include "Library/DMXEntityFixtureType.h"
+#include "Library/DMXImportGDTF.h"
+#include "Library/DMXLibrary.h"
 
 #include "XmlFile.h"
 #include "XmlNode.h"
+#include "Misc/Paths.h"
 
 
-namespace UE::DMXRuntime::DMXMVRFixture::Private
+UDMXMVRGeneralSceneDescription* UDMXMVRGeneralSceneDescription::CreateFromXmlFile(TSharedRef<FXmlFile> GeneralSceneDescriptionXml, UObject* Outer, FName Name, EObjectFlags Flags)
 {
-	/** Parses the Name from the MVR Fixture Node, returns true on success. */
-	bool ParseName(const FXmlNode* const InFixtureNode, FString& OutName)
-	{
-		check(InFixtureNode);
-
-		static const FString InFixtureNodeAttribute_Name = TEXT("Name");
-		OutName = InFixtureNode->GetAttribute(InFixtureNodeAttribute_Name);
-		if (!OutName.IsEmpty())
-		{
-			return true;
-		}
-
-		return false;
-	}
-
-	/** Parses the UUID from the MVR Fixture Node, returns true on success. */
-	bool ParseUUID(const FXmlNode* const InFixtureNode, FGuid& OutUUID)
-	{
-		check(InFixtureNode);
-
-		static const FString InFixtureNodeAttribute_UUID = TEXT("UUID");
-		const FString UUIDString = InFixtureNode->GetAttribute(InFixtureNodeAttribute_UUID);
-		if (FGuid::Parse(UUIDString, OutUUID))
-		{
-			return true;
-		}
-		return false;
-	}
-
-	/** Parses the Matrix from the MVR Fixture Node, optional, always succeeds. */
-	void ParseMatrix(const FXmlNode* const InFixtureNode, TOptional<FTransform>& OutTransform)
-	{
-		check(InFixtureNode);
-		OutTransform.Reset();
-
-		static const FString ChildNodeName_Matrix = TEXT("Matrix");
-		const FXmlNode* const MatrixNode = InFixtureNode->FindChildNode(ChildNodeName_Matrix);
+	UDMXMVRGeneralSceneDescription* GeneralSceneDescription = NewObject<UDMXMVRGeneralSceneDescription>(Outer, Name, Flags);
+	GeneralSceneDescription->ParseGeneralSceneDescriptionXml(GeneralSceneDescriptionXml);
 
-		if (!MatrixNode)
-		{
-			return;
-		}
-
-		auto GetVectorStringsFromMatrixStringLambda([](const FString& MatrixString) -> TArray<FString>
-			{
-				TArray<FString> Result;
-				int32 SubstringIndex = 0;
-				for (int32 SubStringNumber = 0; SubStringNumber < 4; SubStringNumber++)
-				{
-					SubstringIndex = MatrixString.Find(TEXT("{"), ESearchCase::IgnoreCase, ESearchDir::FromStart, SubstringIndex) + 1;
-					const int32 EndIndex = MatrixString.Find(TEXT("}"), ESearchCase::IgnoreCase, ESearchDir::FromStart, SubstringIndex);
-
-					const FString SubString = MatrixString.Mid(SubstringIndex, EndIndex - SubstringIndex);
-					Result.Add(SubString);
-				}
-
-				return Result;
-			});
-
-		auto GetVectorFromVectorStringLambda([](const FString& VectorString) -> FVector
-			{
-				const int32 XIndex = 0;
-				const int32 YIndex = VectorString.Find(TEXT(","), ESearchCase::IgnoreCase, ESearchDir::FromStart, 0) + 1;
-				const int32 ZIndex = VectorString.Find(TEXT(","), ESearchCase::IgnoreCase, ESearchDir::FromStart, YIndex) + 1;
-
-				const FString XString = "X=" + VectorString.Mid(XIndex, YIndex - XIndex);
-				const FString YString = "Y=" + VectorString.Mid(YIndex, ZIndex - YIndex);
-				const FString ZString = "Z=" + VectorString.Mid(ZIndex, VectorString.Len() - ZIndex);
+	return GeneralSceneDescription;
+}
 
-				FVector Vector;
-				if (Vector.InitFromString(XString + YString + ZString))
-				{
-					return Vector;
-				}
+UDMXMVRGeneralSceneDescription* UDMXMVRGeneralSceneDescription::CreateFromDMXLibrary(const UDMXLibrary& DMXLibrary, UObject* Outer, FName Name, EObjectFlags Flags)
+{
+	UDMXMVRGeneralSceneDescription* GeneralSceneDescription = NewObject<UDMXMVRGeneralSceneDescription>(Outer, Name, Flags);
+	GeneralSceneDescription->WriteDMXLibraryToGeneralSceneDescription(DMXLibrary);
 
-				return FVector::ZeroVector;
-			});
+	return GeneralSceneDescription;
+}
 
-		const TArray<FString> SubStrings = GetVectorStringsFromMatrixStringLambda(MatrixNode->GetContent());
-		if (SubStrings.Num() < 4)
+FDMXMVRFixture* UDMXMVRGeneralSceneDescription::FindMVRFixture(const FGuid& MVRFixtureUUID)
+{
+	return MVRFixtures.FindByPredicate([MVRFixtureUUID](const FDMXMVRFixture& MVRFixture)
 		{
-			return;
-		}
-
-		const FVector U = GetVectorFromVectorStringLambda(SubStrings[0]) / 10.f;
-		const FVector V = GetVectorFromVectorStringLambda(SubStrings[1]) / 10.f;
-		const FVector W = GetVectorFromVectorStringLambda(SubStrings[2]) / 10.f;
-		const FVector O = GetVectorFromVectorStringLambda(SubStrings[3]) / 10.f;
-
-		FMatrix Matrix;
-		Matrix.M[0][0] = U.X;
-		Matrix.M[0][1] = V.X;
-		Matrix.M[0][2] = W.X;
-		Matrix.M[0][3] = O.X;
-
-		Matrix.M[1][0] = U.Y;
-		Matrix.M[1][1] = V.Y;
-		Matrix.M[1][2] = W.Y;
-		Matrix.M[1][3] = -O.Y; // From MVR's right hand to UE's left hand coordinate system
-
-		Matrix.M[2][0] = U.Z;
-		Matrix.M[2][1] = V.Z;
-		Matrix.M[2][2] = W.Z;
-		Matrix.M[2][3] = O.Z;
-
-		Matrix.M[3][0] = 0.0;
-		Matrix.M[3][1] = 0.0;
-		Matrix.M[3][2] = 0.0;
-		Matrix.M[3][3] = 1.0;
-
-		OutTransform = FTransform(Matrix.GetTransposed());
-
-		// Scale from milimeters to centimeters
-		OutTransform->SetScale3D(OutTransform->GetScale3D() * 10.f);
+			return MVRFixture.UUID == MVRFixtureUUID;
+		});
+}
 
-		// GDTFs are facing down on the Z-Axis, but UE Actors are facing up 
-		const FQuat InvertUpRotationQuaternion = FQuat(OutTransform->GetRotation().GetAxisY(), PI);
-		OutTransform->SetRotation(OutTransform->GetRotation() * InvertUpRotationQuaternion);
-	}
-	
-	/** Parses the GDTFSpec from the MVR Fixture Node, returns true on success. */
-	bool ParseGDTFSpec(const FXmlNode* const InFixtureNode, FString& OutGDTFSpec)
+void UDMXMVRGeneralSceneDescription::AddMVRFixture(FDMXMVRFixture& MVRFixture)
+{
+	if (ensureAlwaysMsgf(MVRFixture.UUID.IsValid(), TEXT("Trying to add MVR Fixture '%s' to General Scene Description, but the UUID is invalid."), *MVRFixture.Name))
 	{
-		check(InFixtureNode);
-
-		static const FString ChildNodeName_GDTFSpec = TEXT("GDTFSpec");
-		const FXmlNode* const GDTFSpecNode = InFixtureNode->FindChildNode(ChildNodeName_GDTFSpec);
-		if (GDTFSpecNode)
-		{
-			OutGDTFSpec = GDTFSpecNode->GetContent();
-			return true;
-		}
-
-		return false;
+		return;
 	}
 
-	/** Parses the GDTFMode from the MVR Fixture Node, returns true on success. */
-	bool ParseGDTFMode(const FXmlNode* const InFixtureNode, FString& OutGDTFMode)
+	const bool bUUIDAlreadyContained = FindMVRFixture(MVRFixture.UUID) != nullptr;
+	if (ensureAlwaysMsgf(bUUIDAlreadyContained, TEXT("Trying to add MVR Fixture '%s' to General Scene Description, but its UUID is already containted in the General Scene Description."), *MVRFixture.Name))
 	{
-		check(InFixtureNode);
-
-		static const FString ChildNodeName_GDTFMode = TEXT("GDTFMode");
-		const FXmlNode* const GDTFModeNode = InFixtureNode->FindChildNode(ChildNodeName_GDTFMode);
-		if (GDTFModeNode)
-		{
-			OutGDTFMode = GDTFModeNode->GetContent();
-			return true;
-		}
-		return true;
+		return;
 	}
 
-	/** Parses Focus from the MVR Fixture Node, optional, always succeeds. */
-	void ParseFocus(const FXmlNode* const InFixtureNode, TOptional<FGuid>& OutFocus)
-	{
-		check(InFixtureNode);
-		OutFocus.Reset();
+	MVRFixtures.Add(MVRFixture);
+}
 
-		static const FString ChildNodeName_Focus = TEXT("Focus");
-		const FXmlNode* const FocusNode = InFixtureNode->FindChildNode(ChildNodeName_Focus);
-		if (FocusNode)
+void UDMXMVRGeneralSceneDescription::WriteDMXLibraryToGeneralSceneDescription(const UDMXLibrary& DMXLibrary)
+{
+	TArray<FDMXMVRFixture> MVRFixturesInDMXLibrary;
+
+	TArray<UDMXEntityFixturePatch*> FixturePatches = DMXLibrary.GetEntitiesTypeCast<UDMXEntityFixturePatch>();
+	FixturePatches.Sort([](const UDMXEntityFixturePatch& FixturePatchA, const UDMXEntityFixturePatch& FixturePatchB)
 		{
-			const FString FocusString = InFixtureNode->GetContent();
-			FGuid FocusGuid;
-			if (FGuid::Parse(FocusString, FocusGuid))
-			{
-				OutFocus = FocusGuid;
-			}
-		};
-	}
+			const bool bUniverseIsSmaller = FixturePatchA.GetUniverseID() < FixturePatchB.GetUniverseID();
+			const bool bUniverseIsEqual = FixturePatchA.GetUniverseID() == FixturePatchB.GetUniverseID();
+			const bool bAddressIsSmaller = FixturePatchA.GetStartingChannel() <= FixturePatchB.GetStartingChannel();
 
-	/** Parses Cast Shadows from the MVR Fixture Node, optional, always succeeds. */
-	void ParseCastShadow(const FXmlNode* const InFixtureNode, TOptional<bool>& OutCastShadow)
+			return bUniverseIsSmaller || (bUniverseIsEqual && bAddressIsSmaller);
+		});
+
+	TArray<FGuid> MVRFixtureUUIDsInUse;
+	for (UDMXEntityFixturePatch* FixturePatch : FixturePatches)
 	{
-		check(InFixtureNode);
-		OutCastShadow.Reset();
+		MVRFixtureUUIDsInUse.Add(FixturePatch->GetMVRFixtureUUID());
+	}
 
-		static const FString ChidNodeName_CastShadow = TEXT("CastShadow");
-		const FXmlNode* const CastShadowNode = InFixtureNode->FindChildNode(ChidNodeName_CastShadow);
-		if (CastShadowNode)
+	// Remove entries no longer present in the DMX Library
+	MVRFixtures.RemoveAll([&MVRFixtureUUIDsInUse](const FDMXMVRFixture& MVRFixture)
 		{
-			OutCastShadow = !(CastShadowNode->GetContent() == TEXT("0")) || !CastShadowNode->GetContent().Equals(TEXT("true"), ESearchCase::IgnoreCase);
-		}
-	}
+			return !MVRFixtureUUIDsInUse.Contains(MVRFixture.UUID);
+		});
 
-	/** Parses Position of from MVR Fixture Node, optional, always succeeds. */
-	void ParsePosition(const FXmlNode* const InFixtureNode, TOptional<FGuid>& OutPosition)
+	// Create or update MVR Fixtures for each Fixture Patch's MVR Fixture UUIDs
+	for (UDMXEntityFixturePatch* FixturePatch : FixturePatches)
 	{
-		check(InFixtureNode);
-		OutPosition.Reset();
-
-		static const FString ChildNodenName_Position = TEXT("Position");
-		const FXmlNode* const PositionNode = InFixtureNode->FindChildNode(ChildNodenName_Position);
-		if (PositionNode)
+		if (FixturePatch)
 		{
-			FGuid PositionGuid;
-			if (FGuid::Parse(PositionNode->GetContent(), PositionGuid))
-			{
-				OutPosition = PositionGuid;
-			}
+			WriteFixturePatchToGeneralSceneDescription(*FixturePatch);
 		}
 	}
+}
 
-	/** Parses the Fixture Id from the MVR Fixture Node, returns true on success. */
-	bool ParseFixtureId(const FXmlNode* const InFixtureNode, FString& OutFixtureId)
-	{
-		check(InFixtureNode);
-		OutFixtureId.Reset();
-
-		static const FString ChildNodenName_FixtureId = TEXT("FixtureId");
-		const FXmlNode* const FixtureIdNode = InFixtureNode->FindChildNode(ChildNodenName_FixtureId);
-		if (FixtureIdNode)
-		{
-			OutFixtureId = FixtureIdNode->GetContent();
-			return true;
-		}
+void UDMXMVRGeneralSceneDescription::Serialize(FArchive& Ar)
+{
+	Super::Serialize(Ar);
 
-		return false;
-	}
+	Ar << MVRFixtures;
+}
 
-	/** Parses the Unit Number from the MVR Fixture Node, returns true on success. */
-	bool ParseUnitNumber(const FXmlNode* const InFixtureNode, int32& OutUnitNumber)
+void UDMXMVRGeneralSceneDescription::WriteFixturePatchToGeneralSceneDescription(const UDMXEntityFixturePatch& FixturePatch)
+{
+	if (const UDMXLibrary* DMXLibrary = FixturePatch.GetParentLibrary())
 	{
-		check(InFixtureNode);
+		const FGuid& MVRFixtureUUID = FixturePatch.GetMVRFixtureUUID();
+		FDMXMVRFixture* MVRFixturePtr = FindMVRFixture(MVRFixtureUUID);
 
-		static const FString ChildNodenName_UnitNumber = TEXT("FixtureId");
-		const FXmlNode* const UnitNumberNode = InFixtureNode->FindChildNode(ChildNodenName_UnitNumber);
-		if (UnitNumberNode)
+		FDMXMVRFixture MVRFixture;
+		if (MVRFixturePtr)
 		{
-			int32 UnitNumber;
-			if (LexTryParseString(UnitNumber, *UnitNumberNode->GetContent()))
-			{
-				OutUnitNumber = UnitNumber;
-				return true;
-			}
+			MVRFixture = *MVRFixturePtr;
 		}
-		return false;
-	}
-
-	/** Parses the Addresses from the MVR Fixture Node, defaulted, always succeeds */
-	bool ParseAddresses(const FXmlNode* const InFixtureNode, FDMXMVRFixtureAddresses& OutAddresses)
-	{
-		check(InFixtureNode);
-
-		static const FString ChildNodenName_Addresses = TEXT("Addresses");
-		const FXmlNode* const AddressesNode = InFixtureNode->FindChildNode(ChildNodenName_Addresses);
-
-		static const FString ChildNodenName_Address = TEXT("Address");
-		const FXmlNode* const AddressNode = AddressesNode ? AddressesNode->FindChildNode(ChildNodenName_Address) : nullptr;
-
-		if (AddressNode)
+		else
 		{
-			int32 Universe = -1;
-			int32 Address = -1;
-
-			auto ParseAddressStringLambda = [](const FXmlNode* const InAddressNode, int32& OutUniverse, int32& OutAddress) -> const bool
-			{
-				const FString InAddressString = InAddressNode->GetContent();
-
-				int32 SeparatorIndex = INDEX_NONE;
-				if (InAddressString.FindChar('.', SeparatorIndex))
-				{
-					if (SeparatorIndex > 0 && SeparatorIndex < InAddressString.Len() - 2)
-					{
-						const FString UniverseString = InAddressString.Left(SeparatorIndex);
-						const FString AddressString = InAddressString.RightChop(SeparatorIndex + 1);
-
-						const bool bValidUniverse = LexTryParseString(OutUniverse, *UniverseString) && OutUniverse >= 0 && OutUniverse <= DMX_MAX_UNIVERSE;
-						const bool bValidAddress = LexTryParseString(OutAddress, *UniverseString) && OutAddress >= 0 && OutAddress <= DMX_MAX_ADDRESS;
-
-						return bValidUniverse && bValidAddress;
-					}
-				}
-				else if (LexTryParseString(OutAddress, *InAddressString))
-				{
-					static const int32 UniverseSize = 512;
+			const TArray<int32> UnitNumbersInUse = GetUnitNumbersInUse(*DMXLibrary);
+			const int32 UnitNumber = UnitNumbersInUse.Num() > 0 ? UnitNumbersInUse.Last() + 1 : 1;
 
-					OutUniverse = OutAddress / UniverseSize;
-					OutAddress = OutAddress % UniverseSize;
+			MVRFixture.Name = FixturePatch.Name;
+			MVRFixture.UUID = MVRFixtureUUID;
 
-					return 
-						OutUniverse >= 0 && 
-						OutUniverse <= DMX_MAX_UNIVERSE &&
-						OutAddress >= 0 && 
-						OutAddress <= DMX_MAX_ADDRESS;
-				}
-
-				return false;
-			}(AddressNode, Universe, Address);
-			if (ParseAddressStringLambda)
-			{
-				OutAddresses.Universe = Universe;
-				OutAddresses.Address = Address;
-
-				return true;
-			}
+			// Create a new Unit Number and add it to the array of Unit Numbers in use
+			MVRFixture.UnitNumber = UnitNumbersInUse.Last() + 1;
 		}
-		return false;
-	}
 
-	/** Parses the CID Color from the MVR Fixture Node, optional, always succeeds. */
-	void ParseCIEColor(const FXmlNode* const InFixtureNode, TOptional<FDMXMVRColorCIE>& OutCIEColor)
-	{
-		check(InFixtureNode);
+		MVRFixture.Addresses.Universe = FixturePatch.GetUniverseID();
+		MVRFixture.Addresses.Address = FixturePatch.GetStartingChannel();
 
-		static const FString ChildNodenName_CIEColor = TEXT("CIEColor");
-		const FXmlNode* const CIEColorNode = InFixtureNode->FindChildNode(ChildNodenName_CIEColor);
-		if (CIEColorNode)
+		bool bSetGDTFSpec = false;
+		if (UDMXEntityFixtureType* FixtureType = FixturePatch.GetFixtureType())
 		{
-			const FString CIEColorString = CIEColorNode->GetContent();
-			if (CIEColorString.IsEmpty())
+			if (FixtureType->GDTF)
 			{
-				return;
+				MVRFixture.GDTFSpec = FPaths::GetBaseFilename(FixtureType->GDTF->SourceFilename);
 			}
 
-			TArray<FString> CIEColorArray;
-			CIEColorString.ParseIntoArray(CIEColorArray, TEXT(","));
-
-			if (CIEColorArray.Num() != 3	)
+			const int32 ModeIndex = FixturePatch.GetActiveModeIndex();
+			if (ensureAlwaysMsgf(FixtureType->Modes.IsValidIndex(ModeIndex), TEXT("Trying to write Active Mode of %s to General Scene Description, but the Mode Index is invalid")))
 			{
-				return;
-			}
-
-			FDMXMVRColorCIE Result;
-			bool bSuccess = LexTryParseString<float>(Result.X, *CIEColorArray[0]);
-			bSuccess = bSuccess && LexTryParseString<float>(Result.Y, *CIEColorArray[1]);
-			bSuccess = bSuccess && LexTryParseString<uint8>(Result.YY, *CIEColorArray[2]);
-
-			if (bSuccess)
-			{
-				OutCIEColor = Result;
+				MVRFixture.GDTFMode = FixtureType->Modes[ModeIndex].ModeName;
 			}
 		}
-	}
-
-	/** Parses the Fixture Type Id from MVR Fixture Node, optional, always succeeds. */
-	void ParseFixtureTypeId(const FXmlNode* const InFixtureNode, TOptional<int32>& OutFixtureTypeId)
-	{
-		check(InFixtureNode);
-		OutFixtureTypeId.Reset();
 
-		static const FString ChildNodenName_FixtureTypeId = TEXT("FixtureTypeId");
-		const FXmlNode* const FixtureTypeIdNode = InFixtureNode->FindChildNode(ChildNodenName_FixtureTypeId);
-		if (FixtureTypeIdNode)
+		if (!bSetGDTFSpec)
 		{
-			int32 FixtureTypeId;
-			if (LexTryParseString(FixtureTypeId, *FixtureTypeIdNode->GetContent()))
-			{
-				OutFixtureTypeId = FixtureTypeId;
-			}
+			// No point to set a mode when there's no GDTF
+			MVRFixture.GDTFMode = TEXT("");
+			MVRFixture.GDTFSpec = TEXT("");
 		}
-	}
 
-	/** Parses the Custom Id from MVR Fixture Node, optional, always succeeds. */
-	void ParseCustomId(const FXmlNode* const InFixtureNode, TOptional<int32>& OutCustomId)
-	{
-		check(InFixtureNode);
-		OutCustomId.Reset();
-
-		static const FString ChildNodenName_Position = TEXT("Position");
-		const FXmlNode* const PositionNode = InFixtureNode->FindChildNode(ChildNodenName_Position);
-		if (PositionNode)
+		if (MVRFixturePtr)
 		{
-			int32 FixtureTypeId;
-			if (LexTryParseString(FixtureTypeId, *PositionNode->GetContent()))
-			{
-				OutCustomId = FixtureTypeId;
-			}
+			(*MVRFixturePtr) = MVRFixture;
 		}
-	}
-
-	/** Parses the Mapping from MVR Fixture Node, optional, always succeeds. */
-	void ParseMapping(const FXmlNode* const InFixtureNode, TOptional<FDMXMVRFixtureMapping>& OutMapping)
-	{
-		check(InFixtureNode);
-		OutMapping.Reset();
-
-		static const FString ChildNodenName_Mapping = TEXT("Position");
-		const FXmlNode* const MappingNode = InFixtureNode->FindChildNode(ChildNodenName_Mapping);
-		if (MappingNode)
-		{
-			int32 FixtureTypeId;
-			if (LexTryParseString(FixtureTypeId, *MappingNode->GetContent()))
-			{
-				static const FString ChildNodenName_LinkDef = TEXT("LinkDef");
-				const FXmlNode* const LinkDefNOde = MappingNode->FindChildNode(ChildNodenName_LinkDef);
-				FDMXMVRFixtureMapping Mapping;
-				if (LinkDefNOde && !FGuid::Parse(LinkDefNOde->GetContent(), Mapping.LinkDef))
-				{
-					return;
-				}
-
-				static const FString ChildNodenName_UX = TEXT("UX");
-				const FXmlNode* const UXNode = InFixtureNode->FindChildNode(ChildNodenName_UX);
-				int32 UX;
-				if (UXNode && !LexTryParseString(UX, *UXNode->GetContent()))
-				{
-					Mapping.UX = UX;
-					return;
-				}
-
-				static const FString ChildNodenName_UY = TEXT("UY");
-				const FXmlNode* const UYNode = InFixtureNode->FindChildNode(ChildNodenName_UY);
-				int32 UY;
-				if (UYNode && !LexTryParseString(UY, *UYNode->GetContent()))
-				{
-					return;
-				}
-				Mapping.UY = UY;
-
-				static const FString ChildNodenName_OX = TEXT("OX");
-				const FXmlNode* const OXNode = InFixtureNode->FindChildNode(ChildNodenName_OX);
-				int32 OX;
-				if (OXNode && !LexTryParseString(OX, *OXNode->GetContent()))
-				{
-					return;
-				}
-				Mapping.OX = OX;
-
-				static const FString ChildNodenName_OY = TEXT("OY");
-				const FXmlNode* const OYNode = InFixtureNode->FindChildNode(ChildNodenName_OY);
-				int32 OY;
-				if (OYNode && !LexTryParseString(OY, *OYNode->GetContent()))
-				{
-					return;
-				}
-				Mapping.OY = OY;
-
-				static const FString ChildNodenName_RZ = TEXT("RZ");
-				const FXmlNode* const RZNode = InFixtureNode->FindChildNode(ChildNodenName_RZ);
-				int32 RZ;
-				if (RZNode && !LexTryParseString(RZ, *RZNode->GetContent()))
-				{
-					return;
-				}
-				Mapping.RZ = RZ;
-
-				OutMapping = Mapping;
-			}
-		}
-	}
-
-	/** Parses the Gobo from MVR Fixture Node, optional, always succeeds. */
-	void ParseGobo(const FXmlNode* const InFixtureNode, TOptional<FDMXMVRFixtureGobo>& OutGobo)
-	{
-		check(InFixtureNode);
-		OutGobo.Reset();
-
-		static const FString ChildNodenName_Gobo = TEXT("Position");
-		const FXmlNode* const GoboNode = InFixtureNode->FindChildNode(ChildNodenName_Gobo);
-		if (GoboNode)
+		else
 		{
-			const FString& GoboResource = GoboNode->GetContent();
-			OutGobo->Value = GoboResource;
-
-			static const FString ChildNodeName_Gobo = TEXT("Rotation");
-			const FXmlNode* const GoboRotationNode = GoboNode->FindChildNode(ChildNodeName_Gobo);
-			float GoboRotation;
-			if (GoboRotationNode && LexTryParseString(GoboRotation, *GoboRotationNode->GetContent()))
-			{
-				OutGobo->Rotation = GoboRotation;
-			}
+			MVRFixtures.Add(MVRFixture);
 		}
 	}
 }
 
-FDMXMVRFixture::FDMXMVRFixture()
-{}
-
-FDMXMVRFixture::FDMXMVRFixture(const FXmlNode* const FixtureNode)
-{
-	// Ensure a valid node
-	if (!ensureAlwaysMsgf(FixtureNode, TEXT("Trying to parse MVR Fixture Node, but node is null.")))
-	{
-		return;
-	}
-
-	static const FString NodeName_Fixture = TEXT("Fixture");
-	if (!ensureAlwaysMsgf(FixtureNode->GetTag() == NodeName_Fixture, TEXT("Trying to parse MVR Fixture Node, but node name doesn't match.")))
-	{
-		return;
-	}
-
-	// Parse nodes. Those that don't fail are either defaulted or optional.
-	using namespace UE::DMXRuntime::DMXMVRFixture::Private;
-	if (!ParseName(FixtureNode, Name))
-	{
-		// Non-optional
-		UE_LOG(LogDMXRuntime, Warning, TEXT("Cannot parse Fixture Node as MVR Fixture Name. Failed to parse MVR Node %s."), *FixtureNode->GetTag());
-		return;
-	}
-
-	if(!ParseUUID(FixtureNode, UUID))
-	{
-		// Non-optional
-		UE_LOG(LogDMXRuntime, Warning, TEXT("Cannot parse Fixture '%s' from MVR. Failed to parse MVR UUID."), *Name);
-		return;
-	}
-
-	ParseMatrix(FixtureNode, Transform);
-	if (Transform.IsSet() && !Transform.GetValue().IsValid())
-	{
-		// Mend invalid transforms
-		UE_LOG(LogDMXRuntime, Warning, TEXT("Imported Fixture '%s' from MVR but its Transform is not valid, using Identity instead."), *Name);
-		Transform = FTransform::Identity;
-	}
-
-	if (!ParseGDTFSpec(FixtureNode, GDTFSpec))
-	{
-		// Non-optional
-		UE_LOG(LogDMXRuntime, Warning, TEXT("Cannot parse Fixture '%s' from MVR. No valid GDTF Spec specified."), *Name);
-		return;
-	}
-
-	if (!ParseGDTFMode(FixtureNode, GDTFMode))
-	{
-		// Non-optional
-		UE_LOG(LogDMXRuntime, Warning, TEXT("Cannot parse Fixture '%s' from MVR. No valid GDTF Mode specified."), *Name);
-		return;
-	}
-
-	ParseFocus(FixtureNode, Focus);
-	ParseCastShadow(FixtureNode, bCastShadows);
-	ParsePosition(FixtureNode, Position);
-	
-	if (!ParseFixtureId(FixtureNode, FixtureId))
-	{
-		// Non-optional, mend with empty string
-		UE_LOG(LogDMXRuntime, Warning, TEXT("Imported Fixture '%s' from MVR. No valid Fixture Id specified. Fixture Id set to Fixture Name."), *Name);
-		FixtureId = TEXT("");
-	}
-
-	if (!ParseUnitNumber(FixtureNode, UnitNumber))
-	{
-		// Non-optional, but defaulted by standard
-		UE_LOG(LogDMXRuntime, Warning, TEXT("Imported Fixture '%s' from MVR. No valid Unit Number specified. Unit Number set to ''0'."), *Name);
-		UnitNumber = 0;
-	}
-
-	if (!ParseAddresses(FixtureNode, Addresses))
-	{
-		// Non-optional
-		UE_LOG(LogDMXRuntime, Warning, TEXT("Cannot parse Fixture '%s' from MVR. No valid Addresses specified."), *Name);
-		return;
-	}
-
-	ParseCIEColor(FixtureNode, CIEColor);
-	ParseFixtureTypeId(FixtureNode, FixtureTypeId);
-	ParseCustomId(FixtureNode, CustomId);
-	ParseMapping(FixtureNode, Mapping);
-	ParseGobo(FixtureNode, Gobo);
-}
-
-bool FDMXMVRFixture::IsValid() const
-{
-	if (!UUID.IsValid())
-	{
-		return false;
-	}
-	
-	if (Name.IsEmpty())
-	{
-		return false;
-	}
-
-	if (GDTFSpec.IsEmpty())
-	{
-		return false;
-	}
-
-	if (GDTFMode.IsEmpty())
-	{
-		return false;
-	}
-
-	if (Addresses.Universe < 0 || Addresses.Universe > DMX_MAX_UNIVERSE)
-	{
-		return false;
-	}
-
-	if (Addresses.Address < 1 || Addresses.Address > DMX_MAX_ADDRESS)
-	{
-		return false;
-	}
-
-	return true;
-}
-
-FDMXMVRGeneralSceneDescription::FDMXMVRGeneralSceneDescription()
-{}
-
-FDMXMVRGeneralSceneDescription::FDMXMVRGeneralSceneDescription(TSharedRef<FXmlFile> GeneralSceneDescription)
-{
-	ParseMVRFixtures(GeneralSceneDescription);
-}
-
-void FDMXMVRGeneralSceneDescription::ParseMVRFixtures(const TSharedRef<FXmlFile>& GeneralSceneDescription)
+void UDMXMVRGeneralSceneDescription::ParseGeneralSceneDescriptionXml(const TSharedRef<FXmlFile>& GeneralSceneDescription)
 {
 	if (const FXmlNode* RootNode = GeneralSceneDescription->GetRootNode())
 	{
@@ -646,3 +197,38 @@ void FDMXMVRGeneralSceneDescription::ParseMVRFixtures(const TSharedRef<FXmlFile>
 		}
 	}
 }
+
+TArray<int32> UDMXMVRGeneralSceneDescription::GetUnitNumbersInUse(const UDMXLibrary& DMXLibrary)
+{
+	TArray<UDMXEntityFixturePatch*> FixturePatches = DMXLibrary.GetEntitiesTypeCast<UDMXEntityFixturePatch>();
+
+	TArray<FGuid> MVRFixtureUUIDsInUse;
+	for (UDMXEntityFixturePatch* FixturePatch : FixturePatches)
+	{
+		MVRFixtureUUIDsInUse.Add(FixturePatch->GetMVRFixtureUUID());
+	}
+
+	// Fetch Unit Numbers in use to create new ones when adding new MVR Fixtures
+	TArray<int32> UnitNumbersInUse;
+	for (const FGuid& MVRFixtureUUID : MVRFixtureUUIDsInUse)
+	{
+		const FDMXMVRFixture* MVRFixturePtr = FindMVRFixture(MVRFixtureUUID);
+		if (MVRFixturePtr)
+		{
+			UnitNumbersInUse.Add(MVRFixturePtr->UnitNumber);
+		}
+	}
+	if (UnitNumbersInUse.Num() == 0)
+	{
+		UnitNumbersInUse.Add(0);
+	}
+	else
+	{
+		UnitNumbersInUse.Sort([](int32 UnitNumberA, int32 UnitNumberB)
+			{
+				return UnitNumberA < UnitNumberB;
+			});
+	}
+
+	return UnitNumbersInUse;
+}
diff --git a/Engine/Plugins/VirtualProduction/DMX/DMXEngine/Source/DMXRuntime/Private/MVR/DMXMVRSceneActor.cpp b/Engine/Plugins/VirtualProduction/DMX/DMXEngine/Source/DMXRuntime/Private/MVR/DMXMVRSceneActor.cpp
--- a/Engine/Plugins/VirtualProduction/DMX/DMXEngine/Source/DMXRuntime/Private/MVR/DMXMVRSceneActor.cpp
+++ b/Engine/Plugins/VirtualProduction/DMX/DMXEngine/Source/DMXRuntime/Private/MVR/DMXMVRSceneActor.cpp
@@ -5,6 +5,7 @@
 #include "DMXRuntimeLog.h"
 #include "Game/DMXComponent.h"
 #include "Library/DMXEntityFixturePatch.h"
+#include "Library/DMXImportGDTF.h"
 #include "Library/DMXLibrary.h"
 #include "MVR/DMXMVRFixtureActorLibrary.h"
 #include "MVR/DMXMVRAssetUserData.h"
@@ -69,7 +70,7 @@ ADMXMVRSceneActor::~ADMXMVRSceneActor()
 void ADMXMVRSceneActor::PostLoad()
 {
 	Super::PostLoad();
-	SetMVRUUIDsForRelatedActors();
+	EnsureMVRUUIDsForRelatedActors();
 }
 
 void ADMXMVRSceneActor::PostRegisterAllComponents()
@@ -88,21 +89,48 @@ void ADMXMVRSceneActor::PostRegisterAllComponents()
 		}
 	}
 
-	SetMVRUUIDsForRelatedActors();
+	EnsureMVRUUIDsForRelatedActors();
 }
 
-void ADMXMVRSceneActor::SetMVRUUIDsForRelatedActors()
+#if WITH_EDITOR
+void ADMXMVRSceneActor::PreEditChange(FProperty* PropertyAboutToChange)
+{
+	Super::PreEditChange(PropertyAboutToChange);
+
+	if (PropertyAboutToChange && PropertyAboutToChange->GetFName() == GET_MEMBER_NAME_CHECKED(FDMXMVRSceneGDTFToActorClassPair, ActorClass))
+	{
+		GDTFToDefaultActorClasses_PreEditChange = GDTFToDefaultActorClasses;
+	}
+}
+#endif // WITH_EDITOR
+
+#if WITH_EDITOR
+void ADMXMVRSceneActor::PostEditChangeProperty(FPropertyChangedEvent& PropertyChangedEvent)
+{
+	Super::PostEditChangeProperty(PropertyChangedEvent);
+
+	const FName PropertyName = PropertyChangedEvent.GetPropertyName();
+
+	if (PropertyName == GET_MEMBER_NAME_CHECKED(FDMXMVRSceneGDTFToActorClassPair, ActorClass))
+	{
+		HandleDefaultActorClassForGDTFChanged();
+	}
+}
+#endif // WITH_EDITOR
+
+void ADMXMVRSceneActor::EnsureMVRUUIDsForRelatedActors()
 {
-	for (const TTuple<FGuid, TSoftObjectPtr<AActor>>& MVRUUIDToActorPair : MVRUUIDToRelatedActorMap)
+	for (const TSoftObjectPtr<AActor>& RelatedActor : RelatedActors)
 	{
-		if (AActor* Actor = MVRUUIDToActorPair.Value.Get())
+		if (AActor* Actor = RelatedActor.Get())
 		{			
-			const FString MVRUUID = UDMXMVRAssetUserData::GetMVRAssetUserDataValueForkey(Actor, UDMXMVRAssetUserData::MVRUUIDMetaDataKey);
-			if (MVRUUID.IsEmpty())
-			{
-				if (!UDMXMVRAssetUserData::SetMVRAssetUserDataValueForKey(Actor, UDMXMVRAssetUserData::MVRUUIDMetaDataKey, MVRUUIDToActorPair.Key.ToString()))
+			const FString MVRFixtureUUID = UDMXMVRAssetUserData::GetMVRAssetUserDataValueForKey(*Actor, UDMXMVRAssetUserData::MVRFixtureUUIDMetaDataKey);
+			if (MVRFixtureUUID.IsEmpty())
+			{	
+				// Try to acquire the MVR Fixture UUID
+				if (UDMXEntityFixturePatch* FixturePatch = GetFixturePatch(Actor))
 				{
-					UE_LOG(LogDMXRuntime, Warning, TEXT("Actor %s is referenced by MVR scene but doesn't have an MVR UUID."), *Actor->GetName());
+					UDMXMVRAssetUserData::SetMVRAssetUserDataValueForKey(*Actor, UDMXMVRAssetUserData::MVRFixtureUUIDMetaDataKey, FixturePatch->GetMVRFixtureUUID().ToString());
 				}
 			}
 		}
@@ -116,26 +144,47 @@ void ADMXMVRSceneActor::SetDMXLibrary(UDMXLibrary* NewDMXLibrary)
 		return;
 	}
 
-	UWorld* World = GetWorld();
-	if (!World)
+	if (!NewDMXLibrary || NewDMXLibrary == DMXLibrary)
 	{
 		return;
 	}
+	DMXLibrary = NewDMXLibrary;
 
-	if (!NewDMXLibrary)
+
+	const TSharedRef<FDMXMVRFixtureActorLibrary> MVRFixtureActorLibrary = MakeShared<FDMXMVRFixtureActorLibrary>();
+	const TArray<UDMXEntityFixturePatch*> FixturePatches = DMXLibrary->GetEntitiesTypeCast<UDMXEntityFixturePatch>();
+
+	// Build the GDTF to Actor Class Pair array
+	for (UDMXEntityFixturePatch* FixturePatch : FixturePatches)
 	{
-		return;
+		UDMXEntityFixtureType* FixtureType = FixturePatch->GetFixtureType();
+		if (FixtureType && FixtureType->GDTF)
+		{
+			const bool bGDTFToDefaultActorClassPairAlreadyCreated = GDTFToDefaultActorClasses.ContainsByPredicate([FixtureType](const FDMXMVRSceneGDTFToActorClassPair& GDTFToActorClassPair)
+				{
+					return GDTFToActorClassPair.GDTF == FixtureType->GDTF;
+				});
+			if (bGDTFToDefaultActorClassPairAlreadyCreated)
+			{
+				continue;
+			}
+
+			UClass* ActorClass = MVRFixtureActorLibrary->FindMostAppropriateActorClassForPatch(FixturePatch);
+			FDMXMVRSceneGDTFToActorClassPair GDTFToActorClassPair;
+			GDTFToActorClassPair.GDTF = FixtureType->GDTF;
+			GDTFToActorClassPair.ActorClass = ActorClass;
+			GDTFToDefaultActorClasses.Add(GDTFToActorClassPair);
+		}
 	}
 
-	if (NewDMXLibrary == DMXLibrary)
+	// Spawn Fixture Actors
+	DMXLibrary->UpdateGeneralSceneDescription();
+	UDMXMVRGeneralSceneDescription* GeneralSceneDescription = DMXLibrary->GetLazyGeneralSceneDescription();
+	if (!GeneralSceneDescription)
 	{
 		return;
 	}
 
-	DMXLibrary = NewDMXLibrary;
-
-	const TSharedRef<FDMXMVRFixtureActorLibrary> MVRFixtureActorLibrary = MakeShared<FDMXMVRFixtureActorLibrary>();
-	const TArray<UDMXEntityFixturePatch*> FixturePatches = DMXLibrary->GetEntitiesTypeCast<UDMXEntityFixturePatch>();
 	for (UDMXEntityFixturePatch* FixturePatch : FixturePatches)
 	{
 		UClass* ActorClass = MVRFixtureActorLibrary->FindMostAppropriateActorClassForPatch(FixturePatch);
@@ -144,43 +193,17 @@ void ADMXMVRSceneActor::SetDMXLibrary(UDMXLibrary* NewDMXLibrary)
 			continue;
 		}
 
-		const TArray<UDMXMVRFixtureInstance*>& MVRFixturesInstances = FixturePatch->GetMVRFixtureInstances();
-		for (const UDMXMVRFixtureInstance* MVRFixtureInstance : MVRFixturesInstances)
+		const FGuid& MVRFixtureUUID = FixturePatch->GetMVRFixtureUUID();
+		const FDMXMVRFixture* MVRFixturePtr = GeneralSceneDescription->FindMVRFixture(MVRFixtureUUID);
+		if (!MVRFixturePtr)
 		{
-			const FDMXMVRFixture& MVRFixture = MVRFixtureInstance->GetMVRFixture();
-			if (!ensureAlwaysMsgf(MVRFixture.UUID.IsValid(), TEXT("Tried to set MVR Fixture %s for MVR Scene Actor %s, but the MVR UUID is invalid."), *MVRFixture.Name, *GetName()))
-			{
-				continue;
-			}
-
-			const FTransform Transform = MVRFixture.Transform.IsSet() ? MVRFixture.Transform.GetValue() : FTransform::Identity;
-			AActor* ChildActor = World->SpawnActor<AActor>(ActorClass, Transform);
-			if (!ChildActor)
-			{
-				continue;
-			}
+			continue;
+		}
+		const FDMXMVRFixture& MVRFixture = *MVRFixturePtr;
 
-			ChildActor->RegisterAllComponents();
-			USceneComponent* RootComponentOfChildActor = ChildActor->GetRootComponent();
-			if (!RootComponentOfChildActor)
-			{
-				continue;
-			}
-			RootComponentOfChildActor->AttachToComponent(MVRSceneRoot, FAttachmentTransformRules::KeepRelativeTransform);
-			
-			// Set the patch either via the interface or via a present DMX Component.
-			// Prefer the interface way as it may further customize how the patch is set.
-			if (IDMXMVRFixtureActorInterface* MVRFixtureActorInterface = Cast<IDMXMVRFixtureActorInterface>(ChildActor))
-			{
-				MVRFixtureActorInterface->Execute_OnMVRSetFixturePatch(ChildActor, FixturePatch);
-			}
-			else if(UActorComponent* Component = ChildActor->GetComponentByClass(UDMXComponent::StaticClass()))
-			{
-				CastChecked<UDMXComponent>(Component)->SetFixturePatch(FixturePatch);
-			}
+		const FTransform Transform = MVRFixture.Transform.IsSet() ? MVRFixture.Transform.GetValue() : FTransform::Identity;
 
-			MVRUUIDToRelatedActorMap.Add(MVRFixture.UUID, ChildActor);
-		}
+		SpawnMVRActor(ActorClass, FixturePatch, Transform);
 	}
 }
 
@@ -190,30 +213,21 @@ void ADMXMVRSceneActor::OnMapChange(uint32 MapEventFlags)
 	// Whenever a sub-level is loaded, we need to apply the fix
 	if (MapEventFlags == MapChangeEventFlags::NewMap)
 	{
-		SetMVRUUIDsForRelatedActors();
+		EnsureMVRUUIDsForRelatedActors();
 	}
 }
 #endif // WITH_EDITOR
 
 #if WITH_EDITOR
-void ADMXMVRSceneActor::OnActorDeleted(AActor* ActorDeleted)
+void ADMXMVRSceneActor::OnActorDeleted(AActor* DeletedActor)
 {
-	const FString MVRUUIDString = UDMXMVRAssetUserData::GetMVRAssetUserDataValueForkey(ActorDeleted, UDMXMVRAssetUserData::MVRUUIDMetaDataKey);
-	FGuid MVRUUID;
-	if (FGuid::Parse(MVRUUIDString, MVRUUID))
+	const int32 RelatedActorIndex = RelatedActors.Find(DeletedActor);
+	if (RelatedActorIndex != INDEX_NONE)
 	{
-		TSoftObjectPtr<AActor>* RelatedActorPtr = MVRUUIDToRelatedActorMap.Find(MVRUUID);
-		if (RelatedActorPtr)
-		{
-			AActor* RelatedActor = RelatedActorPtr->Get();
-			if (RelatedActor == ActorDeleted)
-			{
-				// This will add this actor to the transaction if there is one currently recording
-				Modify();
+		// This will add this actor to the transaction if there is one currently recording
+		Modify();
 
-				RelatedActorPtr->Reset();
-			}
-		}
+		RelatedActors[RelatedActorIndex]->Reset();
 	}
 }
 #endif // WITH_EDITOR
@@ -225,23 +239,218 @@ void ADMXMVRSceneActor::OnAssetPostImport(UFactory* InFactory, UObject* ActorAdd
 	{
 		AActor* Actor = *It;
 
-		const FString MVRUUIDString = UDMXMVRAssetUserData::GetMVRAssetUserDataValueForkey(Actor, UDMXMVRAssetUserData::MVRUUIDMetaDataKey);
-		FGuid MVRUUID;
-		if (FGuid::Parse(MVRUUIDString, MVRUUID))
+		const int32 RelatedActorIndex = RelatedActors.Find(Actor);
+		if (RelatedActorIndex != INDEX_NONE)
+		{
+			// This will add this actor to the transaction if there is one currently recording
+			Modify();
+
+			RelatedActors[RelatedActorIndex] = Actor;
+		}
+	}
+}
+#endif // WITH_EDITOR
+
+#if WITH_EDITOR
+void ADMXMVRSceneActor::HandleDefaultActorClassForGDTFChanged()
+{
+	// Handle element changes, but not add/remove
+	if (GDTFToDefaultActorClasses_PreEditChange.Num() != GDTFToDefaultActorClasses.Num())
+	{
+		return;
+	}
+
+	int32 IndexOfChangedElement = INDEX_NONE;
+	for (const FDMXMVRSceneGDTFToActorClassPair& GDTFToDefaultActorClassPair : GDTFToDefaultActorClasses)
+	{
+		IndexOfChangedElement = GDTFToDefaultActorClasses_PreEditChange.IndexOfByPredicate([&GDTFToDefaultActorClassPair](const FDMXMVRSceneGDTFToActorClassPair& GDTFToActorClassPair)
+			{
+				return
+					GDTFToActorClassPair.GDTF == GDTFToDefaultActorClassPair.GDTF &&
+					GDTFToActorClassPair.ActorClass != GDTFToDefaultActorClassPair.ActorClass;
+			});
+
+		if (IndexOfChangedElement != INDEX_NONE)
+		{
+			break;
+		}
+	}
+
+	if (IndexOfChangedElement == INDEX_NONE)
+	{
+		return;
+	}
+
+	const TSubclassOf<AActor> Class = GDTFToDefaultActorClasses[IndexOfChangedElement].ActorClass.Get();
+	if (!Class.Get())
+	{
+		return;
+	}
+
+	for (const TSoftObjectPtr<AActor>& RelatedActor : TArray<TSoftObjectPtr<AActor>>(RelatedActors))
+	{
+		AActor* Actor = RelatedActor.Get();
+		if (!Actor)
+		{
+			continue;
+		}
+
+		UDMXEntityFixturePatch* FixturePatch = GetFixturePatch(Actor);
+		if (FixturePatch && 
+			FixturePatch->GetFixtureType() && 
+			FixturePatch->GetFixtureType()->GDTF == GDTFToDefaultActorClasses[IndexOfChangedElement].GDTF)
+		{
+			ReplaceMVRActor(Actor, Class);
+		}
+	}
+}
+#endif // WITH_EDITOR
+
+AActor* ADMXMVRSceneActor::SpawnMVRActor(const TSubclassOf<AActor>&ActorClass, UDMXEntityFixturePatch* FixturePatch, const FTransform & Transform, AActor * Template)
+{
+	UWorld* World = GetWorld();
+	if (!ensureAlwaysMsgf(World, TEXT("Trying to spawn MVR Fixture in MVR Scene, but the world is not valid.")))
+	{
+		return nullptr;
+	}
+
+	if (!ensureAlwaysMsgf(FixturePatch, TEXT("Trying to spawn MVR Fixture in MVR Scene, but the Fixture Patch is not valid.")))
+	{
+		return nullptr;
+	}
+
+	FActorSpawnParameters ActorSpawnParameters;
+	ActorSpawnParameters.Template = Template;
+	ActorSpawnParameters.SpawnCollisionHandlingOverride = ESpawnActorCollisionHandlingMethod::AlwaysSpawn;
+	AActor* NewFixtureActor = World->SpawnActor<AActor>(ActorClass, Transform, ActorSpawnParameters);
+	if (!NewFixtureActor)
+	{
+		return nullptr;
+	}
+
+	NewFixtureActor->RegisterAllComponents();
+	USceneComponent* RootComponentOfChildActor = NewFixtureActor->GetRootComponent();
+	if (!RootComponentOfChildActor)
+	{
+		UE_LOG(LogDMXRuntime, Warning, TEXT("Cannot spawn MVR Fixture Actor of Class %s, the Actor does not specifiy a root component."), *ActorClass->GetName());
+		NewFixtureActor->Destroy();
+		return nullptr;
+	}
+
+#if WITH_EDITOR
+	// Create Property Change Events so editor objects related to the actor have a chance to update (e.g. Details, World Outliner).
+	PreEditChange(ADMXMVRSceneActor::StaticClass()->FindPropertyByName(GET_MEMBER_NAME_CHECKED(ADMXMVRSceneActor, RelatedActors)));
+	NewFixtureActor->PreEditChange(nullptr);
+#endif 
+
+	// Attach, set MVR Fixture UUID, set Fixture Patch, remember as a Related Actor
+	RootComponentOfChildActor->AttachToComponent(MVRSceneRoot, FAttachmentTransformRules::KeepRelativeTransform);
+	const FGuid& MVRFixtureUUID = FixturePatch->GetMVRFixtureUUID();
+	UDMXMVRAssetUserData::SetMVRAssetUserDataValueForKey(*NewFixtureActor, UDMXMVRAssetUserData::MVRFixtureUUIDMetaDataKey, MVRFixtureUUID.ToString());
+	SetFixturePatch(NewFixtureActor, FixturePatch);
+	RelatedActors.Add(NewFixtureActor);
+
+#if WITH_EDITOR
+	PostEditChange();
+	NewFixtureActor->PostEditChange();
+#endif
+
+	return NewFixtureActor;
+}
+
+AActor* ADMXMVRSceneActor::ReplaceMVRActor(AActor* ActorToReplace, const TSubclassOf<AActor>& ClassOfNewActor)
+{
+	if (!ensureAlwaysMsgf(ActorToReplace, TEXT("Trying to replace MVR Fixture in MVR Scene, but the Actor to replace is not valid.")))
+	{
+		return nullptr;
+	}
+	
+	if (ActorToReplace->GetClass() == ClassOfNewActor)
+	{
+		// No need to replace
+		return nullptr;
+	}
+
+	const FString MVRFixtureUUIDString = UDMXMVRAssetUserData::GetMVRAssetUserDataValueForKey(*ActorToReplace, UDMXMVRAssetUserData::MVRFixtureUUIDMetaDataKey);
+	FGuid MVRFixtureUUID;
+	if (FGuid::Parse(MVRFixtureUUIDString, MVRFixtureUUID))
+	{
+		// Try to find a Fixture Patch in following order:
+		// By the MVR Fixture Actor Interface, it may customize the getter
+		// By a DMX Component present in the Actor, it might have overriden the patch
+		// By MVR Fixture UUID in the DMX Library
+		UDMXEntityFixturePatch* FixturePatch = nullptr;
+		if (IDMXMVRFixtureActorInterface* MVRFixtureActorInterface = Cast<IDMXMVRFixtureActorInterface>(ActorToReplace))
+		{
+			FixturePatch = MVRFixtureActorInterface->Execute_OnMVRGetFixturePatch(ActorToReplace);
+		}
+		
+		if (!FixturePatch)
 		{
-			TSoftObjectPtr< AActor >* RelatedActorPtr = MVRUUIDToRelatedActorMap.Find(MVRUUID);
-			if (RelatedActorPtr)
+			if (UActorComponent* Component = ActorToReplace->GetComponentByClass(UDMXComponent::StaticClass()))
 			{
-				AActor* RelatedActor = RelatedActorPtr->Get();
-				if (!RelatedActor)
+				FixturePatch = CastChecked<UDMXComponent>(Component)->GetFixturePatch();
+			}
+		}
+
+		if (!FixturePatch && DMXLibrary)
+		{
+			const TArray<UDMXEntityFixturePatch*> FixturePatches = DMXLibrary->GetEntitiesTypeCast<UDMXEntityFixturePatch>();
+			UDMXEntityFixturePatch* const* FixturePatchPtr = FixturePatches.FindByPredicate([&MVRFixtureUUID](const UDMXEntityFixturePatch* FixturePatch)
 				{
-					// This will add this actor to the transaction if there is one currently recording
-					Modify();
+					return FixturePatch->GetMVRFixtureUUID() == MVRFixtureUUID;
+				});
 
-					*RelatedActorPtr = Actor;
-				}
+			if (FixturePatchPtr)
+			{
+				FixturePatch = *FixturePatchPtr;
 			}
 		}
+
+		if (AActor* NewFixtureActor = SpawnMVRActor(ClassOfNewActor, FixturePatch, ActorToReplace->GetTransform()))
+		{
+			RelatedActors.Remove(ActorToReplace);
+			ActorToReplace->Destroy();
+			return NewFixtureActor;
+		}
+	}
+
+	return nullptr;
+}
+
+UDMXEntityFixturePatch* ADMXMVRSceneActor::GetFixturePatch(AActor* Actor) const
+{
+	UDMXEntityFixturePatch* FixturePatch = nullptr;
+	if (IDMXMVRFixtureActorInterface* MVRFixtureActorInterface = Cast<IDMXMVRFixtureActorInterface>(Actor))
+	{
+		FixturePatch = MVRFixtureActorInterface->Execute_OnMVRGetFixturePatch(Actor);
+	}
+	if (!FixturePatch)
+	{
+		if (UActorComponent* Component = Actor->GetComponentByClass(UDMXComponent::StaticClass()))
+		{
+			FixturePatch = CastChecked<UDMXComponent>(Component)->GetFixturePatch();
+		}
+	}
+
+	return FixturePatch;
+}
+
+void ADMXMVRSceneActor::SetFixturePatch(AActor* Actor, UDMXEntityFixturePatch* FixturePatch)
+{
+	if (!ensureMsgf(Actor || FixturePatch, TEXT("Trying to Set Fixture Patch on Actor, but Actor or Fixture Patch are invalid.")))
+	{
+		return;
+	}
+
+	// Set the patch either via the interface or via a present DMX Component.
+	// Prefer the interface way as it may further customize how the patch is set.
+	if (IDMXMVRFixtureActorInterface* MVRFixtureActorInterface = Cast<IDMXMVRFixtureActorInterface>(Actor))
+	{
+		MVRFixtureActorInterface->Execute_OnMVRSetFixturePatch(Actor, FixturePatch);
+	}
+	else if (UActorComponent* Component = Actor->GetComponentByClass(UDMXComponent::StaticClass()))
+	{
+		CastChecked<UDMXComponent>(Component)->SetFixturePatch(FixturePatch);
 	}
 }
-#endif // WITH_EDITOR
diff --git a/Engine/Plugins/VirtualProduction/DMX/DMXEngine/Source/DMXRuntime/Private/Sequencer/MovieSceneDMXLibrarySection.cpp b/Engine/Plugins/VirtualProduction/DMX/DMXEngine/Source/DMXRuntime/Private/Sequencer/MovieSceneDMXLibrarySection.cpp
--- a/Engine/Plugins/VirtualProduction/DMX/DMXEngine/Source/DMXRuntime/Private/Sequencer/MovieSceneDMXLibrarySection.cpp
+++ b/Engine/Plugins/VirtualProduction/DMX/DMXEngine/Source/DMXRuntime/Private/Sequencer/MovieSceneDMXLibrarySection.cpp
@@ -5,7 +5,7 @@
 #include "DMXConversions.h"
 #include "DMXProtocolCommon.h"
 #include "DMXRuntimeLog.h"
-#include "DMXRuntimeObjectVersion.h"
+#include "DMXRuntimeMainStreamObjectVersion.h"
 #include "DMXSubsystem.h"
 #include "Interfaces/IDMXProtocol.h"
 #include "IO/DMXInputPort.h"
@@ -319,16 +319,16 @@ void UMovieSceneDMXLibrarySection::Serialize(FArchive& Ar)
 {
 	Super::Serialize(Ar);
 
-	Ar.UsingCustomVersion(FDMXRuntimeObjectVersion::GUID);
+	Ar.UsingCustomVersion(FDMXRuntimeMainStreamObjectVersion::GUID);
 	if (Ar.IsLoading())
 	{
-		if (Ar.CustomVer(FDMXRuntimeObjectVersion::GUID) < FDMXRuntimeObjectVersion::DefaultToNormalizedValuesInDMXLibrarySection)
+		if (Ar.CustomVer(FDMXRuntimeMainStreamObjectVersion::GUID) < FDMXRuntimeMainStreamObjectVersion::DefaultToNormalizedValuesInDMXLibrarySection)
 		{
 			// For assets created before normalized values were supported, use absolute values
 			bUseNormalizedValues = false;
 		}
 
-		if (Ar.CustomVer(FDMXRuntimeObjectVersion::GUID) < FDMXRuntimeObjectVersion::ReplaceWeakWithStrongFixturePatchReferncesInLibrarySection)
+		if (Ar.CustomVer(FDMXRuntimeMainStreamObjectVersion::GUID) < FDMXRuntimeMainStreamObjectVersion::ReplaceWeakWithStrongFixturePatchReferncesInLibrarySection)
 		{
 			// Add a library reference if possible. In cases where this is serialized before the section's library is loaded, this is not possible programmatically.
 			// This is expected to be a rare case (as it wouldn't prevent the sequence from playing before 4.27). For these cases, provide detailed steps to the user how to upgrade in logs.
diff --git a/Engine/Plugins/VirtualProduction/DMX/DMXEngine/Source/DMXRuntime/Public/Library/DMXEntityFixturePatch.h b/Engine/Plugins/VirtualProduction/DMX/DMXEngine/Source/DMXRuntime/Public/Library/DMXEntityFixturePatch.h
--- a/Engine/Plugins/VirtualProduction/DMX/DMXEngine/Source/DMXRuntime/Public/Library/DMXEntityFixturePatch.h
+++ b/Engine/Plugins/VirtualProduction/DMX/DMXEngine/Source/DMXRuntime/Public/Library/DMXEntityFixturePatch.h
@@ -27,36 +27,6 @@ struct FPropertyChangedEvent;
 
 DECLARE_MULTICAST_DELEGATE_OneParam(FDMXOnFixturePatchChangedDelegate, const UDMXEntityFixturePatch* /** ChangedFixturePatch */);
 
-
-/**
- * An instance of an MVR Fixture in a Fixture Patch.
- */
-UCLASS(DefaultToInstanced)
-class DMXRUNTIME_API UDMXMVRFixtureInstance
-	: public UObject
-{
-	GENERATED_BODY()
-
-public:
-	/** Creates a new instance from an MVR Fixture */
-	static UDMXMVRFixtureInstance* CreateFromMVRFixture(const FDMXMVRFixture& InMVRFixture);
-
-	//~ Begin UObject interface
-	virtual void Serialize(FArchive& Ar) override;
-	//~ End UObject interface
-
-	/** Sets the MVR Fixture. Note, the UUID has to match the previous MVR Fixture. */
-	void SetMVRFixture(const FDMXMVRFixture& InMVRFixture);
-
-	/** Returns the MVR Fixture */
-	FORCEINLINE const FDMXMVRFixture& GetMVRFixture() const { return MVRFixture; }
-
-private:
-	/** The mvr fixture data */
-	FDMXMVRFixture MVRFixture;
-};
-
-
 /** Parameters to construct a Fixture Patch. */
 USTRUCT(BlueprintType)
 struct DMXRUNTIME_API FDMXEntityFixturePatchConstructionParams
@@ -78,6 +48,13 @@ struct DMXRUNTIME_API FDMXEntityFixturePatchConstructionParams
 	/** Starting channel for when auto-assign address is false */
 	UPROPERTY(EditAnywhere, BlueprintReadWrite, Category = "Fixture Patch", meta = (DisplayName = "Starting Address", UIMin = "1", UIMax = "512", ClampMin = "1", ClampMax = "512"))
 	int32 StartingAddress = 1;
+
+	/** 
+	 * When spawning the DMX Library as MVR Scene in Editor, each Fixture Patch has to correspond to a Fixture in the World (if it is desired to export the Scene as MVR later).
+	 * Mostly useful when importing an MVR into the DMX Library (see DMXLibraryFromMVRFactory). If left '00000000-00000000-00000000-00000000', a Unique ID will be generated for the patch. 
+	 * Ensures the Unique ID is not used by another patch in the DMX Library already.
+	 */
+	FGuid MVRFixtureUUID;
 };
 
 
@@ -111,6 +88,7 @@ public:
 	virtual bool Modify(bool bAlwaysMarkDirty = true) override;
 #endif
 protected:
+	virtual void Serialize(FArchive& Ar) override;
 	virtual void PostInitProperties() override;
 	virtual void PostLoad() override;
 
@@ -236,24 +214,18 @@ public:
 	/** Returns custom tags defined for the patch */
 	FORCEINLINE const TArray<FName>& GetCustomTags() const { return CustomTags; }
 
-	/** Adds an MVR Fixture to the patch */
-	void AddMVRFixtureInstance(const FDMXMVRFixture& MVRFixtureToAdd);
-
-	/** Makes an array of MVR Fixtures from the MVR fixture instances of this patch */
-	TArray<FDMXMVRFixture> MakeMVRFixtures() const;
-
-	/** Returns the MVR Fixture instances that use this patch */
-	FORCEINLINE const TArray<UDMXMVRFixtureInstance*>& GetMVRFixtureInstances() const { return MVRFixtureInstances; }
+	/** Returns the MVR Fixture UUIDs of this patch */
+	FORCEINLINE const FGuid& GetMVRFixtureUUID() const { return MVRFixtureUUID; }
 
 #if WITH_EDITOR
-	/** Property name getters. When accessing the this way, use with care. The patch will need to update its cache after using property setters. Call RebuildCache to do that. */
+	/** Property name getters. When accessing the this way, use with care. The patch will need to update its cache after using property setters, use Pre/PostEditChanged events to achieve that. */
 	static FName GetUniverseIDPropertyNameChecked() { return GET_MEMBER_NAME_CHECKED(UDMXEntityFixturePatch, UniverseID); }
 	static FName GetAutoAssignAddressPropertyNameChecked() { return GET_MEMBER_NAME_CHECKED(UDMXEntityFixturePatch, bAutoAssignAddress); }
 	static FName GetManualStartingAddressPropertyNameChecked() { return GET_MEMBER_NAME_CHECKED(UDMXEntityFixturePatch, ManualStartingAddress); }
 	static FName GetAutoStartingAddressPropertyNameChecked() { return GET_MEMBER_NAME_CHECKED(UDMXEntityFixturePatch, AutoStartingAddress); }
 	static FName GetParentFixtureTypeTemplatePropertyNameChecked() { return GET_MEMBER_NAME_CHECKED(UDMXEntityFixturePatch, ParentFixtureTypeTemplate); }
 	static FName GetActiveModePropertyNameChecked() { return GET_MEMBER_NAME_CHECKED(UDMXEntityFixturePatch, ActiveMode); }
-	static FName GetMVRFixtureInstancesPropertyNameChecked() { return GET_MEMBER_NAME_CHECKED(UDMXEntityFixturePatch, MVRFixtureInstances); }
+	static FName GetMVRFixtureUUIDPropertyNameChecked() { return GET_MEMBER_NAME_CHECKED(UDMXEntityFixturePatch, MVRFixtureUUID); }
 #endif // WITH_EDITOR
 
 protected:
@@ -281,9 +253,9 @@ protected:
 	UPROPERTY(EditAnywhere, BlueprintReadOnly, Category = "Fixture Patch")
 	int32 ActiveMode;
 
-	/** MVR Fixtures that use this patch */
-	UPROPERTY(BlueprintReadWrite, EditAnywhere, Category = "Fixture Patch")
-	TArray<UDMXMVRFixtureInstance*> MVRFixtureInstances;
+	/** The MVR Fixture UUID when used as such */
+	UPROPERTY(BlueprintReadOnly, VisibleAnywhere, Category = "Fixture Patch")
+	FGuid MVRFixtureUUID;
 
 	/** Delegate broadcast when a Fixture Patch changed */
 	static FDMXOnFixturePatchChangedDelegate OnFixturePatchChangedDelegate;
diff --git a/Engine/Plugins/VirtualProduction/DMX/DMXEngine/Source/DMXRuntime/Public/Library/DMXEntityFixtureType.h b/Engine/Plugins/VirtualProduction/DMX/DMXEngine/Source/DMXRuntime/Public/Library/DMXEntityFixtureType.h
--- a/Engine/Plugins/VirtualProduction/DMX/DMXEngine/Source/DMXRuntime/Public/Library/DMXEntityFixtureType.h
+++ b/Engine/Plugins/VirtualProduction/DMX/DMXEngine/Source/DMXRuntime/Public/Library/DMXEntityFixtureType.h
@@ -11,6 +11,8 @@
 #include "DMXEntityFixtureType.generated.h"
 
 class UDMXImport;
+class UDMXImportGDTF;
+
 
 UENUM(BlueprintType)
 enum class EDMXPixelMappingDistribution : uint8
@@ -324,10 +326,15 @@ public:
 	/** Returns a delegate that is and should be broadcast whenever a Fixture Type changed */
 	static FDMXOnFixtureTypeChangedDelegate& GetOnFixtureTypeChanged();
 
+	/** DEPRECATED 5.0 */
+	UE_DEPRECATED(5.1, "Deprecated in favor of the GDTF property, which is of the GDTF type instead of just a base class.")
+	UPROPERTY(EditAnywhere, BlueprintReadOnly, Category = "Fixture Settings", Meta = (DeprecatedProperty, DeprecationMessage = "Deprecated in favor of the GDTF property, which is of the GDTF type instead of just a base class."))
+	UDMXImport* DMXImport;
+
 	/** The GDTF file from which the Fixture Type was setup */
 	UPROPERTY(EditAnywhere, BlueprintReadOnly, Category = "Fixture Settings")
-	UDMXImport* DMXImport;
-	
+	UDMXImportGDTF* GDTF;
+
 	/** The Category of the Fixture, useful for Filtering */
 	UPROPERTY(EditAnywhere, BlueprintReadOnly, Category = "Fixture Settings", meta = (DisplayName = "DMX Category"))
 	FDMXFixtureCategory DMXCategory;
@@ -354,10 +361,12 @@ private:
 public:
 	/**
 	 * Adds a Mode to the Modes Array
+	 * 
+	 * @param(optional)						The Base Mode Name when generating a name
 	 *
 	 * @return								The Index of the newly added Mode.
 	 */	
-	int32 AddMode();
+	int32 AddMode(FString BaseModeName = FString("Mode"));
 
 	/** 
 	 * Duplicates the Modes at specified Indices 
diff --git a/Engine/Plugins/VirtualProduction/DMX/DMXEngine/Source/DMXRuntime/Public/Library/DMXLibrary.h b/Engine/Plugins/VirtualProduction/DMX/DMXEngine/Source/DMXRuntime/Public/Library/DMXLibrary.h
--- a/Engine/Plugins/VirtualProduction/DMX/DMXEngine/Source/DMXRuntime/Public/Library/DMXLibrary.h
+++ b/Engine/Plugins/VirtualProduction/DMX/DMXEngine/Source/DMXRuntime/Public/Library/DMXLibrary.h
@@ -7,7 +7,6 @@
 #include "IO/DMXInputPortReference.h"
 #include "IO/DMXOutputPortReference.h"
 #include "Library/DMXEntity.h"
-#include "MVR/DMXMVRGeneralSceneDescription.h"
 
 #include "Templates/SubclassOf.h"
 #include "Misc/Guid.h"
@@ -15,6 +14,9 @@
 #include "DMXLibrary.generated.h"
 
 class UDMXEntity;
+class UDMXEntityFixturePatch;
+class UDMXEntityFixtureType;
+class UDMXMVRGeneralSceneDescription;
 
 
 /** Custom struct of in put and output port references for custom details customization with an enabled state */
@@ -48,8 +50,13 @@ class DMXRUNTIME_API UDMXLibrary
 	/** Friend DMXEntity so they can register and unregister themselves with the Library */
 	friend class UDMXEntity;
 
+public:
+	/** Constructor */
+	UDMXLibrary();
+
 protected:
 	// ~Begin UObject Interface
+	virtual void Serialize(FArchive& Ar) override;
 	virtual void PostInitProperties() override;
 	virtual void PostLoad() override;
 	virtual void PostDuplicate(EDuplicateMode::Type DuplicateMode) override;
@@ -201,18 +208,30 @@ public:
 	/** Updates the ports from what's set in the Input and Output Port References arrays */
 	void UpdatePorts();
 
-	/** Creates an MVR General Scene Description from the Library's MVR data */
-	FDMXMVRGeneralSceneDescription MakeGeneralSceneDescription() const;
+	/** 
+	 * Sets the MVR General Scene Description of the Library. 
+	 * Note, this will not add any patches occur in the new General Scene Description. 
+	 * Note, this will remove any patches that do not occur in the new General Scene Description. 
+	 */
+	void SetMVRGeneralSceneDescription(UDMXMVRGeneralSceneDescription* NewGeneralSceneDescription);
+
+	/** Updates the General Scene Description to reflect the Library */
+	UDMXMVRGeneralSceneDescription* UpdateGeneralSceneDescription();
+
+	/** 
+	 * Returns the General Scene Description of the Library. 
+	 * May not always represent the state of the library. 
+	 * Use WriteGeneralSceneDescription to update it to get an exact version.
+	 */
+	FORCEINLINE UDMXMVRGeneralSceneDescription* GetLazyGeneralSceneDescription() const { return GeneralSceneDescription; }
 
 	/** Returns the Entity that was last added to the Library */
 	FORCEINLINE TWeakObjectPtr<UDMXEntity> GetLastAddedEntity() const { return LastAddedEntity; }
 
 #if WITH_EDITOR
-	/** Returns the name of the Ports property. */
 	static FName GetPortReferencesPropertyName() { return GET_MEMBER_NAME_CHECKED(UDMXLibrary, PortReferences); }
-
-	/** Returns the name of the Entities property. */
 	static FName GetEntitiesPropertyName() { return GET_MEMBER_NAME_CHECKED(UDMXLibrary, Entities); }
+	static FName GetGeneralSceneDescriptionPropertyName() { return GET_MEMBER_NAME_CHECKED(UDMXLibrary, GeneralSceneDescription); }
 #endif // WITH_EDITOR
 
 protected:
@@ -229,10 +248,14 @@ protected:
 	FDMXLibraryPortReferences PortReferences;
 
 private:
-	/** All Fixture Types and Fixture Patches in the library */
+	/** All Fixture Types and Fixture Patches in the Library */
 	UPROPERTY()
 	TArray<UDMXEntity*> Entities;
 
+	/** The General Scene Description of this Library */
+	UPROPERTY()
+	UDMXMVRGeneralSceneDescription* GeneralSceneDescription;
+
 	/** The entity that was added last to the Library */
 	TWeakObjectPtr<UDMXEntity> LastAddedEntity;
 
diff --git a/Engine/Plugins/VirtualProduction/DMX/DMXEngine/Source/DMXRuntime/Public/MVR/DMXMVRAssetUserData.h b/Engine/Plugins/VirtualProduction/DMX/DMXEngine/Source/DMXRuntime/Public/MVR/DMXMVRAssetUserData.h
--- a/Engine/Plugins/VirtualProduction/DMX/DMXEngine/Source/DMXRuntime/Public/MVR/DMXMVRAssetUserData.h
+++ b/Engine/Plugins/VirtualProduction/DMX/DMXEngine/Source/DMXRuntime/Public/MVR/DMXMVRAssetUserData.h
@@ -17,18 +17,21 @@ class DMXRUNTIME_API UDMXMVRAssetUserData
 
 public:
 	/** Returns MVR Asset User Data for the Actor, or nullptr if there is no asset user data */
-	static UDMXMVRAssetUserData* GetMVRAssetUserData(AActor* Actor);
+	static UDMXMVRAssetUserData* GetMVRAssetUserData(const AActor& Actor);
 
 	/** Returns the Value in Meta data given a Key, or an empty String if the Key doesn't exist */
-	static FString GetMVRAssetUserDataValueForkey(AActor* Actor, const FName Key);
+	static FString GetMVRAssetUserDataValueForKey(const AActor& Actor, const FName Key);
 
 	/** Sets the Value for Key in Meta Data. Creates the Key if it doesn't exist. */
-	static bool SetMVRAssetUserDataValueForKey(AActor* Actor, FName Key, const FString& Value);
+	static bool SetMVRAssetUserDataValueForKey(AActor& InOutActor, FName InKey, const FString& InValue);
 
 	/** MVR Meta Data for the Actor */
 	UPROPERTY(VisibleAnywhere, BlueprintReadWrite, AdvancedDisplay, Category = "MVR")
 	TMap<FName, FString> MetaData;
 
 	/** The key for the MVR UUID meta data */
-	static const FName MVRUUIDMetaDataKey;
+	static const FName MVRFixtureUUIDMetaDataKey;
+
+	/** The key for the Fixture Patch meta data */
+	static const FName FixturePatchMetaDataKey;
 };
diff --git a/Engine/Plugins/VirtualProduction/DMX/DMXEngine/Source/DMXRuntime/Public/MVR/DMXMVRFixture.h b/Engine/Plugins/VirtualProduction/DMX/DMXEngine/Source/DMXRuntime/Public/MVR/DMXMVRFixture.h
new file mode 100644
--- /dev/null
+++ b/Engine/Plugins/VirtualProduction/DMX/DMXEngine/Source/DMXRuntime/Public/MVR/DMXMVRFixture.h
@@ -0,0 +1,192 @@
+// Copyright Epic Games, Inc. All Rights Reserved.
+
+#pragma once
+
+#include "CoreMinimal.h"
+
+#include "Misc/Optional.h"
+#include "Serialization/Archive.h"
+
+class FXmlNode;
+class FXmlFile;
+
+
+/** MVR Fixture Addresses as per Standard */
+struct DMXRUNTIME_API FDMXMVRFixtureAddresses
+{
+	FDMXMVRFixtureAddresses()
+	{}
+
+	FDMXMVRFixtureAddresses(int32 InUniverse, int32 InAddress)
+		: Universe(InUniverse)
+		, Address(InAddress)
+	{}
+
+	int32 Universe;
+	int32 Address;
+
+	/** Serializes an MVR Fixture Address from or into an archive */
+	friend FArchive& operator<<(FArchive& Ar, FDMXMVRFixtureAddresses& Addresses)
+	{
+		Ar << Addresses.Universe;
+		Ar << Addresses.Address;
+
+		return Ar;
+	}
+};
+
+/** MVR Fixture Color CIE as per Standard */
+struct DMXRUNTIME_API FDMXMVRColorCIE
+{
+	float X;
+	float Y;
+	uint8 YY;
+
+	/** Serializes an MVR Color CIE from or into an archive */
+	friend FArchive& operator<<(FArchive& Ar, FDMXMVRColorCIE& ColorCIE)
+	{
+		Ar << ColorCIE.X;
+		Ar << ColorCIE.Y;
+		Ar << ColorCIE.YY;
+
+		return Ar;
+	}
+};
+
+/** MVR Fixture Fixture Mapping as per Standard */
+struct DMXRUNTIME_API FDMXMVRFixtureMapping
+{
+	/** The unique identifier of the MappingDefinition node that will be the source of the mapping. */
+	FGuid LinkDef;
+
+	/** The offset in pixels in x direction from top left corner of the source that will be used for the mapped object. */
+	TOptional<int32> UX;
+
+	/** The offset in pixels in y direction from top left corner of the source that will be used for the mapped object. */
+	TOptional<int32> UY;
+
+	/** The size in pixels in x direction from top left of the starting point. */
+	TOptional<int32> OX;
+
+	/** The size in pixels in y direction from top left of the starting point. */
+	TOptional<int32> OY;
+
+	/** The rotation around the middle point of the defined rectangle in degree. Positive direction is counter cock wise. */
+	TOptional<int32> RZ;
+
+	/** Serializes an MVR Fixture Mapping from or into an archive */
+	friend FArchive& operator<<(FArchive& Ar, FDMXMVRFixtureMapping& Mapping)
+	{
+
+
+		Ar << Mapping.LinkDef;
+		Ar << Mapping.UX;
+		Ar << Mapping.UY;
+		Ar << Mapping.OX;
+		Ar << Mapping.OY;
+		Ar << Mapping.RZ;
+
+		return Ar;
+	}
+};
+
+/** MVR Fixture Fixture Gobo as per Standard */
+struct DMXRUNTIME_API FDMXMVRFixtureGobo
+{
+	/** The node value is the Gobo used for the fixture. The image ressource must apply to the GDTF standard. Use a FileName to specify. */
+	FString Value;
+
+	/** The roation of the Gobo in degree. */
+	float Rotation = 0.f;
+
+	/** Serializes an MVR Fixture Gobo from or into an archive */
+	friend FArchive& operator<<(FArchive& Ar, FDMXMVRFixtureGobo& Gobo)
+	{
+		Ar << Gobo.Value;
+		Ar << Gobo.Rotation;
+
+		return Ar;
+	}
+};
+
+/** MVR Fixture Fixture as per Standard */
+struct DMXRUNTIME_API FDMXMVRFixture
+{
+public:
+	FDMXMVRFixture();
+	FDMXMVRFixture(const FXmlNode* const FixtureNode);
+
+	/** Returns true if this is a valid MVR Fixture */
+	bool IsValid() const;
+
+	/** The unique identifier of the object. */
+	FGuid UUID;
+
+	/** The name of the object. */
+	FString Name;
+
+	/** The location of the object inside the parent coordinate system. */
+	TOptional<FTransform> Transform;
+	
+	/** The name of the file containing the GDTF information for this light fixture. */
+	FString GDTFSpec;
+
+	/** The name of the used DMX mode. This has to match the name of a DMXMode in the GDTF file. */
+	FString GDTFMode;
+
+	/** A focus point reference that this lighting fixture aims at if this reference exists. */
+	TOptional<FGuid> Focus;
+
+	/** Defines if a Object cast Shadows. */
+	TOptional<bool> bCastShadows;
+
+	/** A position reference that this lighting fixture belongs to if this reference exists. */
+	TOptional<FGuid> Position;
+
+	/** The Fixture Id of the lighting fixture. This is the short name of the fixture. */
+	FString FixtureId;
+
+	/** The unit number of the lighting fixture in a position. */
+	int32 UnitNumber = 0;
+
+	/** The container for DMX Addresses for this fixture. */
+	FDMXMVRFixtureAddresses Addresses;
+
+	/** A color assigned to a fixture. If it is not defined, there is no color for the fixture. */
+	TOptional<FDMXMVRColorCIE> CIEColor;
+
+	/** The Fixture Type ID is a value that can be used as a short name of the Fixture Type. This does not have to be unique. The default value is 0. */
+	TOptional<int32> FixtureTypeId;
+
+	/** The Custom ID is a value that can be used as a short name of the Fixture Instance. This does not have to be unique. The default value is 0. */
+	TOptional<int32> CustomId;
+
+	/** The container for Mappings for this fixture. */
+	TOptional<FDMXMVRFixtureMapping> Mapping;
+
+	/** The Gobo used for the fixture. The image ressource must apply to the GDTF standard. */
+	TOptional<FDMXMVRFixtureGobo> Gobo;
+
+	/** Serializes an MVR Fixture from or into an archive */
+	friend FArchive& operator<<(FArchive& Ar, FDMXMVRFixture& MVRFixture)
+	{
+		Ar << MVRFixture.UUID;
+		Ar << MVRFixture.Name;
+		Ar << MVRFixture.Transform;
+		Ar << MVRFixture.GDTFSpec;
+		Ar << MVRFixture.GDTFMode;
+		Ar << MVRFixture.Focus;
+		Ar << MVRFixture.bCastShadows;
+		Ar << MVRFixture.Position;
+		Ar << MVRFixture.FixtureId;
+		Ar << MVRFixture.UnitNumber;	
+		Ar << MVRFixture.Addresses;
+		Ar << MVRFixture.CIEColor;
+		Ar << MVRFixture.FixtureTypeId;
+		Ar << MVRFixture.CustomId;
+		Ar << MVRFixture.Mapping;
+		Ar << MVRFixture.Gobo;
+		
+		return Ar;		 
+	}
+};
diff --git a/Engine/Plugins/VirtualProduction/DMX/DMXEngine/Source/DMXRuntime/Public/MVR/DMXMVRFixtureActorInterface.h b/Engine/Plugins/VirtualProduction/DMX/DMXEngine/Source/DMXRuntime/Public/MVR/DMXMVRFixtureActorInterface.h
--- a/Engine/Plugins/VirtualProduction/DMX/DMXEngine/Source/DMXRuntime/Public/MVR/DMXMVRFixtureActorInterface.h
+++ b/Engine/Plugins/VirtualProduction/DMX/DMXEngine/Source/DMXRuntime/Public/MVR/DMXMVRFixtureActorInterface.h
@@ -29,7 +29,11 @@ public:
 	UFUNCTION(BlueprintNativeEvent, BlueprintCallable, Category = "DMX", Meta = (DisplayName = "On MVR Get Supported DMX Attributes"))
 	void OnMVRGetSupportedDMXAttributes(TArray<FName>& OutAttributeNames, TArray<FName>& OutMatrixAttributeNames) const;
 
-	/** Should set the Fixture Patch the MVR Fixture uses */
-	UFUNCTION(BlueprintNativeEvent, BlueprintCallable, CallInEditor, Category = "DMX", Meta = (DisplayName = "On MVR Set Fixture Patch"))
-	void OnMVRSetFixturePatch(UDMXEntityFixturePatch* FixturePatch);
+	/** Should return the Fixture Patch the MVR Fixture uses */
+	UFUNCTION(BlueprintNativeEvent, CallInEditor, Category = "DMX", Meta = (DisplayName = "On MVR Set Fixture Patch"))
+	UDMXEntityFixturePatch* OnMVRGetFixturePatch() const;
+
+	/** Should SET the Fixture Patch the MVR Fixture uses */
+	UFUNCTION(BlueprintNativeEvent, CallInEditor, Category = "DMX", Meta = (DisplayName = "On MVR Get Fixture Patch"))
+	void OnMVRSetFixturePatch(UDMXEntityFixturePatch* FixturePatchInUse);
 };
diff --git a/Engine/Plugins/VirtualProduction/DMX/DMXEngine/Source/DMXRuntime/Public/MVR/DMXMVRGeneralSceneDescription.h b/Engine/Plugins/VirtualProduction/DMX/DMXEngine/Source/DMXRuntime/Public/MVR/DMXMVRGeneralSceneDescription.h
--- a/Engine/Plugins/VirtualProduction/DMX/DMXEngine/Source/DMXRuntime/Public/MVR/DMXMVRGeneralSceneDescription.h
+++ b/Engine/Plugins/VirtualProduction/DMX/DMXEngine/Source/DMXRuntime/Public/MVR/DMXMVRGeneralSceneDescription.h
@@ -4,196 +4,63 @@
 
 #include "CoreMinimal.h"
 
-#include "Library/DMXImport.h"
+#include "MVR/DMXMVRFixture.h"
+
 #include "Misc/Optional.h"
 #include "Serialization/Archive.h"
 
-class FXmlNode;
-class FXmlFile;
-
+#include "DMXMVRGeneralSceneDescription.generated.h"
 
-struct DMXRUNTIME_API FDMXMVRFixtureAddresses
-{
-	int32 Address;
-	int32 Universe;
+class UDMXEntityFixturePatch;
+class UDMXLibrary;
 
-	/** Serializes an MVR Fixture Address from or into an archive */
-	friend FArchive& operator<<(FArchive& Ar, FDMXMVRFixtureAddresses& Addresses)
-	{
-		Ar << Addresses.Address;
-		Ar << Addresses.Universe;
+class FXmlFile;
 
-		return Ar;
-	}
-};
 
-struct DMXRUNTIME_API FDMXMVRColorCIE
+/** MVR General Scene Description Object */
+UCLASS()
+class DMXRUNTIME_API UDMXMVRGeneralSceneDescription
+	: public UObject
 {
-	float X;
-	float Y;
-	uint8 YY;
-
-	/** Serializes an MVR Color CIE from or into an archive */
-	friend FArchive& operator<<(FArchive& Ar, FDMXMVRColorCIE& ColorCIE)
-	{
-		Ar << ColorCIE.X;
-		Ar << ColorCIE.Y;
-		Ar << ColorCIE.YY;
-
-		return Ar;
-	}
-};
+	GENERATED_BODY()
 
-struct DMXRUNTIME_API FDMXMVRFixtureMapping
-{
-	/** The unique identifier of the MappingDefinition node that will be the source of the mapping. */
-	FGuid LinkDef;
-
-	/** The offset in pixels in x direction from top left corner of the source that will be used for the mapped object. */
-	TOptional<int32> UX;
-
-	/** The offset in pixels in y direction from top left corner of the source that will be used for the mapped object. */
-	TOptional<int32> UY;
+public:
+	/** Creates an MVR General Scene Description from an Xml File */
+	static UDMXMVRGeneralSceneDescription* CreateFromXmlFile(TSharedRef<FXmlFile> GeneralSceneDescriptionXml, UObject* Outer, FName Name, EObjectFlags Flags = RF_NoFlags);
 
-	/** The size in pixels in x direction from top left of the starting point. */
-	TOptional<int32> OX;
+	/** Creates an MVR General Scene Description from a DMX Library */
+	static UDMXMVRGeneralSceneDescription* CreateFromDMXLibrary(const UDMXLibrary& DMXLibrary, UObject* Outer, FName Name, EObjectFlags Flags = RF_NoFlags);
 
-	/** The size in pixels in y direction from top left of the starting point. */
-	TOptional<int32> OY;
+	/** Returns a pointer to the MVR Fixture with corresponding UUID or nullptr if not found */
+	FDMXMVRFixture* FindMVRFixture(const FGuid& MVRFixtureUUID);
 
-	/** The rotation around the middle point of the defined rectangle in degree. Positive direction is counter cock wise. */
-	TOptional<int32> RZ;
+	/** Adds an MVR Fixture to the General Scene Description */
+	void AddMVRFixture(FDMXMVRFixture& MVRFixture);
 
-	/** Serializes an MVR Fixture Mapping from or into an archive */
-	friend FArchive& operator<<(FArchive& Ar, FDMXMVRFixtureMapping& Mapping)
-	{
-		Ar << Mapping.LinkDef;
-		Ar << Mapping.UX;
-		Ar << Mapping.UY;
-		Ar << Mapping.OX;
-		Ar << Mapping.OY;
-		Ar << Mapping.RZ;
+	/**
+	 * Writes the Library to the General Scene Description, effectively removing inexisting and adding
+	 * new MVR Fixtures, according to what MVR Fixture UUIDs the Fixture Patches of the Library contain.
+	 */
+	void WriteDMXLibraryToGeneralSceneDescription(const UDMXLibrary& DMXLibrary);
 
-		return Ar;
-	}
-};
+	/** Returns the MVR Fixtures of the General Scene Description */
+	FORCEINLINE const TArray<FDMXMVRFixture>& GetMVRFixtures() const { return MVRFixtures; }
 
-struct DMXRUNTIME_API FDMXMVRFixtureGobo
-{
-	/** The node value is the Gobo used for the fixture. The image ressource must apply to the GDTF standard. Use a FileName to specify. */
-	FString Value;
+protected:
+	//~ Begin UObject interface
+	virtual void Serialize(FArchive& Ar) override;
+	//~ End UObject interface
 
-	/** The roation of the Gobo in degree. */
-	float Rotation = 0.f;
-
-	/** Serializes an MVR Fixture Gobo from or into an archive */
-	friend FArchive& operator<<(FArchive& Ar, FDMXMVRFixtureGobo& Gobo)
-	{
-		Ar << Gobo.Value;
-		Ar << Gobo.Rotation;
-
-		return Ar;
-	}
-};
+private:
+	/** Parses a General Scene Description Xml File. Only ment to be used for initialization (ensured) */
+	void ParseGeneralSceneDescriptionXml(const TSharedRef<FXmlFile>& GeneralSceneDescriptionXml);
 
-struct DMXRUNTIME_API FDMXMVRFixture
-{
-public:
-	FDMXMVRFixture();
-	FDMXMVRFixture(const FXmlNode* const FixtureNode);
-
-	/** Returns true if this is a valid MVR Fixture */
-	bool IsValid() const;
-
-	/** The unique identifier of the object. */
-	FGuid UUID;
-
-	/** The name of the object. */
-	FString Name;
-
-	/** The location of the object inside the parent coordinate system. */
-	TOptional<FTransform> Transform;
-	
-	/** The name of the file containing the GDTF information for this light fixture. */
-	FString GDTFSpec;
-
-	/** The name of the used DMX mode. This has to match the name of a DMXMode in the GDTF file. */
-	FString GDTFMode;
-
-	/** A focus point reference that this lighting fixture aims at if this reference exists. */
-	TOptional<FGuid> Focus;
-
-	/** Defines if a Object cast Shadows. */
-	TOptional<bool> bCastShadows;
-
-	/** A position reference that this lighting fixture belongs to if this reference exists. */
-	TOptional<FGuid> Position;
-
-	/** The Fixture Id of the lighting fixture. This is the short name of the fixture. */
-	FString FixtureId;
-
-	/** The unit number of the lighting fixture in a position. */
-	int32 UnitNumber = 0;
-
-	/** The container for DMX Addresses for this fixture. */
-	FDMXMVRFixtureAddresses Addresses;
-
-	/** A color assigned to a fixture. If it is not defined, there is no color for the fixture. */
-	TOptional<FDMXMVRColorCIE> CIEColor;
-
-	/** The Fixture Type ID is a value that can be used as a short name of the Fixture Type. This does not have to be unique. The default value is 0. */
-	TOptional<int32> FixtureTypeId;
-
-	/** The Custom ID is a value that can be used as a short name of the Fixture Instance. This does not have to be unique. The default value is 0. */
-	TOptional<int32> CustomId;
-
-	/** The container for Mappings for this fixture. */
-	TOptional<FDMXMVRFixtureMapping> Mapping;
-
-	/** The Gobo used for the fixture. The image ressource must apply to the GDTF standard. */
-	TOptional<FDMXMVRFixtureGobo> Gobo;
-
-	/** Serializes an MVR Fixture from or into an archive */
-	friend FArchive& operator<<(FArchive& Ar, FDMXMVRFixture& MVRFixture)
-	{
-		Ar << MVRFixture.UUID;
-		Ar << MVRFixture.Name;
-		Ar << MVRFixture.Transform;
-		Ar << MVRFixture.GDTFSpec;
-		Ar << MVRFixture.GDTFMode;
-		Ar << MVRFixture.Focus;
-		Ar << MVRFixture.bCastShadows;
-		Ar << MVRFixture.Position;
-		Ar << MVRFixture.FixtureId;
-		Ar << MVRFixture.UnitNumber;	
-		Ar << MVRFixture.Addresses;
-		Ar << MVRFixture.CIEColor;
-		Ar << MVRFixture.FixtureTypeId;
-		Ar << MVRFixture.CustomId;
-		Ar << MVRFixture.Mapping;
-		Ar << MVRFixture.Gobo;
-		
-		return Ar;		 
-	}
-};
+	/** Writes the Fixture Patch to the General Scene Description, adding a new MVR Fixture if required */
+	void WriteFixturePatchToGeneralSceneDescription(const UDMXEntityFixturePatch& FixturePatch);
 
-struct DMXRUNTIME_API FDMXMVRGeneralSceneDescription
-{
-	FDMXMVRGeneralSceneDescription();
-	FDMXMVRGeneralSceneDescription(TSharedRef<FXmlFile> GeneralSceneDescription);
+	/** Returns the Unit Numbers currently in use, sorted from lowest to highest */
+	TArray<int32> GetUnitNumbersInUse(const UDMXLibrary& DMXLibrary);
 
-	/** The DMMX interactable fixture contained in an MVR file */
+	/** The MVR Fixtures of the General Scene Description */
 	TArray<FDMXMVRFixture> MVRFixtures;
-
-	/** Serializes an General Scene Description from or into an archive */
-	friend FArchive& operator<<(FArchive& Ar, FDMXMVRGeneralSceneDescription& GeneralSceneDescription)
-	{
-		Ar << GeneralSceneDescription.MVRFixtures;
-
-		return Ar;
-	}
-
-private:
-	/** Gets the fixture nodes from the General Scene Description */
-	void ParseMVRFixtures(const TSharedRef<FXmlFile>& GeneralSceneDescription);
 };
diff --git a/Engine/Plugins/VirtualProduction/DMX/DMXEngine/Source/DMXRuntime/Public/MVR/DMXMVRSceneActor.h b/Engine/Plugins/VirtualProduction/DMX/DMXEngine/Source/DMXRuntime/Public/MVR/DMXMVRSceneActor.h
--- a/Engine/Plugins/VirtualProduction/DMX/DMXEngine/Source/DMXRuntime/Public/MVR/DMXMVRSceneActor.h
+++ b/Engine/Plugins/VirtualProduction/DMX/DMXEngine/Source/DMXRuntime/Public/MVR/DMXMVRSceneActor.h
@@ -7,6 +7,7 @@
 
 #include "DMXMVRSceneActor.generated.h"
 
+class UDMXEntityFixturePatch;
 class UDMXImportGDTF;
 class UDMXLibrary;
 class UDMXMVRSceneComponent;
@@ -15,21 +16,19 @@ class UFactory;
 
 
 USTRUCT(BlueprintType)
-struct FDMXMVRSceneActorGDTFToActorClassPair
+struct DMXRUNTIME_API FDMXMVRSceneGDTFToActorClassPair
 {
 	GENERATED_BODY()
 
-#if WITH_EDITORONLY_DATA
-	/** The GDTF Asset */
-	UPROPERTY(EditAnywhere, Category = "MVR")
+	/** The GDTF of the Actor (may not correlate with its Patch after the Patch changed) */
+	UPROPERTY(VisibleAnywhere, Category = "MVR")
 	TSoftObjectPtr<UDMXImportGDTF> GDTF;
-
-	UPROPERTY(EditAnywhere, Category = "MVR", Meta = (MustImplement = "DMXMVRFixtureActorInterface"))
-	TSoftClassPtr<AActor> Actor;
-#endif // WITH_EDITORONLY_DATA
+	
+	/** The Actor Class that should be or was spawned */
+	UPROPERTY(EditAnywhere, BlueprintReadOnly, Category = MVR, NoClear, Meta = (MustImplement = "DMXMVRFixtureActorInterface"))
+	TSoftClassPtr<AActor> ActorClass;
 };
 
-
 UCLASS(NotBlueprintable)
 class DMXRUNTIME_API ADMXMVRSceneActor
 	: public AActor
@@ -46,45 +45,79 @@ public:
 	//~ Begin AActor interface
 	virtual void PostLoad() override;
 	virtual void PostRegisterAllComponents() override;
+#if WITH_EDITOR
+	virtual void PreEditChange(FProperty* PropertyAboutToChange) override;
+	virtual void PostEditChangeProperty(FPropertyChangedEvent& PropertyChangedEvent) override;
+#endif // WITH_EDITOR
 	//~ End AActor interface
 
+public:
 	/** Sets the dmx library for this MVR actor. Should only be called once, further calls will have no effect and hit an ensure condition */
 	void SetDMXLibrary(UDMXLibrary* NewDMXLibrary);
 
 	/** Returns the DMX Library of this MVR Scene Actor */
 	FORCEINLINE UDMXLibrary* GetDMXLibrary() const { return DMXLibrary; }
 
-#if WITH_EDITORONLY_DATA
-	/** The actor class that is spawned for a specific GDTF by default (can be overriden per MVR UUID, see below) */
-	UPROPERTY(EditAnywhere, Category = "MVR", Meta = (DispayName = "Default Actor used for GDTF"))
-	TArray<FDMXMVRSceneActorGDTFToActorClassPair> GDTFToDefaultActorClassArray;
-#endif // WITH_EDITORONLY_DATA
+	/** Returns the MVR UUID To Related Actor Map */
+	FORCEINLINE const TArray<TSoftObjectPtr<AActor>>& GetRelatedActors() const { return RelatedActors; };
+
+#if WITH_EDITOR
+	// Property name getters
+	static FName GetDMXLibraryPropertyNameChecked() { return GET_MEMBER_NAME_CHECKED(ADMXMVRSceneActor, DMXLibrary); }
+	static FName GetRelatedAcctorsPropertyNameChecked() { return GET_MEMBER_NAME_CHECKED(ADMXMVRSceneActor, RelatedActors); }
+	static FName GetGDTFToDefaultActorClassesPropertyNameChecked() { return GET_MEMBER_NAME_CHECKED(ADMXMVRSceneActor, GDTFToDefaultActorClasses); }
+	static FName GetMVRSceneRootPropertyNameChecked() { return GET_MEMBER_NAME_CHECKED(ADMXMVRSceneActor, MVRSceneRoot); }
+#endif // WITH_EDITOR
 
 private:
 	/** Set MVR UUIDs for related Actors */
-	void SetMVRUUIDsForRelatedActors();
+	void EnsureMVRUUIDsForRelatedActors();
 
 #if WITH_EDITOR
 	/** Called when a sub-level is loaded */
 	void OnMapChange(uint32 MapChangeFlags);
 
 	/** Called when an actor got deleted in editor */
-	void OnActorDeleted(AActor* ActorDeleted);
+	void OnActorDeleted(AActor* DeletedActor);
 
 	/** Called when an asset was imported */
 	void OnAssetPostImport(UFactory* InFactory, UObject* ActorAdded);
+
+	/** Replaces related Actors that use the Default Actor Class for the GDTF with an instance of the new Default Actor Class */
+	void HandleDefaultActorClassForGDTFChanged();
 #endif // WITH_EDITOR
 
-private:
+	/** Spawns an MVR Actor in this Scene. Returns the newly spawned Actor, or nullptr if no Actor could be spawned. */
+	AActor* SpawnMVRActor(const TSubclassOf<AActor>&ActorClass, UDMXEntityFixturePatch * FixturePatch, const FTransform & Transform, AActor * Template = nullptr);
+
+	/** Replaces an MVR Actor in this Scene with another. Returns the newly spawned Actor, or nullptr if no Actor could be spawned. */
+	AActor* ReplaceMVRActor(AActor* ActorToReplace, const TSubclassOf<AActor>& ClassOfNewActor);
+
 	/** The DMX Library this Scene Actor uses */
 	UPROPERTY(BlueprintReadOnly, VisibleAnywhere, Category = "MVR", Meta = (AllowPrivateAccess = true))
 	TObjectPtr<UDMXLibrary> DMXLibrary;
 
 	/** The actors that created along with this scene */
 	UPROPERTY(VisibleAnywhere, Category = "Actor", AdvancedDisplay, Meta = (AllowPrivateAccess = true))
-	TMap<FGuid, TSoftObjectPtr<AActor>> MVRUUIDToRelatedActorMap;
+	TArray<TSoftObjectPtr<AActor>> RelatedActors;
+	
+	/** The actor class that is spawned for a specific GDTF by default (can be overriden per MVR UUID, see below) */
+	UPROPERTY(EditAnywhere, Category = "MVR", Meta = (DispayName = "Default Actor used for GDTF"))
+	TArray<FDMXMVRSceneGDTFToActorClassPair> GDTFToDefaultActorClasses;
+
+#if WITH_EDITORONLY_DATA
+	/** The GDTFToDefaultActorClassMap cached of PreEditChange, to find changes. */
+	TArray<FDMXMVRSceneGDTFToActorClassPair> GDTFToDefaultActorClasses_PreEditChange;
+#endif // WITH_EDITORONLY_DATA
 
 	/** The root component to which all actors are attached initially */
 	UPROPERTY(VisibleAnywhere, Category = "Actor", AdvancedDisplay, Meta = (AllowPrivateAccess = true))
 	USceneComponent* MVRSceneRoot;
+
+private:
+	/** Gets the Fixture Patch of an Actor, or nullptr if the Fixture Patch cannot be retrieved anymore */
+	UDMXEntityFixturePatch* GetFixturePatch(AActor* Actor) const;
+
+	/** Sets the Fixture Patch on an Actor if possible or fails silently */
+	void SetFixturePatch(AActor* Actor, UDMXEntityFixturePatch* FixturePatch);
 };
diff --git a/Engine/Plugins/VirtualProduction/DMX/DMXFixtures/Source/DMXFixtures/Private/DMXFixtureActor.cpp b/Engine/Plugins/VirtualProduction/DMX/DMXFixtures/Source/DMXFixtures/Private/DMXFixtureActor.cpp
--- a/Engine/Plugins/VirtualProduction/DMX/DMXFixtures/Source/DMXFixtures/Private/DMXFixtureActor.cpp
+++ b/Engine/Plugins/VirtualProduction/DMX/DMXFixtures/Source/DMXFixtures/Private/DMXFixtureActor.cpp
@@ -69,9 +69,16 @@ void ADMXFixtureActor::OnMVRGetSupportedDMXAttributes_Implementation(TArray<FNam
 
 void ADMXFixtureActor::OnMVRSetFixturePatch_Implementation(UDMXEntityFixturePatch* FixturePatch)
 {
+	check(DMX);
 	DMX->SetFixturePatch(FixturePatch);
 }
 
+UDMXEntityFixturePatch* ADMXFixtureActor::OnMVRGetFixturePatch_Implementation() const
+{
+	check(DMX);
+	return DMX->GetFixturePatch();
+}
+
 #if WITH_EDITOR
 void ADMXFixtureActor::PostEditChangeProperty(FPropertyChangedEvent& PropertyChangedEvent)
 {
diff --git a/Engine/Plugins/VirtualProduction/DMX/DMXFixtures/Source/DMXFixtures/Public/DMXFixtureActor.h b/Engine/Plugins/VirtualProduction/DMX/DMXFixtures/Source/DMXFixtures/Public/DMXFixtureActor.h
--- a/Engine/Plugins/VirtualProduction/DMX/DMXFixtures/Source/DMXFixtures/Public/DMXFixtureActor.h
+++ b/Engine/Plugins/VirtualProduction/DMX/DMXFixtures/Source/DMXFixtures/Public/DMXFixtureActor.h
@@ -50,6 +50,7 @@ public:
 	//~ Begin DMXMVRFixtureActorInterface interface
 	virtual void OnMVRGetSupportedDMXAttributes_Implementation(TArray<FName>& OutAttributeNames, TArray<FName>& OutMatrixAttributeNames) const override;
 	virtual void OnMVRSetFixturePatch_Implementation(UDMXEntityFixturePatch* FixturePatch) override;
+	virtual UDMXEntityFixturePatch* OnMVRGetFixturePatch_Implementation() const override;
 	//~ End DMXMVRFixtureActorInterface interface
 
 	bool HasBeenInitialized;
diff --git a/Engine/Plugins/VirtualProduction/DMX/DatasmithMVR/Source/DatasmithMVRTranslator/Private/DatasmithMVRNativeTranslator.cpp b/Engine/Plugins/VirtualProduction/DMX/DatasmithMVR/Source/DatasmithMVRTranslator/Private/DatasmithMVRNativeTranslator.cpp
--- a/Engine/Plugins/VirtualProduction/DMX/DatasmithMVR/Source/DatasmithMVRTranslator/Private/DatasmithMVRNativeTranslator.cpp
+++ b/Engine/Plugins/VirtualProduction/DMX/DatasmithMVR/Source/DatasmithMVRTranslator/Private/DatasmithMVRNativeTranslator.cpp
@@ -115,7 +115,7 @@ UDMXLibrary* FDatasmithMVRNativeTranslator::CreateDMXLibraryFromMVR(const FStrin
 
 	const FString SceneName = GetSource().GetSceneName();
 	const FString DMXLibraryName = FPaths::GetBaseFilename(MVRFilePathAndName);
-	const FString PackageName = FPaths::ProjectContentDir() / SceneName;
+	const FString PackageName = TEXT("/Game") / SceneName;
 	UPackage* DMXLibraryAssetPackage = CreatePackage(*PackageName);
 	DMXLibraryAssetPackage->FullyLoad();
 	
@@ -172,59 +172,61 @@ void FDatasmithMVRNativeTranslator::ReplaceMVRActorsWithMVRSceneActor(TSharedRef
 	TArray<TSharedRef<IDatasmithActorElement>> ActorElementsToRemove;
 
 	const int32 NumMetaDatas = InOutScene->GetMetaDataCount();
-	TArray<UDMXEntityFixturePatch*> FixturePatches = DMXLibrary->GetEntitiesTypeCast<UDMXEntityFixturePatch>();
-	for (UDMXEntityFixturePatch* FixturePatch : FixturePatches)
+
+	DMXLibrary->UpdateGeneralSceneDescription();
+	UDMXMVRGeneralSceneDescription* GeneralSceneDescription = DMXLibrary->GetLazyGeneralSceneDescription();
+	if (!ensureAlwaysMsgf(GeneralSceneDescription, TEXT("Trying to remove MVR scene elements from datasmith scene, but the DMX Library that was created for it does not provide a General Scene Description.")))
+	{
+		return;
+	}
+
+	for (const FDMXMVRFixture& MVRFixture : GeneralSceneDescription->GetMVRFixtures())
 	{
-		const TArray<UDMXMVRFixtureInstance*>& MVRFixtureInstances = FixturePatch->GetMVRFixtureInstances();
-		for (UDMXMVRFixtureInstance* MVRFixtureInstance : MVRFixtureInstances)
+		for (int32 MetaDataIndex = 0; MetaDataIndex < NumMetaDatas; MetaDataIndex++)
 		{
-			const FDMXMVRFixture& MVRFixture = MVRFixtureInstance->GetMVRFixture();
-			for (int32 MetaDataIndex = 0; MetaDataIndex < NumMetaDatas; MetaDataIndex++)
-			{
-				const TSharedPtr<IDatasmithMetaDataElement> MetaDataElement = InOutScene->GetMetaData(MetaDataIndex);
+			const TSharedPtr<IDatasmithMetaDataElement> MetaDataElement = InOutScene->GetMetaData(MetaDataIndex);
 
-				// Try to find an MVR ID in MetaData. Note Vectorworks uses VW_UUID as that, so we search for that too.
-				static const FString Generic_MVRUUIDName = TEXT("MVR_UUID");
-				static const FString VW_MVRUUIDName = TEXT("VW_UUID");
+			// Try to find an MVR ID in MetaData. Note Vectorworks uses VW_UUID as that, so we search for that too.
+			static const FString Generic_MVRUUIDName = TEXT("MVR_UUID");
+			static const FString VW_MVRUUIDName = TEXT("VW_UUID");
 
-				const TSharedPtr<IDatasmithKeyValueProperty> DatasmithKeyValueProperty = [&MetaDataElement]()
+			const TSharedPtr<IDatasmithKeyValueProperty> DatasmithKeyValueProperty = [&MetaDataElement]()
+			{
+				TSharedPtr<IDatasmithKeyValueProperty> Result = MetaDataElement->GetPropertyByName(*Generic_MVRUUIDName);
+				if (!Result.IsValid())
 				{
-					TSharedPtr<IDatasmithKeyValueProperty> Result = MetaDataElement->GetPropertyByName(*Generic_MVRUUIDName);
-					if (!Result.IsValid())
-					{
-						Result = MetaDataElement->GetPropertyByName(*VW_MVRUUIDName);
-					}
+					Result = MetaDataElement->GetPropertyByName(*VW_MVRUUIDName);
+				}
 
-					return Result;
-				}();
+				return Result;
+			}();
 
-				if (!DatasmithKeyValueProperty.IsValid())
-				{
-					continue;
-				}
+			if (!DatasmithKeyValueProperty.IsValid())
+			{
+				continue;
+			}
+
+			// Try to find the node that corresponds to the MVR node
+			const FString MetaDataMVRUUIDString = FString(DatasmithKeyValueProperty->GetValue());
+			FGuid MetaDataMVRUUID;
+			if (!FGuid::Parse(MetaDataMVRUUIDString, MetaDataMVRUUID))
+			{
+				continue;
+			}
 
-				// Try to find the node that corresponds to the MVR node
-				FString MetaDataMVRUUIDString = FString(DatasmithKeyValueProperty->GetValue());
-				FGuid MetaDataMVRUUID;
-				if (!FGuid::Parse(MetaDataMVRUUIDString, MetaDataMVRUUID))
+			if (MetaDataMVRUUID == MVRFixture.UUID)
+			{
+				const TSharedPtr<IDatasmithElement> AssociatedElementToRemove = MetaDataElement->GetAssociatedElement();
+
+				const bool bValidElementType = AssociatedElementToRemove.IsValid() || AssociatedElementToRemove->IsA(EDatasmithElementType::Actor);
+				if (!ensureAlwaysMsgf(bValidElementType, TEXT("DatasmithMVRNativeTranslator unexpected Element type when trying to remove MVR element in favor of an MVR Scene Actor."), AssociatedElementToRemove->GetName()))
 				{
 					continue;
 				}
 
-				if (MetaDataMVRUUID == MVRFixture.UUID)
-				{
-					const TSharedPtr<IDatasmithElement> AssociatedElementToRemove = MetaDataElement->GetAssociatedElement();
-
-					if (!AssociatedElementToRemove.IsValid() || !AssociatedElementToRemove->IsA(EDatasmithElementType::Actor))
-					{
-						ensureAlwaysMsgf(0, TEXT("DatasmithMVRNativeTranslator unexpected Element type when trying to remove MVR element in favor of an MVR Scene Actor."), AssociatedElementToRemove->GetName());
-						continue;
-					}
-
-					const TSharedRef<IDatasmithActorElement> ActorElementToRemove = StaticCastSharedRef<IDatasmithActorElement>(AssociatedElementToRemove.ToSharedRef());
-					MetaDataElementsToRemove.Add(MetaDataElement.ToSharedRef());
-					ActorElementsToRemove.Add(ActorElementToRemove);
-				}
+				const TSharedRef<IDatasmithActorElement> ActorElementToRemove = StaticCastSharedRef<IDatasmithActorElement>(AssociatedElementToRemove.ToSharedRef());
+				MetaDataElementsToRemove.Add(MetaDataElement.ToSharedRef());
+				ActorElementsToRemove.Add(ActorElementToRemove);
 			}
 		}
 	}
