Fix static analysis warning
#rnx

diff --git a/Engine/Plugins/VirtualProduction/DMX/DMXEngine/Source/DMXRuntime/Private/MVR/DMXMVRSceneActor.cpp b/Engine/Plugins/VirtualProduction/DMX/DMXEngine/Source/DMXRuntime/Private/MVR/DMXMVRSceneActor.cpp
--- a/Engine/Plugins/VirtualProduction/DMX/DMXEngine/Source/DMXRuntime/Private/MVR/DMXMVRSceneActor.cpp
+++ b/Engine/Plugins/VirtualProduction/DMX/DMXEngine/Source/DMXRuntime/Private/MVR/DMXMVRSceneActor.cpp
@@ -438,7 +438,7 @@ UDMXEntityFixturePatch* ADMXMVRSceneActor::GetFixturePatch(AActor* Actor) const
 
 void ADMXMVRSceneActor::SetFixturePatch(AActor* Actor, UDMXEntityFixturePatch* FixturePatch)
 {
-	if (!ensureMsgf(Actor || FixturePatch, TEXT("Trying to Set Fixture Patch on Actor, but Actor or Fixture Patch are invalid.")))
+	if (!ensureMsgf(Actor && FixturePatch, TEXT("Trying to Set Fixture Patch on Actor, but Actor or Fixture Patch are invalid.")))
 	{
 		return;
 	}
