Reverting accidentally submitted files

diff --git a/Engine/Plugins/Runtime/Metasound/Source/MetasoundEditor/Public/MetasoundEditorModule.h b/Engine/Plugins/Runtime/Metasound/Source/MetasoundEditor/Public/MetasoundEditorModule.h
--- a/Engine/Plugins/Runtime/Metasound/Source/MetasoundEditor/Public/MetasoundEditorModule.h
+++ b/Engine/Plugins/Runtime/Metasound/Source/MetasoundEditor/Public/MetasoundEditorModule.h
@@ -82,7 +82,7 @@ namespace Metasound
 			// By default, proxy classes support child classes & inheritance.
 			virtual void RegisterExplicitProxyClass(const UClass& InClass) = 0;
 
-			virtual bool IsMetaSoundAssetClass(const FTopLevelAssetPath& InClassName) const = 0;
+			virtual bool IsMetaSoundAssetClass(const FName InClassName) const = 0;
 
 			// Primes MetaSound assets, effectively loading the asset asynchronously (if not already
 			// loaded) & registers them if not already registered with the MetaSound Class Registry.
diff --git a/Engine/Plugins/VirtualProduction/DMX/DMXEngine/Source/DMXEditor/Public/Widgets/SAssetPickerButton.h b/Engine/Plugins/VirtualProduction/DMX/DMXEngine/Source/DMXEditor/Public/Widgets/SAssetPickerButton.h
--- a/Engine/Plugins/VirtualProduction/DMX/DMXEngine/Source/DMXEditor/Public/Widgets/SAssetPickerButton.h
+++ b/Engine/Plugins/VirtualProduction/DMX/DMXEngine/Source/DMXEditor/Public/Widgets/SAssetPickerButton.h
@@ -37,11 +37,11 @@ public:
 		/** Allows more than a single class on the picker list.
 		 * Overrides AssetClass (doesn't affect the Use Selected on Content Browser button)
 		 */
-		SLATE_ARGUMENT(TArray<FTopLevelAssetPath>, AllowedClasses)
+		SLATE_ARGUMENT(TArray<FName>, AllowedClasses)
 		/** Disallows subclasses from the allowed asset class(es).
 		 * Doesn't affect the Use Selected on Content Browser button.
 		 */
-		SLATE_ARGUMENT(TArray<FTopLevelAssetPath>, DisallowedClasses)
+		SLATE_ARGUMENT(TArray<FName>, DisallowedClasses)
 
 		/** Called to decide the hovered state using the parent's interact able area */
 		SLATE_EVENT(FOnParentIsHovered, OnParentIsHovered)
@@ -56,9 +56,9 @@ public:
 	/** The class to filter the list of assets in the picker menu */
 	TWeakObjectPtr<UClass> AssetClass;
 	/** Allows more than a single class. Overrides AssetClass */
-	TArray<FTopLevelAssetPath> AllowedClasses;
+	TArray<FName> AllowedClasses;
 	/** Disallows subclasses from the allowed asset class(es) */
-	TArray<FTopLevelAssetPath> DisallowedClasses;
+	TArray<FName> DisallowedClasses;
 
 protected:
 	/** Clicked Use button */
