DMX - Fixture Patch and Fixutre Type Shared Data now update selection on undo/redo

diff --git a/Engine/Plugins/VirtualProduction/DMX/DMXEngine/Source/DMXEditor/Private/DMXFixturePatchSharedData.cpp b/Engine/Plugins/VirtualProduction/DMX/DMXEngine/Source/DMXEditor/Private/DMXFixturePatchSharedData.cpp
--- a/Engine/Plugins/VirtualProduction/DMX/DMXEngine/Source/DMXEditor/Private/DMXFixturePatchSharedData.cpp
+++ b/Engine/Plugins/VirtualProduction/DMX/DMXEngine/Source/DMXEditor/Private/DMXFixturePatchSharedData.cpp
@@ -3,54 +3,93 @@
 #include "DMXFixturePatchSharedData.h"
 
 #include "DMXEditor.h"
+#include "DMXFixturePatchSharedDataSelection.h"
 #include "Library/DMXEntityFixturePatch.h"
 
 
 #define LOCTEXT_NAMESPACE "DMXFixturePatchSharedData"
 
+FDMXFixturePatchSharedData::FDMXFixturePatchSharedData(TWeakPtr<FDMXEditor> InDMXEditorPtr)
+	: DMXEditorPtr(InDMXEditorPtr)
+{
+	Selection = NewObject<UDMXFixturePatchSharedDataSelection>(GetTransientPackage(), "FixturePatchSharedDataSelection", RF_Transactional);
+}
+
+void FDMXFixturePatchSharedData::AddReferencedObjects(FReferenceCollector& Collector)
+{
+	Collector.AddReferencedObject(Selection);
+}
+
+void FDMXFixturePatchSharedData::PostUndo(bool bSuccess)
+{
+	check(Selection);
+	Selection->SelectedFixturePatches.RemoveAll([](TWeakObjectPtr<UDMXEntityFixturePatch> FixturePatch)
+		{
+			return !FixturePatch.IsValid();
+		});
+	OnFixturePatchSelectionChanged.Broadcast();
+}
+
+void FDMXFixturePatchSharedData::PostRedo(bool bSuccess)
+{
+	check(Selection);
+	Selection->SelectedFixturePatches.RemoveAll([](TWeakObjectPtr<UDMXEntityFixturePatch> FixturePatch)
+		{
+			return !FixturePatch.IsValid();
+		});
+	OnFixturePatchSelectionChanged.Broadcast();
+}
+
 void FDMXFixturePatchSharedData::SelectUniverse(int32 UniverseID)
-{	
+{
+	check(Selection);
 	check(UniverseID >= 0);
 
-	if (UniverseID == SelectedUniverse)
+	if (UniverseID == Selection->SelectedUniverse)
 	{
 		return;
 	}
 
-	SelectedUniverse = UniverseID;
+	Selection->SelectedUniverse = UniverseID;
 	OnUniverseSelectionChanged.Broadcast();
 }
 
 int32 FDMXFixturePatchSharedData::GetSelectedUniverse()
 {
-	return SelectedUniverse;
+	check(Selection);
+	return Selection->SelectedUniverse;
 }
 
 void FDMXFixturePatchSharedData::SelectFixturePatch(TWeakObjectPtr<UDMXEntityFixturePatch> Patch)
 {
-	if (SelectedFixturePatches.Num() == 1 &&
-		SelectedFixturePatches[0] == Patch)
+	check(Selection);
+	if (Selection->SelectedFixturePatches.Num() == 1 &&
+		Selection->SelectedFixturePatches[0] == Patch)
 	{
 		return;
 	}
 
-	SelectedFixturePatches.Reset();
-	SelectedFixturePatches.Add(Patch);
+	Selection->Modify();
+	Selection->SelectedFixturePatches.Reset();
+	Selection->SelectedFixturePatches.Add(Patch);
 	OnFixturePatchSelectionChanged.Broadcast();
 }
 
 void FDMXFixturePatchSharedData::AddFixturePatchToSelection(TWeakObjectPtr<UDMXEntityFixturePatch> Patch)
 {
-	if (!SelectedFixturePatches.Contains(Patch))
+	check(Selection);
+	if (!Selection->SelectedFixturePatches.Contains(Patch))
 	{
-		SelectedFixturePatches.Add(Patch);
+		Selection->Modify();
+		Selection->SelectedFixturePatches.Add(Patch);
 		OnFixturePatchSelectionChanged.Broadcast();
 	}
 }
 
 void FDMXFixturePatchSharedData::SelectFixturePatches(const TArray<TWeakObjectPtr<UDMXEntityFixturePatch>>& Patches)
 {
-	if (SelectedFixturePatches == Patches)
+	check(Selection);
+	if (Selection->SelectedFixturePatches == Patches)
 	{
 		return;
 	}
@@ -63,14 +102,16 @@ void FDMXFixturePatchSharedData::SelectFixturePatches(const TArray<TWeakObjectPt
 		UniquePatchesOnlyArray.AddUnique(FixturePatch);
 	}
 
-	SelectedFixturePatches.Reset();
-	SelectedFixturePatches = UniquePatchesOnlyArray;
+	Selection->Modify();
+	Selection->SelectedFixturePatches.Reset();
+	Selection->SelectedFixturePatches = UniquePatchesOnlyArray;
 	OnFixturePatchSelectionChanged.Broadcast();
 }
 
 const TArray<TWeakObjectPtr<UDMXEntityFixturePatch>>& FDMXFixturePatchSharedData::GetSelectedFixturePatches() const
 {
-	return SelectedFixturePatches;
+	check(Selection);
+	return Selection->SelectedFixturePatches;
 }
 
 #undef LOCTEXT_NAMESPACE
diff --git a/Engine/Plugins/VirtualProduction/DMX/DMXEngine/Source/DMXEditor/Private/DMXFixtureTypeSharedData.cpp b/Engine/Plugins/VirtualProduction/DMX/DMXEngine/Source/DMXEditor/Private/DMXFixtureTypeSharedData.cpp
--- a/Engine/Plugins/VirtualProduction/DMX/DMXEngine/Source/DMXEditor/Private/DMXFixtureTypeSharedData.cpp
+++ b/Engine/Plugins/VirtualProduction/DMX/DMXEngine/Source/DMXEditor/Private/DMXFixtureTypeSharedData.cpp
@@ -3,6 +3,7 @@
 #include "DMXFixtureTypeSharedData.h"
 
 #include "DMXEditorUtils.h"
+#include "DMXFixtureTypeSharedDataSelection.h"
 #include "DMXRuntimeUtils.h"
 #include "Library/DMXEntityFixtureType.h"
 
@@ -70,9 +71,40 @@ namespace
 	}
 }
 
+FDMXFixtureTypeSharedData::FDMXFixtureTypeSharedData(TWeakPtr<FDMXEditor> InDMXEditorPtr)
+{
+	Selection = NewObject<UDMXFixtureTypeSharedDataSelection>(GetTransientPackage(), "FixtureTypeSharedDataSelection", RF_Transactional);
+}
+
+void FDMXFixtureTypeSharedData::AddReferencedObjects(FReferenceCollector& Collector)
+{
+	Collector.AddReferencedObject(Selection);
+}
+
+void FDMXFixtureTypeSharedData::PostUndo(bool bSuccess)
+{
+	check(Selection);
+	Selection->SelectedFixtureTypes.RemoveAll([](TWeakObjectPtr<UDMXEntityFixtureType> FixtureType)
+		{
+			return !FixtureType.IsValid();
+		});
+	OnFixtureTypesSelected.Broadcast();
+}
+
+void FDMXFixtureTypeSharedData::PostRedo(bool bSuccess)
+{
+	check(Selection);
+	Selection->SelectedFixtureTypes.RemoveAll([](TWeakObjectPtr<UDMXEntityFixtureType> FixtureType)
+		{
+			return !FixtureType.IsValid();
+		});
+	OnFixtureTypesSelected.Broadcast();
+}
+
 void FDMXFixtureTypeSharedData::SelectFixtureTypes(const TArray<TWeakObjectPtr<UDMXEntityFixtureType>>& InFixtureTypes)
 {
-	if (SelectedFixtureTypes != InFixtureTypes)
+	check(Selection);
+	if (Selection->SelectedFixtureTypes != InFixtureTypes)
 	{
 		// Make a copy without duplicates
 		TArray<TWeakObjectPtr<UDMXEntityFixtureType>> UniqueFixtureTypesOnlyArray;
@@ -82,12 +114,13 @@ void FDMXFixtureTypeSharedData::SelectFixtureTypes(const TArray<TWeakObjectPtr<U
 			UniqueFixtureTypesOnlyArray.AddUnique(FixtureType);
 		}
 
-		SelectedFixtureTypes = UniqueFixtureTypesOnlyArray;
+		Selection->Modify();
+		Selection->SelectedFixtureTypes = UniqueFixtureTypesOnlyArray;
 
 		// Selected Modes and Functions turn invalid
-		SelectedModeIndices.Reset();
-		SelectedFunctionIndices.Reset();
-		bFixtureMatrixSelected = false;
+		Selection->SelectedModeIndices.Reset();
+		Selection->SelectedFunctionIndices.Reset();
+		Selection->bFixtureMatrixSelected = false;
 
 		OnFixtureTypesSelected.Broadcast();
 	}
@@ -95,13 +128,15 @@ void FDMXFixtureTypeSharedData::SelectFixtureTypes(const TArray<TWeakObjectPtr<U
 
 void FDMXFixtureTypeSharedData::SelectModes(const TArray<int32>& InModeIndices)
 {
-	if (SelectedModeIndices != InModeIndices)
+	check(Selection);
+	if (Selection->SelectedModeIndices != InModeIndices)
 	{
-		SelectedModeIndices = InModeIndices;
+		Selection->Modify();
+		Selection->SelectedModeIndices = InModeIndices;
 
 		// Selected Functions turn invalid
-		SelectedFunctionIndices.Reset();
-		bFixtureMatrixSelected = false;
+		Selection->SelectedFunctionIndices.Reset();
+		Selection->bFixtureMatrixSelected = false;
 
 		OnModesSelected.Broadcast();
 	}
@@ -109,11 +144,14 @@ void FDMXFixtureTypeSharedData::SelectModes(const TArray<int32>& InModeIndices)
 
 void FDMXFixtureTypeSharedData::SetFunctionAndMatrixSelection(const TArray<int32>& InFunctionIndices, bool bMatrixSelected)
 {
-	if (SelectedFunctionIndices != InFunctionIndices ||
-		bFixtureMatrixSelected != bMatrixSelected)
+	check(Selection);
+	if (Selection->SelectedFunctionIndices != InFunctionIndices ||
+		Selection->bFixtureMatrixSelected != bMatrixSelected)
 	{
-		SelectedFunctionIndices = InFunctionIndices;
-		bFixtureMatrixSelected = bMatrixSelected;
+		Selection->Modify();
+
+		Selection->SelectedFunctionIndices = InFunctionIndices;
+		Selection->bFixtureMatrixSelected = bMatrixSelected;
 
 		OnFunctionsSelected.Broadcast();
 		OnMatrixSelectionChanged.Broadcast();
@@ -122,17 +160,18 @@ void FDMXFixtureTypeSharedData::SetFunctionAndMatrixSelection(const TArray<int32
 
 bool FDMXFixtureTypeSharedData::CanCopyModesToClipboard() const
 {
-	if (SelectedModeIndices.Num() == 0)
+	check(Selection);
+	if (Selection->SelectedModeIndices.Num() == 0)
 	{
 		return false;
 	}
-	else if (SelectedFixtureTypes.Num() == 1)
+	else if (Selection->SelectedFixtureTypes.Num() == 1)
 	{
-		if (UDMXEntityFixtureType* FixtureType = SelectedFixtureTypes[0].Get())
+		if (UDMXEntityFixtureType* FixtureType = Selection->SelectedFixtureTypes[0].Get())
 		{
-			for (int32 ModeIndex : SelectedModeIndices)
+			for (int32 ModeIndex : Selection->SelectedModeIndices)
 			{
-				if (!FixtureType->Modes.IsValidIndex(SelectedModeIndices[0]))
+				if (!FixtureType->Modes.IsValidIndex(Selection->SelectedModeIndices[0]))
 				{
 					return false;
 				}
@@ -145,15 +184,16 @@ bool FDMXFixtureTypeSharedData::CanCopyModesToClipboard() const
 
 void FDMXFixtureTypeSharedData::CopyModesToClipboard()
 {
+	check(Selection);
 	if (ensureMsgf(CanCopyModesToClipboard(), TEXT("Cannot copy Modes to clipboard. Please test first with FDMXFixtureTypeSharedData::CanCopyModesToClipboard.")))
 	{
 		ModesClipboard.Reset();
 
-		for (TWeakObjectPtr<UDMXEntityFixtureType> WeakFixtureType : SelectedFixtureTypes)
+		for (TWeakObjectPtr<UDMXEntityFixtureType> WeakFixtureType : Selection->SelectedFixtureTypes)
 		{
 			if (UDMXEntityFixtureType* FixtureType = WeakFixtureType.Get())
 			{
-				for (const int32 ModeIndex : SelectedModeIndices)
+				for (const int32 ModeIndex : Selection->SelectedModeIndices)
 				{
 					const FDMXFixtureMode& Mode = FixtureType->Modes[ModeIndex];
 					if (FixtureType->Modes.IsValidIndex(ModeIndex))
@@ -172,9 +212,10 @@ void FDMXFixtureTypeSharedData::CopyModesToClipboard()
 
 bool FDMXFixtureTypeSharedData::CanPasteModesFromClipboard() const
 {
-	if (SelectedFixtureTypes.Num() == 1 && ModesClipboard.Num() > 0)
+	check(Selection);
+	if (Selection->SelectedFixtureTypes.Num() == 1 && ModesClipboard.Num() > 0)
 	{
-		if (UDMXEntityFixtureType* FixtureType = SelectedFixtureTypes[0].Get())
+		if (UDMXEntityFixtureType* FixtureType = Selection->SelectedFixtureTypes[0].Get())
 		{
 			return true;
 		}
@@ -185,11 +226,12 @@ bool FDMXFixtureTypeSharedData::CanPasteModesFromClipboard() const
 
 void FDMXFixtureTypeSharedData::PasteModesFromClipboard(TArray<int32>& OutNewlyAddedModeIndices)
 {
+	check(Selection);
 	if (ensureMsgf(CanPasteModesFromClipboard(), TEXT("Cannot paste Modes from clipboard. Please test first with FDMXFixtureTypeSharedData::CanPasteModesFromClipboard.")))
 	{
-		if (SelectedFixtureTypes.Num() == 1)
+		if (Selection->SelectedFixtureTypes.Num() == 1)
 		{
-			if (UDMXEntityFixtureType* FixtureType = SelectedFixtureTypes[0].Get())
+			if (UDMXEntityFixtureType* FixtureType = Selection->SelectedFixtureTypes[0].Get())
 			{
 				const FText TransactionText = FText::Format(LOCTEXT("PasteModesTransaction", "Paste Fixture {0}|plural(one=Mode, other=Modes)"), GetSelectedModeIndices().Num());
 				const FScopedTransaction PasteModesTransaction(TransactionText);
@@ -203,13 +245,13 @@ void FDMXFixtureTypeSharedData::PasteModesFromClipboard(TArray<int32>& OutNewlyA
 
 					const int32 PasteOntoIndex = [FixtureType, &ClipboardIndex, this]()
 					{
-						if (SelectedModeIndices.IsValidIndex(ClipboardIndex))
+						if (Selection->SelectedModeIndices.IsValidIndex(ClipboardIndex))
 						{
-							return SelectedModeIndices[ClipboardIndex];
+							return Selection->SelectedModeIndices[ClipboardIndex];
 						}
-						else if (SelectedModeIndices.Num() > 0)
+						else if (Selection->SelectedModeIndices.Num() > 0)
 						{
-							return SelectedModeIndices.Last();
+							return Selection->SelectedModeIndices.Last();
 						}
 						else
 						{
@@ -254,19 +296,20 @@ void FDMXFixtureTypeSharedData::PasteModesFromClipboard(TArray<int32>& OutNewlyA
 
 bool FDMXFixtureTypeSharedData::CanCopyFunctionsToClipboard() const
 {
-	if (SelectedFunctionIndices.Num() == 0 || bFixtureMatrixSelected)
+	check(Selection);
+	if (Selection->SelectedFunctionIndices.Num() == 0 || Selection->bFixtureMatrixSelected)
 	{
 		return false;
 	}
-	else if (SelectedFixtureTypes.Num() == 1 && SelectedModeIndices.Num() == 1)
+	else if (Selection->SelectedFixtureTypes.Num() == 1 && Selection->SelectedModeIndices.Num() == 1)
 	{
-		if (UDMXEntityFixtureType* FixtureType = SelectedFixtureTypes[0].Get())
+		if (UDMXEntityFixtureType* FixtureType = Selection->SelectedFixtureTypes[0].Get())
 		{
-			if(FixtureType->Modes.IsValidIndex(SelectedModeIndices[0]))
+			if(FixtureType->Modes.IsValidIndex(Selection->SelectedModeIndices[0]))
 			{
-				const FDMXFixtureMode& Mode = FixtureType->Modes[SelectedModeIndices[0]];
+				const FDMXFixtureMode& Mode = FixtureType->Modes[Selection->SelectedModeIndices[0]];
 
-				for (const int32 FunctionIndex : SelectedFunctionIndices)
+				for (const int32 FunctionIndex : Selection->SelectedFunctionIndices)
 				{
 					if (!Mode.Functions.IsValidIndex(FunctionIndex))
 					{
@@ -282,20 +325,21 @@ bool FDMXFixtureTypeSharedData::CanCopyFunctionsToClipboard() const
 
 void FDMXFixtureTypeSharedData::CopyFunctionsToClipboard()
 {
+	check(Selection);
 	if (ensureMsgf(CanCopyFunctionsToClipboard(), TEXT("Cannot copy Functions to clipboard. Please test first with FDMXFixtureTypeSharedData::CopyFunctionsToClipboard.")))
 	{
 		FunctionsClipboard.Reset();
 
-		if (SelectedFixtureTypes.Num() == 1 && SelectedModeIndices.Num() == 1)
+		if (Selection->SelectedFixtureTypes.Num() == 1 && Selection->SelectedModeIndices.Num() == 1)
 		{
-			if (UDMXEntityFixtureType* FixtureType = SelectedFixtureTypes[0].Get())
+			if (UDMXEntityFixtureType* FixtureType = Selection->SelectedFixtureTypes[0].Get())
 			{
-				if (FixtureType->Modes.IsValidIndex(SelectedModeIndices[0]))
+				if (FixtureType->Modes.IsValidIndex(Selection->SelectedModeIndices[0]))
 				{
-					const FDMXFixtureMode& Mode = FixtureType->Modes[SelectedModeIndices[0]];
+					const FDMXFixtureMode& Mode = FixtureType->Modes[Selection->SelectedModeIndices[0]];
 
 					// Sort selected function indices by starting channel, so elements are copied in visual right order, not in order of how things were selected
-					SelectedFunctionIndices.Sort([Mode](const int32 FunctionIndexA, const int32 FunctionIndexB)
+					Selection->SelectedFunctionIndices.Sort([Mode](const int32 FunctionIndexA, const int32 FunctionIndexB)
 						{
 							if (Mode.Functions.IsValidIndex(FunctionIndexA) && Mode.Functions.IsValidIndex(FunctionIndexB))
 							{
@@ -304,7 +348,7 @@ void FDMXFixtureTypeSharedData::CopyFunctionsToClipboard()
 							return false;
 						});
 
-					for (const int32 FunctionIndex : SelectedFunctionIndices)
+					for (const int32 FunctionIndex : Selection->SelectedFunctionIndices)
 					{
 						if (Mode.Functions.IsValidIndex(FunctionIndex))
 						{
@@ -323,11 +367,12 @@ void FDMXFixtureTypeSharedData::CopyFunctionsToClipboard()
 
 bool FDMXFixtureTypeSharedData::CanPasteFunctionsFromClipboard() const
 {
-	if (SelectedFixtureTypes.Num() == 1 && SelectedModeIndices.Num() == 1 && ModesClipboard.Num() > 0)
+	check(Selection);
+	if (Selection->SelectedFixtureTypes.Num() == 1 && Selection->SelectedModeIndices.Num() == 1 && ModesClipboard.Num() > 0)
 	{
-		if (UDMXEntityFixtureType* FixtureType = SelectedFixtureTypes[0].Get())
+		if (UDMXEntityFixtureType* FixtureType = Selection->SelectedFixtureTypes[0].Get())
 		{
-			if (FixtureType->Modes.IsValidIndex(SelectedModeIndices[0]))
+			if (FixtureType->Modes.IsValidIndex(Selection->SelectedModeIndices[0]))
 			{
 				return true;
 			}
@@ -339,19 +384,20 @@ bool FDMXFixtureTypeSharedData::CanPasteFunctionsFromClipboard() const
 
 void FDMXFixtureTypeSharedData::PasteFunctionsFromClipboard(TArray<int32>& OutNewlyAddedFunctionIndices)
 {
+	check(Selection);
 	if (ensureMsgf(CanPasteFunctionsFromClipboard(), TEXT("Cannot paste Functions from clipboard. Please test first with FDMXFixtureTypeSharedData::CanPasteFunctionsFromClipboard.")))
 	{
-		if (SelectedFixtureTypes.Num() == 1 && SelectedModeIndices.Num() == 1)
+		if (Selection->SelectedFixtureTypes.Num() == 1 && Selection->SelectedModeIndices.Num() == 1)
 		{
-			if (UDMXEntityFixtureType* FixtureType = SelectedFixtureTypes[0].Get())
+			if (UDMXEntityFixtureType* FixtureType = Selection->SelectedFixtureTypes[0].Get())
 			{
-				const int32 SelectedModeIndex = SelectedModeIndices[0];
+				const int32 SelectedModeIndex = Selection->SelectedModeIndices[0];
 				if (FixtureType->Modes.IsValidIndex(SelectedModeIndex))
 				{
 					FDMXFixtureMode& Mode = FixtureType->Modes[SelectedModeIndex];
 
 					// Sort selected function indices by starting channel, so elements are pasted in visual right order, not in order of how things were selected
-					SelectedFunctionIndices.Sort([Mode](const int32 FunctionIndexA, const int32 FunctionIndexB)
+					Selection->SelectedFunctionIndices.Sort([Mode](const int32 FunctionIndexA, const int32 FunctionIndexB)
 						{
 							if (Mode.Functions.IsValidIndex(FunctionIndexA) && Mode.Functions.IsValidIndex(FunctionIndexB))
 							{
@@ -380,13 +426,13 @@ void FDMXFixtureTypeSharedData::PasteFunctionsFromClipboard(TArray<int32>& OutNe
 								// Find an index to paste onto
 								const int32 PasteOntoIndex = [Mode, ClipboardIndex, this]()
 								{
-									if (SelectedFunctionIndices.IsValidIndex(ClipboardIndex))
+									if (Selection->SelectedFunctionIndices.IsValidIndex(ClipboardIndex))
 									{
-										return SelectedFunctionIndices[ClipboardIndex];
+										return Selection->SelectedFunctionIndices[ClipboardIndex];
 									}
-									else if (SelectedFunctionIndices.Num() > 0)
+									else if (Selection->SelectedFunctionIndices.Num() > 0)
 									{
-										return SelectedFunctionIndices.Last();
+										return Selection->SelectedFunctionIndices.Last();
 									}
 									else
 									{
@@ -422,18 +468,45 @@ void FDMXFixtureTypeSharedData::PasteFunctionsFromClipboard(TArray<int32>& OutNe
 void FDMXFixtureTypeSharedData::SelectFunctions(const TArray<int32>& InFunctionIndices)
 {
 	// DEPRECATED 5.0
-	if (SelectedFunctionIndices != InFunctionIndices)
+	check(Selection);
+	if (Selection->SelectedFunctionIndices != InFunctionIndices)
 	{
-		SelectedFunctionIndices = InFunctionIndices;
+		Selection->Modify();
+		Selection->SelectedFunctionIndices = InFunctionIndices;
 		OnFunctionsSelected.Broadcast();
 	}
 }
 
+
+const TArray<TWeakObjectPtr<UDMXEntityFixtureType>>& FDMXFixtureTypeSharedData::GetSelectedFixtureTypes() const 
+{
+	check(Selection);
+	return Selection->SelectedFixtureTypes;
+}
+
+const TArray<int32>& FDMXFixtureTypeSharedData::GetSelectedModeIndices() const
+{
+	check(Selection);
+	return Selection->SelectedModeIndices;
+}
+
+const TArray<int32>& FDMXFixtureTypeSharedData::GetSelectedFunctionIndices() const
+{
+	check(Selection);
+	return Selection->SelectedFunctionIndices;
+}
+
+bool FDMXFixtureTypeSharedData::IsFixtureMatrixSelected() const
+{
+	check(Selection);
+	return Selection->bFixtureMatrixSelected;
+}
+
 bool FDMXFixtureTypeSharedData::CanAddMode() const
 {
 	// DEPRECATED 5.0
-
-	return SelectedFixtureTypes.Num() == 1;
+	check(Selection);
+	return Selection->SelectedFixtureTypes.Num() == 1;
 }
 
 void FDMXFixtureTypeSharedData::AddMode()
@@ -446,7 +519,8 @@ void FDMXFixtureTypeSharedData::AddMode()
 
 	if (ensureMsgf(bCanAddMode, TEXT("Call to FDMXFixtureTypeSharedData::AddMode when no mode can be added")))
 	{
-		for (const TWeakObjectPtr<UDMXEntityFixtureType>& FixtureType : SelectedFixtureTypes)
+		check(Selection);
+		for (const TWeakObjectPtr<UDMXEntityFixtureType>& FixtureType : Selection->SelectedFixtureTypes)
 		{
 			if (!FixtureType.IsValid())
 			{
@@ -473,8 +547,8 @@ void FDMXFixtureTypeSharedData::AddMode()
 void FDMXFixtureTypeSharedData::DuplicateModes(const TArray<int32>& ModeÎndiciesToDuplicate)
 {
 	// DEPRECATED 5.0
-
-	for (const TWeakObjectPtr<UDMXEntityFixtureType>& WeakFixtureType : SelectedFixtureTypes)
+	check(Selection);
+	for (const TWeakObjectPtr<UDMXEntityFixtureType>& WeakFixtureType : Selection->SelectedFixtureTypes)
 	{ 
 		if (UDMXEntityFixtureType* FixtureType = WeakFixtureType.Get())
 		{
@@ -512,8 +586,8 @@ void FDMXFixtureTypeSharedData::DuplicateModes(const TArray<int32>& ModeÎndicie
 void FDMXFixtureTypeSharedData::DeleteModes(const TArray<int32>& ModeIndicesToDelete)
 {
 	// DEPRECATED 5.0
-
-	for (const TWeakObjectPtr<UDMXEntityFixtureType>& WeakFixtureType : SelectedFixtureTypes)
+	check(Selection);
+	for (const TWeakObjectPtr<UDMXEntityFixtureType>& WeakFixtureType : Selection->SelectedFixtureTypes)
 	{
 		if (UDMXEntityFixtureType* FixtureType = WeakFixtureType.Get())
 		{
@@ -564,7 +638,8 @@ void FDMXFixtureTypeSharedData::PasteClipboardToModes(const TArray<int32>& ModeI
 			FDMXFixtureMode ModeFromClipboard;
 			if (FJsonObjectConverter::JsonObjectToUStruct(RootJsonObject.ToSharedRef(), FDMXFixtureMode::StaticStruct(), &ModeFromClipboard, 0, 0))
 			{
-				for (const TWeakObjectPtr<UDMXEntityFixtureType>& WeakFixtureType : SelectedFixtureTypes)
+				check(Selection);
+				for (const TWeakObjectPtr<UDMXEntityFixtureType>& WeakFixtureType : Selection->SelectedFixtureTypes)
 				{
 					if (UDMXEntityFixtureType* FixtureType = WeakFixtureType.Get())
 					{
@@ -595,8 +670,8 @@ void FDMXFixtureTypeSharedData::PasteClipboardToModes(const TArray<int32>& ModeI
 bool FDMXFixtureTypeSharedData::CanAddFunction() const
 {
 	// DEPRECATED 5.0
-
-	return SelectedModeIndices.Num() == 1;
+	check(Selection);
+	return Selection->SelectedModeIndices.Num() == 1;
 }
 
 void FDMXFixtureTypeSharedData::AddFunctionToSelectedMode()
@@ -609,7 +684,8 @@ void FDMXFixtureTypeSharedData::AddFunctionToSelectedMode()
 
 	if (ensureMsgf(bCanAddFunction, TEXT("Call to FDMXFixtureTypeSharedData::AddFunctionToSelectedMode, when no function can be added")))
 	{
-		for (const TWeakObjectPtr<UDMXEntityFixtureType>& WeakFixtureType : SelectedFixtureTypes)
+		check(Selection);
+		for (const TWeakObjectPtr<UDMXEntityFixtureType>& WeakFixtureType : Selection->SelectedFixtureTypes)
 		{
 			if (UDMXEntityFixtureType* FixtureType = WeakFixtureType.Get())
 			{
@@ -618,7 +694,7 @@ void FDMXFixtureTypeSharedData::AddFunctionToSelectedMode()
 				FixtureType->PreEditChange(UDMXEntityFixtureType::StaticClass()->FindPropertyByName(GET_MEMBER_NAME_STRING_CHECKED(UDMXEntityFixtureType, Modes)));
 				FixtureType->Modify();
 
-				int32 ModeIndex = SelectedModeIndices[0]; // Implicitly valid via CanAddFunctions tested before
+				int32 ModeIndex = Selection->SelectedModeIndices[0]; // Implicitly valid via CanAddFunctions tested before
 				if (FixtureType->Modes.IsValidIndex(ModeIndex))
 				{
 					FDMXFixtureMode& Mode = FixtureType->Modes[ModeIndex];
@@ -647,12 +723,12 @@ void FDMXFixtureTypeSharedData::AddFunctionToSelectedMode()
 void FDMXFixtureTypeSharedData::DuplicateFunctions(const TArray<int32>& FunctionIndicesToDuplicate)
 {
 	// DEPRECATED 5.0
-
-	for (const TWeakObjectPtr<UDMXEntityFixtureType>& WeakFixtureType : SelectedFixtureTypes)
+	check(Selection);
+	for (const TWeakObjectPtr<UDMXEntityFixtureType>& WeakFixtureType : Selection->SelectedFixtureTypes)
 	{
 		if (UDMXEntityFixtureType* FixtureType = WeakFixtureType.Get())
 		{
-			for (int32 ModeIndex : SelectedModeIndices)
+			for (int32 ModeIndex : Selection->SelectedModeIndices)
 			{
 				if (FixtureType->Modes.IsValidIndex(ModeIndex))
 				{
@@ -686,12 +762,12 @@ void FDMXFixtureTypeSharedData::DuplicateFunctions(const TArray<int32>& Function
 void FDMXFixtureTypeSharedData::DeleteFunctions(const TArray<int32>& FunctionIndicesToDelete)
 {
 	// DEPRECATED 5.0
-
-	for (const TWeakObjectPtr<UDMXEntityFixtureType>& WeakFixtureType : SelectedFixtureTypes)
+	check(Selection);
+	for (const TWeakObjectPtr<UDMXEntityFixtureType>& WeakFixtureType : Selection->SelectedFixtureTypes)
 	{
 		if (UDMXEntityFixtureType* FixtureType = WeakFixtureType.Get())
 		{
-			for (int32 ModeIndex : SelectedModeIndices)
+			for (int32 ModeIndex : Selection->SelectedModeIndices)
 			{
 				if (FixtureType->Modes.IsValidIndex(ModeIndex))
 				{
@@ -761,7 +837,8 @@ void FDMXFixtureTypeSharedData::PasteClipboardToFunctions(const TArray<int32>& F
 			FDMXFixtureFunction FunctionFromClipboard;
 			if (FJsonObjectConverter::JsonObjectToUStruct(RootJsonObject.ToSharedRef(), FDMXFixtureFunction::StaticStruct(), &FunctionFromClipboard, 0, 0))
 			{
-				for (const TWeakObjectPtr<UDMXEntityFixtureType>& WeakFixtureType : SelectedFixtureTypes)
+				check(Selection);
+				for (const TWeakObjectPtr<UDMXEntityFixtureType>& WeakFixtureType : Selection->SelectedFixtureTypes)
 				{
 					if (UDMXEntityFixtureType* FixtureType = WeakFixtureType.Get())
 					{
@@ -770,7 +847,7 @@ void FDMXFixtureTypeSharedData::PasteClipboardToFunctions(const TArray<int32>& F
 						FixtureType->PreEditChange(UDMXEntityFixtureType::StaticClass()->FindPropertyByName(GET_MEMBER_NAME_STRING_CHECKED(UDMXEntityFixtureType, Modes)));
 						FixtureType->Modify();
 
-						for (int32 ModeIndex : SelectedModeIndices)
+						for (int32 ModeIndex : Selection->SelectedModeIndices)
 						{
 							if (FixtureType->Modes.IsValidIndex(ModeIndex))
 							{
@@ -801,10 +878,10 @@ void FDMXFixtureTypeSharedData::PasteClipboardToFunctions(const TArray<int32>& F
 bool FDMXFixtureTypeSharedData::CanAddCellAttribute() const
 {
 	// DEPRECATED 5.0
-
+	check(Selection);
 	return 
-		SelectedFixtureTypes.Num() == 1 && 
-		SelectedModeIndices.Num() == 1;
+		Selection->SelectedFixtureTypes.Num() == 1 &&
+		Selection->SelectedModeIndices.Num() == 1;
 }
 
 void FDMXFixtureTypeSharedData::AddCellAttributeToSelectedMode()
@@ -817,12 +894,13 @@ void FDMXFixtureTypeSharedData::AddCellAttributeToSelectedMode()
 
 	if (ensureMsgf(bCanAddCellAttribute, TEXT("Call to FDMXFixtureTypeSharedData::AddCellAttributeToSelectedMode, when no Cell Attribute can be added.")))
 	{
-		if (UDMXEntityFixtureType* FixtureType = SelectedFixtureTypes[0].Get())
+		check(Selection);
+		if (UDMXEntityFixtureType* FixtureType = Selection->SelectedFixtureTypes[0].Get())
 		{
-			const int32 SelectedModeIndex = SelectedModeIndices[0];
+			const int32 SelectedModeIndex = Selection->SelectedModeIndices[0];
 			if (FixtureType->Modes.IsValidIndex(SelectedModeIndex) && FixtureType->Modes[SelectedModeIndex].bFixtureMatrixEnabled)
 			{
-				FDMXFixtureMode& Mode = FixtureType->Modes[SelectedModeIndices[0]];
+				FDMXFixtureMode& Mode = FixtureType->Modes[Selection->SelectedModeIndices[0]];
 
 				FDMXFixtureCellAttribute NewAttribute;
 				TArray<FName> AvailableAttributes = FDMXAttributeName::GetPossibleValues();
diff --git a/Engine/Plugins/VirtualProduction/DMX/DMXEngine/Source/DMXEditor/Public/DMXFixturePatchSharedData.h b/Engine/Plugins/VirtualProduction/DMX/DMXEngine/Source/DMXEditor/Public/DMXFixturePatchSharedData.h
--- a/Engine/Plugins/VirtualProduction/DMX/DMXEngine/Source/DMXEditor/Public/DMXFixturePatchSharedData.h
+++ b/Engine/Plugins/VirtualProduction/DMX/DMXEngine/Source/DMXEditor/Public/DMXFixturePatchSharedData.h
@@ -3,23 +3,41 @@
 #pragma once
 
 #include "CoreMinimal.h"
+#include "EditorUndoClient.h"
 #include "Templates/SharedPointer.h"
+#include "UObject/GCObject.h"
 #include "UObject/WeakObjectPtrTemplates.h"
 
 class FDMXEditor;
 class FDMXFixturePatchNode;
 class SDMXEntityList;
 class UDMXEntityFixturePatch;
+class UDMXFixturePatchSharedDataSelection;
 class UDMXLibrary;
 
+
 /** Shared data for fixture patch editors */
 class FDMXFixturePatchSharedData
-	: public TSharedFromThis<FDMXFixturePatchSharedData>
+	: public FGCObject
+	, public FSelfRegisteringEditorUndoClient
+	, public TSharedFromThis<FDMXFixturePatchSharedData>
 {
 public:
-	FDMXFixturePatchSharedData(TWeakPtr<FDMXEditor> InDMXEditorPtr)
-		: DMXEditorPtr(InDMXEditorPtr)
-	{}
+	/** Constructor */
+	FDMXFixturePatchSharedData(TWeakPtr<FDMXEditor> InDMXEditorPtr);
+
+	//~ Begin FGCObject
+	virtual void AddReferencedObjects(FReferenceCollector& Collector) override;
+	virtual FString GetReferencerName() const override
+	{
+		return TEXT("DMXEditor::DMXFixturePatchSharedData");
+	}
+	//~End FGCObject
+
+	//~Begin EditorUndoClient interface
+	virtual void PostUndo(bool bSuccess) override;
+	virtual void PostRedo(bool bSuccess) override;
+	//~End EditorUndoClient interface
 
 	/** Broadcast when a universe is selected by an editor */
 	FSimpleMulticastDelegate OnUniverseSelectionChanged;
@@ -46,11 +64,8 @@ public:
 	const TArray<TWeakObjectPtr<UDMXEntityFixturePatch>>& GetSelectedFixturePatches() const;
 
 private:
-	/** The universe currently edited by editors */
-	int32 SelectedUniverse = 1;
-
-	/** Patch node currently selected or nullptr if nothing is selected */
-	TArray<TWeakObjectPtr<UDMXEntityFixturePatch>> SelectedFixturePatches;
+	/** The current selection */
+	UDMXFixturePatchSharedDataSelection* Selection;
 
 	/** Weak reference to the DMX editor */
 	TWeakPtr<FDMXEditor> DMXEditorPtr;
diff --git a/Engine/Plugins/VirtualProduction/DMX/DMXEngine/Source/DMXEditor/Public/DMXFixturePatchSharedDataSelection.h b/Engine/Plugins/VirtualProduction/DMX/DMXEngine/Source/DMXEditor/Public/DMXFixturePatchSharedDataSelection.h
new file mode 100644
--- /dev/null
+++ b/Engine/Plugins/VirtualProduction/DMX/DMXEngine/Source/DMXEditor/Public/DMXFixturePatchSharedDataSelection.h
@@ -0,0 +1,26 @@
+// Copyright Epic Games, Inc. All Rights Reserved.
+
+#pragma once
+
+#include "CoreMinimal.h"
+#include "DMXFixturePatchSharedDataSelection.generated.h"
+
+class UDMXEntityFixturePatch;
+
+
+/** An Array of Fixture Patches that supports a transaction history resp. undo / redo */
+UCLASS()
+class UDMXFixturePatchSharedDataSelection
+	: public UObject
+{
+	GENERATED_BODY()
+
+public:
+	/** Selected Fixture Patches */
+	UPROPERTY()
+	TArray<TWeakObjectPtr<UDMXEntityFixturePatch>> SelectedFixturePatches;
+
+	/** Selected Universe */
+	UPROPERTY()
+	int32 SelectedUniverse = 1;
+};
diff --git a/Engine/Plugins/VirtualProduction/DMX/DMXEngine/Source/DMXEditor/Public/DMXFixtureTypeSharedData.h b/Engine/Plugins/VirtualProduction/DMX/DMXEngine/Source/DMXEditor/Public/DMXFixtureTypeSharedData.h
--- a/Engine/Plugins/VirtualProduction/DMX/DMXEngine/Source/DMXEditor/Public/DMXFixtureTypeSharedData.h
+++ b/Engine/Plugins/VirtualProduction/DMX/DMXEngine/Source/DMXEditor/Public/DMXFixtureTypeSharedData.h
@@ -5,7 +5,9 @@
 #include "IDMXNamedType.h"
 
 #include "CoreMinimal.h"
+#include "EditorUndoClient.h"
 #include "Templates/SharedPointer.h"
+#include "UObject/GCObject.h"
 #include "UObject/WeakObjectPtrTemplates.h"
 
 #include "JsonObjectConverter.h"
@@ -14,20 +16,34 @@ enum class EDMXFixtureSignalFormat : uint8;
 class FDMXEditor;
 class FDMXFixtureTypeSharedData;
 class UDMXEntityFixtureType;
+class UDMXFixtureTypeSharedDataSelection;
 
 class FScopedTransaction;
 
 
-/** Shared data and utilities for objects in a DMX Editor */
+/** Shared data for Fixture Types in a DMX Editor */
 class FDMXFixtureTypeSharedData
-	: public TSharedFromThis<FDMXFixtureTypeSharedData>
+	: public FGCObject
+	, public FSelfRegisteringEditorUndoClient
+	, public TSharedFromThis<FDMXFixtureTypeSharedData>
 {
 public:
-	FDMXFixtureTypeSharedData(TWeakPtr<FDMXEditor> InDMXEditorPtr)
-		: DMXEditorPtr(InDMXEditorPtr)
-	{}	
+	/** Constructor */
+	FDMXFixtureTypeSharedData(TWeakPtr<FDMXEditor> InDMXEditorPtr);
+
+	//~ Begin FGCObject
+	virtual void AddReferencedObjects(FReferenceCollector& Collector) override;
+	virtual FString GetReferencerName() const override
+	{
+		return TEXT("DMXEditor::DMXFixtureTypeSharedData");
+	}
+	//~End FGCObject
+
+	//~Begin EditorUndoClient interface
+	virtual void PostUndo(bool bSuccess) override;
+	virtual void PostRedo(bool bSuccess) override;
+	//~End EditorUndoClient interface
 
-public:
 	/** Selects specified Fixture Types */
 	void SelectFixtureTypes(const TArray<TWeakObjectPtr<UDMXEntityFixtureType>>& InFixtureTypes);
 
@@ -41,10 +57,10 @@ public:
 	UE_DEPRECATED(5.0, "Deprecated in favor of SetFunctionAndMatrixSelection to avoid the unclear state where both need change, but one contains the old state while the other changed.")
 	void SelectFunctions(const TArray<int32>& InFunctionIndices);
 
-	const TArray<TWeakObjectPtr<UDMXEntityFixtureType>>& GetSelectedFixtureTypes() const { return SelectedFixtureTypes; }
-	const TArray<int32>& GetSelectedModeIndices() const { return SelectedModeIndices; }
-	const TArray<int32>& GetSelectedFunctionIndices() const { return SelectedFunctionIndices; }
-	bool IsFixtureMatrixSelected() const { return bFixtureMatrixSelected; }
+	const TArray<TWeakObjectPtr<UDMXEntityFixtureType>>& GetSelectedFixtureTypes() const;
+	const TArray<int32>& GetSelectedModeIndices() const;
+	const TArray<int32>& GetSelectedFunctionIndices() const;
+	bool IsFixtureMatrixSelected() const;
 
 	/** Returns true if selected Modes can be copied */
 	bool CanCopyModesToClipboard() const;
@@ -120,16 +136,7 @@ public:
 
 private:
 	/** The Fixture types being edited */
-	TArray<TWeakObjectPtr<UDMXEntityFixtureType>> SelectedFixtureTypes;
-
-	/** The Mode indices in the Selected Fixture Types currently being selected */
-	TArray<int32> SelectedModeIndices;
-
-	/** The Function indices in the Selected Fixture Types currently being selected  */
-	TArray<int32> SelectedFunctionIndices;
-
-	/** If true the Fixture Matrices in the currently seleccted Modes are selected */
-	bool bFixtureMatrixSelected = false;
+	UDMXFixtureTypeSharedDataSelection* Selection;
 
 	/** Cache for multi mode copy/paste*/
 	TArray<FString> ModesClipboard;
diff --git a/Engine/Plugins/VirtualProduction/DMX/DMXEngine/Source/DMXEditor/Public/DMXFixtureTypeSharedDataSelection.h b/Engine/Plugins/VirtualProduction/DMX/DMXEngine/Source/DMXEditor/Public/DMXFixtureTypeSharedDataSelection.h
new file mode 100644
--- /dev/null
+++ b/Engine/Plugins/VirtualProduction/DMX/DMXEngine/Source/DMXEditor/Public/DMXFixtureTypeSharedDataSelection.h
@@ -0,0 +1,34 @@
+// Copyright Epic Games, Inc. All Rights Reserved.
+
+#pragma once
+
+#include "CoreMinimal.h"
+#include "DMXFixtureTypeSharedDataSelection.generated.h"
+
+class UDMXEntityFixtureType;
+
+
+/** An Array of Fixture Types that supports a transaction history resp. undo / redo */
+UCLASS()
+class UDMXFixtureTypeSharedDataSelection
+	: public UObject
+{
+	GENERATED_BODY()
+
+public:
+	/** The Fixture types being edited */
+	UPROPERTY()
+	TArray<TWeakObjectPtr<UDMXEntityFixtureType>> SelectedFixtureTypes;
+
+	/** The Mode indices in the Selected Fixture Types currently being selected */
+	UPROPERTY()
+	TArray<int32> SelectedModeIndices;
+
+	/** The Function indices in the Selected Fixture Types currently being selected  */
+	UPROPERTY()
+	TArray<int32> SelectedFunctionIndices;
+
+	/** If true the Fixture Matrices in the currently seleccted Modes are selected */
+	UPROPERTY()
+	bool bFixtureMatrixSelected = false;
+};
