DMX - Improve MVR and GDTF import to retain imported data

diff --git a/Engine/Plugins/VirtualProduction/DMX/DMXEngine/Source/DMXEditor/Private/Factories/DMXGDTFFactory.cpp b/Engine/Plugins/VirtualProduction/DMX/DMXEngine/Source/DMXEditor/Private/Factories/DMXGDTFFactory.cpp
--- a/Engine/Plugins/VirtualProduction/DMX/DMXEngine/Source/DMXEditor/Private/Factories/DMXGDTFFactory.cpp
+++ b/Engine/Plugins/VirtualProduction/DMX/DMXEngine/Source/DMXEditor/Private/Factories/DMXGDTFFactory.cpp
@@ -4,15 +4,18 @@
 #include "Factories/DMXGDTFImportUI.h"
 #include "DMXEditorLog.h"
 #include "Factories/DMXGDTFImporter.h"
+#include "Library/DMXGDTFAssetImportData.h"
 #include "Library/DMXImportGDTF.h"
 
-#include "Misc/Paths.h"
 #include "Editor.h"
 #include "EditorReimportHandler.h"
 #include "AssetImportTask.h"
-#include "Misc/FeedbackContext.h"
-#include "Misc/CommandLine.h"
+#include "EditorFramework/AssetImportData.h"
 #include "HAL/FileManager.h"
+#include "Misc/CommandLine.h"
+#include "Misc/FeedbackContext.h"
+#include "Misc/Paths.h"
+
 
 const TCHAR* UDMXGDTFFactory::Extension = TEXT("gdtf");
 
@@ -99,7 +102,6 @@ UObject* UDMXGDTFFactory::FactoryCreateFile(UClass* InClass, UObject* InParent,
 
     GEditor->GetEditorSubsystem<UImportSubsystem>()->BroadcastAssetPreImport(this, InClass, InParent, InName, Type);
 
-    UDMXImportGDTF* CreatedObject = nullptr;
     UObject* ExistingObject = nullptr;
     if (InParent != nullptr)
     {
@@ -118,7 +120,7 @@ UObject* UDMXGDTFFactory::FactoryCreateFile(UClass* InClass, UObject* InParent,
     ImportArgs.CurrentFilename = InFilename;
     ImportArgs.Flags = Flags;
     ImportArgs.bCancelOperation = bOperationCanceled;
-    TUniquePtr<FDMXGDTFImporter> Importer = MakeUnique<FDMXGDTFImporter>(ImportArgs);
+    const TUniquePtr<FDMXGDTFImporter> Importer = MakeUnique<FDMXGDTFImporter>(ImportArgs);
 
     // Set Import UI
     if (FParse::Param(FCommandLine::Get(), TEXT("NoDMXImportOption")))
@@ -136,38 +138,37 @@ UObject* UDMXGDTFFactory::FactoryCreateFile(UClass* InClass, UObject* InParent,
         bShowOption = false;
     }
 
-    bool bImportError = false;
-
     if (!ImportUI->bImportXML && !ImportUI->bImportModels && !ImportUI->bImportTextures)
     {
         Warn->Log(ELogVerbosity::Error, TEXT("Nothing to Import") );
-        bImportError = true;
+		return nullptr;
     }
 
     // Try to load and parse the content
-    if ( !Importer->AttemptImportFromFile() )
+    if (!Importer->AttemptImportFromFile())
     {
-        Warn->Log(ELogVerbosity::Error, TEXT("Attempt Import") );
-        bImportError = true;
+        Warn->Log(ELogVerbosity::Error, TEXT("Failed to import GDTF") );
+		return nullptr;
     }
 
-    // Import was successful, create new file
-    if (!bImportError)
-    {
-		// Import to the Editor
-		CreatedObject = Importer->Import();
+	// Import to the Editor
+	UDMXImportGDTF* GDTF = Importer->Import();
+	if (!GDTF)
+	{
+		return nullptr;
+	}
 
-        if (CreatedObject != nullptr)
-        {
-            CreatedObject->SourceFilename = InFilename;
-        }
-		GEditor->GetEditorSubsystem<UImportSubsystem>()->BroadcastAssetPostImport(this, CreatedObject);
-    }
+	// Set Asset Import Data
+	UDMXGDTFAssetImportData* GDTFAssetImportData = GDTF->GetGDTFAssetImportData();
+	if (!ensureAlwaysMsgf(GDTFAssetImportData, TEXT("Unexpected missing Asset Import Data for newly created GDTF %s"), *GDTF->GetName()))
+	{
+		return nullptr;
+	}
+	GDTFAssetImportData->SetSourceFile(InFilename);
 
-    // We do not need importer any more
-    Importer.Reset();
+	GEditor->GetEditorSubsystem<UImportSubsystem>()->BroadcastAssetPostImport(this, GDTF);
 
-    return CreatedObject;
+    return GDTF;
 }
 
 bool UDMXGDTFFactory::FactoryCanImport(const FString& Filename)
@@ -184,10 +185,11 @@ bool UDMXGDTFFactory::FactoryCanImport(const FString& Filename)
 
 bool UDMXGDTFFactory::CanReimport(UObject* Obj, TArray<FString>& OutFilenames)
 {
-	UDMXImportGDTF* GDTFReimport = Cast<UDMXImportGDTF>(Obj);
-	if (GDTFReimport)
+	UDMXImportGDTF* GDTF = Cast<UDMXImportGDTF>(Obj);
+	if (GDTF && GDTF->GetGDTFAssetImportData())
 	{
-		OutFilenames.Add(GDTFReimport->SourceFilename);
+		const FString SourceFilename = GDTF->GetGDTFAssetImportData()->GetSourceFilePathAndName();
+		OutFilenames.Add(SourceFilename);
 		return true;
 	}
 	return false;
@@ -195,34 +197,35 @@ bool UDMXGDTFFactory::CanReimport(UObject* Obj, TArray<FString>& OutFilenames)
 
 void UDMXGDTFFactory::SetReimportPaths(UObject* Obj, const TArray<FString>& NewReimportPaths)
 {
-	UDMXImportGDTF* GDTFReimport = Cast<UDMXImportGDTF>(Obj);
-	if (GDTFReimport && ensure(NewReimportPaths.Num() == 1))
+	UDMXImportGDTF* GDTF = Cast<UDMXImportGDTF>(Obj);
+	if (GDTF && GDTF->GetGDTFAssetImportData() && ensure(NewReimportPaths.Num() == 1))
 	{
-        GDTFReimport->SourceFilename = NewReimportPaths[0];
+		GDTF->GetGDTFAssetImportData()->SetSourceFile(NewReimportPaths[0]);
 	}
 }
 
 EReimportResult::Type UDMXGDTFFactory::Reimport(UObject* InObject)
 {
-	UDMXImportGDTF* GDTFReimport = Cast<UDMXImportGDTF>(InObject);
+	UDMXImportGDTF* GDTF = Cast<UDMXImportGDTF>(InObject);
 
-	if (!GDTFReimport)
+	if (!GDTF || !GDTF->GetGDTFAssetImportData())
 	{
 		return EReimportResult::Failed;
 	}
 
-	if (GDTFReimport->SourceFilename.IsEmpty() || !FPaths::FileExists(GDTFReimport->SourceFilename))
+	const FString SourceFilename = GDTF->GetGDTFAssetImportData()->GetSourceFilePathAndName();
+	if (!FPaths::FileExists(SourceFilename))
 	{
 		return EReimportResult::Failed;
 	}
 
-	bool OutCanceled = false;
-	if (ImportObject(InObject->GetClass(), InObject->GetOuter(), *InObject->GetName(), RF_Public | RF_Standalone, GDTFReimport->SourceFilename, nullptr, OutCanceled))
+	bool bOutCanceled = false;
+	if (ImportObject(InObject->GetClass(), InObject->GetOuter(), *InObject->GetName(), RF_Public | RF_Standalone, SourceFilename, nullptr, bOutCanceled))
 	{
 		return EReimportResult::Succeeded;
 	}
 
-	return OutCanceled ? EReimportResult::Cancelled : EReimportResult::Failed;
+	return bOutCanceled ? EReimportResult::Cancelled : EReimportResult::Failed;
 }
 
 int32 UDMXGDTFFactory::GetPriority() const
diff --git a/Engine/Plugins/VirtualProduction/DMX/DMXEngine/Source/DMXEditor/Private/Factories/DMXLibraryFromMVRFactory.cpp b/Engine/Plugins/VirtualProduction/DMX/DMXEngine/Source/DMXEditor/Private/Factories/DMXLibraryFromMVRFactory.cpp
--- a/Engine/Plugins/VirtualProduction/DMX/DMXEngine/Source/DMXEditor/Private/Factories/DMXLibraryFromMVRFactory.cpp
+++ b/Engine/Plugins/VirtualProduction/DMX/DMXEngine/Source/DMXEditor/Private/Factories/DMXLibraryFromMVRFactory.cpp
@@ -8,8 +8,10 @@
 #include "Factories/DMXGDTFFactory.h"
 #include "Library/DMXEntityFixturePatch.h"
 #include "Library/DMXEntityFixtureType.h"
+#include "Library/DMXGDTFAssetImportData.h"
 #include "Library/DMXImportGDTF.h"
 #include "Library/DMXLibrary.h"
+#include "MVR/DMXMVRAssetImportData.h"
 #include "MVR/DMXMVRGeneralSceneDescription.h"
 #include "MVR/DMXMVRUnzip.h"
 
@@ -17,6 +19,7 @@
 #include "AssetToolsModule.h"
 #include "ObjectTools.h"
 #include "XmlFile.h"
+#include "EditorFramework/AssetImportData.h"
 #include "Misc/FileHelper.h"
 #include "Misc/MessageDialog.h"
 
@@ -65,11 +68,12 @@ UObject* UDMXLibraryFromMVRFactory::FactoryCreateFile(UClass* InClass, UObject*
 
 	const FName GeneralSceneDescriptionName = FName(GetName() + TEXT("_MVRGeneralSceneDescription"));
 	UDMXMVRGeneralSceneDescription* GeneralSceneDescription = UDMXMVRGeneralSceneDescription::CreateFromXmlFile(GeneralSceneDescriptionXml, NewDMXLibrary, GeneralSceneDescriptionName);
-
-	if (!ensureAlwaysMsgf(GeneralSceneDescription, TEXT("Failed to create a General Scene Description when importing MVR."), __FUNCTION__))
+	if (!GeneralSceneDescription)
 	{
 		return nullptr;
 	}
+	UDMXMVRAssetImportData* MVRAssetImportData = GeneralSceneDescription->GetMVRAssetImportData();
+	MVRAssetImportData->SetSourceFile(InFilename);
 
 	TArray<UDMXImportGDTF*> GDTFs = CreateGDTFAssets(Parent, Flags, MVRUnzip.ToSharedRef(), *GeneralSceneDescription);
 	InitDMXLibrary(NewDMXLibrary, GDTFs, GeneralSceneDescription);
@@ -133,14 +137,9 @@ TArray<UDMXImportGDTF*> UDMXLibraryFromMVRFactory::CreateGDTFAssets(UObject* Par
 	TArray<UDMXImportGDTF*> ExistingGDTFs;
 	for (const FAssetData& AssetData : ExistingGDTFAssets)
 	{
-		if (UDMXImportGDTF* GDTFAsset = Cast<UDMXImportGDTF>(AssetData.GetAsset()))
+		if (UDMXImportGDTF* GDTF = Cast<UDMXImportGDTF>(AssetData.GetAsset()))
 		{
-			GeneralSceneDescription.GetMVRFixtures().ContainsByPredicate([GDTFAsset](const FDMXMVRFixture& MVRFixture)
-				{
-					return GDTFAsset->SourceFilename.Equals(MVRFixture.GDTFSpec, ESearchCase::IgnoreCase);
-				});
-
-			ExistingGDTFs.AddUnique(GDTFAsset);
+			ExistingGDTFs.AddUnique(GDTF);
 		}
 	}
 
@@ -157,7 +156,12 @@ TArray<UDMXImportGDTF*> UDMXLibraryFromMVRFactory::CreateGDTFAssets(UObject* Par
 				if (GDTFFactory->Reimport(GDTFAsset) == EReimportResult::Succeeded)
 				{
 					ImportedGDTFs.Add(GDTFAsset);
-					ImportedGDTFNames.Add(GDTFAsset->SourceFilename);
+
+					if (UDMXGDTFAssetImportData* GDTFAssetImportData = GDTFAsset->GetGDTFAssetImportData())
+					{
+						const FString SourceFilename = FPaths::GetCleanFilename(GDTFAssetImportData->GetSourceFilePathAndName());
+						ImportedGDTFNames.Add(SourceFilename);
+					}
 				}
 			}
 		}
@@ -220,7 +224,14 @@ void UDMXLibraryFromMVRFactory::InitDMXLibrary(UDMXLibrary* DMXLibrary, const TA
 	TMap<UDMXEntityFixtureType*, FLinearColor> FixtureTypeToColorMap;
 	for (UDMXImportGDTF* GDTF : GDTFAssets)
 	{
-		const FString GDTFFilename = FPaths::GetCleanFilename(GDTF->SourceFilename);
+		UDMXGDTFAssetImportData* GDTFAssetImportData = GDTF->GetGDTFAssetImportData();
+		if (!GDTFAssetImportData)
+		{
+			continue;
+		}
+
+		const FString GDTFSourceFilename = GDTFAssetImportData->GetSourceFilePathAndName();
+		const FString GDTFFilename = FPaths::GetCleanFilename(GDTFSourceFilename);
 		FDMXEntityFixtureTypeConstructionParams FixtureTypeConstructionParams;
 		FixtureTypeConstructionParams.DMXCategory = FDMXFixtureCategory(FDMXFixtureCategory::GetFirstValue());
 		FixtureTypeConstructionParams.ParentDMXLibrary = DMXLibrary;
diff --git a/Engine/Plugins/VirtualProduction/DMX/DMXEngine/Source/DMXRuntime/Private/DMXRuntimeMainStreamObjectVersion.h b/Engine/Plugins/VirtualProduction/DMX/DMXEngine/Source/DMXRuntime/Private/DMXRuntimeMainStreamObjectVersion.h
--- a/Engine/Plugins/VirtualProduction/DMX/DMXEngine/Source/DMXRuntime/Private/DMXRuntimeMainStreamObjectVersion.h
+++ b/Engine/Plugins/VirtualProduction/DMX/DMXEngine/Source/DMXRuntime/Private/DMXRuntimeMainStreamObjectVersion.h
@@ -23,9 +23,6 @@ struct FDMXRuntimeMainStreamObjectVersion
 
 		// DMX Controllers are removed from the pugin
 		DMXControllersAreRemoved,
-		
-		// Update DMX Library to contain an MVR General Scene Description
-		DMXLibraryContainsMVRGeneralSceneDescription,
 
 		// Update Fixture Patch to use an MVR UUID
 		DMXFixturePatchHasMVRUUID,
diff --git a/Engine/Plugins/VirtualProduction/DMX/DMXEngine/Source/DMXRuntime/Private/Library/DMXGDTFAssetImportData.cpp b/Engine/Plugins/VirtualProduction/DMX/DMXEngine/Source/DMXRuntime/Private/Library/DMXGDTFAssetImportData.cpp
new file mode 100644
--- /dev/null
+++ b/Engine/Plugins/VirtualProduction/DMX/DMXEngine/Source/DMXRuntime/Private/Library/DMXGDTFAssetImportData.cpp
@@ -0,0 +1,30 @@
+// Copyright Epic Games, Inc. All Rights Reserved.
+
+#include "Library/DMXGDTFAssetImportData.h"
+
+#include "Misc/FileHelper.h"
+#include "Misc/Paths.h"
+
+#if WITH_EDITOR
+FString UDMXGDTFAssetImportData::GetSourceFilePathAndName() const
+{
+	return GetFirstFilename();
+}
+#endif // WITH_EDITOR
+
+#if WITH_EDITOR
+void UDMXGDTFAssetImportData::SetSourceFile(const FString& FilePathAndName)
+{
+	FAssetImportInfo Info;
+	Info.Insert(FAssetImportInfo::FSourceFile(FilePathAndName));
+	SourceData = Info;
+
+	if (!FPaths::FileExists(FilePathAndName))
+	{
+		return;
+	}
+
+	RawZipFile.Reset();
+	FFileHelper::LoadFileToArray(RawZipFile, *FilePathAndName);
+}
+#endif // WITH_EDITOR
diff --git a/Engine/Plugins/VirtualProduction/DMX/DMXEngine/Source/DMXRuntime/Private/Library/DMXImportGDTF.cpp b/Engine/Plugins/VirtualProduction/DMX/DMXEngine/Source/DMXRuntime/Private/Library/DMXImportGDTF.cpp
--- a/Engine/Plugins/VirtualProduction/DMX/DMXEngine/Source/DMXRuntime/Private/Library/DMXImportGDTF.cpp
+++ b/Engine/Plugins/VirtualProduction/DMX/DMXEngine/Source/DMXRuntime/Private/Library/DMXImportGDTF.cpp
@@ -2,6 +2,11 @@
 
 #include "Library/DMXImportGDTF.h"
 
+#include "Library/DMXGDTFAssetImportData.h"
+
+#include "EditorFramework/AssetImportData.h"
+
+
 bool UDMXImportGDTFAttributeDefinitions::FindFeature(const FString& InQuery, FDMXImportGDTFFeature& OutFeature) const
 {
     if (InQuery.IsEmpty())
@@ -142,3 +147,31 @@ TArray<FDMXImportGDTFChannelFunction> UDMXImportGDTFDMXModes::GetDMXChannelFunct
 	}
 	return MoveTemp(Channels);
 };
+
+UDMXImportGDTF::UDMXImportGDTF()
+{
+#if WITH_EDITORONLY_DATA
+	GDTFAssetImportData = NewObject<UDMXGDTFAssetImportData>(this, TEXT("GDTFAssetImportData"), RF_Public);
+#endif
+}
+
+void UDMXImportGDTF::PostLoad()
+{
+	Super::PostLoad();
+
+#if WITH_EDITORONLY_DATA
+	if (!GDTFAssetImportData)
+	{
+		GDTFAssetImportData = NewObject<UDMXGDTFAssetImportData>(this, TEXT("GDTFAssetImportData"));
+	}
+
+	PRAGMA_DISABLE_DEPRECATION_WARNINGS
+	if (!SourceFilename_DEPRECATED.IsEmpty())
+	{
+		GDTFAssetImportData->SetSourceFile(SourceFilename_DEPRECATED);
+
+		SourceFilename_DEPRECATED.Empty();
+	}
+	PRAGMA_ENABLE_DEPRECATION_WARNINGS
+#endif
+}
diff --git a/Engine/Plugins/VirtualProduction/DMX/DMXEngine/Source/DMXRuntime/Private/Library/DMXLibrary.cpp b/Engine/Plugins/VirtualProduction/DMX/DMXEngine/Source/DMXRuntime/Private/Library/DMXLibrary.cpp
--- a/Engine/Plugins/VirtualProduction/DMX/DMXEngine/Source/DMXRuntime/Private/Library/DMXLibrary.cpp
+++ b/Engine/Plugins/VirtualProduction/DMX/DMXEngine/Source/DMXRuntime/Private/Library/DMXLibrary.cpp
@@ -32,28 +32,15 @@ UDMXLibrary::UDMXLibrary()
 
 void UDMXLibrary::Serialize(FArchive& Ar)
 {
+#if WITH_EDITORONLY_DATA
 	if (Ar.IsSaving())
 	{
 		// Update the General Scene Description before saving it
 		UpdateGeneralSceneDescription();
 	}
+#endif
 
 	Super::Serialize(Ar);
-
-	Ar.UsingCustomVersion(FDMXRuntimeMainStreamObjectVersion::GUID);
-	if (Ar.IsLoading())
-	{
-		if (Ar.CustomVer(FDMXRuntimeMainStreamObjectVersion::GUID) < FDMXRuntimeMainStreamObjectVersion::DMXLibraryContainsMVRGeneralSceneDescription)
-		{
-			if (!GeneralSceneDescription)
-			{
-				const FName GeneralSceneDescriptionName = FName(GetName() + TEXT("_MVRGeneralSceneDescription"));
-				GeneralSceneDescription = UDMXMVRGeneralSceneDescription::CreateFromDMXLibrary(*this, this, GeneralSceneDescriptionName);
-			}
-
-			GeneralSceneDescription->WriteDMXLibraryToGeneralSceneDescription(*this);
-		}
-	}
 }
 
 void UDMXLibrary::PostInitProperties()
@@ -74,6 +61,7 @@ void UDMXLibrary::PostLoad()
 	if (!HasAnyFlags(RF_ClassDefaultObject))
 	{
 #if WITH_EDITOR
+		// Upgrade from controllers to ports
 		bool bNeedsUpgradeFromControllersToPorts = Entities.ContainsByPredicate([](UDMXEntity* Entity) {
 			return Cast<UDMXEntityController>(Entity) != nullptr;
 			});
@@ -81,6 +69,13 @@ void UDMXLibrary::PostLoad()
 		{
 			UpgradeFromControllersToPorts();
 		}
+
+		// Upgrade to contain an MVR General Scene Description
+		if (!GeneralSceneDescription)
+		{
+			const FName GeneralSceneDescriptionName = FName(GetName() + TEXT("_MVRGeneralSceneDescription"));
+			GeneralSceneDescription = UDMXMVRGeneralSceneDescription::CreateFromDMXLibrary(*this, this, GeneralSceneDescriptionName);
+		}
 #endif 
 		UpdatePorts();
 
@@ -587,6 +582,7 @@ void UDMXLibrary::SetMVRGeneralSceneDescription(UDMXMVRGeneralSceneDescription*
 	GeneralSceneDescription = NewGeneralSceneDescription;
 }
 
+#if WITH_EDITOR
 UDMXMVRGeneralSceneDescription* UDMXLibrary::UpdateGeneralSceneDescription()
 {
 	if (ensureAlwaysMsgf(GeneralSceneDescription, TEXT("Trying to update General Scene Description of %s, but the General Scene Description is not valid."), *GetName()))
@@ -597,6 +593,7 @@ UDMXMVRGeneralSceneDescription* UDMXLibrary::UpdateGeneralSceneDescription()
 
 	return nullptr;
 }
+#endif // WITH_EDTIOR
 
 #if WITH_EDITOR
 void UDMXLibrary::UpgradeFromControllersToPorts()
diff --git a/Engine/Plugins/VirtualProduction/DMX/DMXEngine/Source/DMXRuntime/Private/MVR/DMXMVRAssetImportData.cpp b/Engine/Plugins/VirtualProduction/DMX/DMXEngine/Source/DMXRuntime/Private/MVR/DMXMVRAssetImportData.cpp
new file mode 100644
--- /dev/null
+++ b/Engine/Plugins/VirtualProduction/DMX/DMXEngine/Source/DMXRuntime/Private/MVR/DMXMVRAssetImportData.cpp
@@ -0,0 +1,24 @@
+// Copyright Epic Games, Inc. All Rights Reserved.
+
+#include "MVR/DMXMVRAssetImportData.h"
+
+#include "Misc/FileHelper.h"
+#include "Misc/Paths.h"
+
+
+#if WITH_EDITOR
+void UDMXMVRAssetImportData::SetSourceFile(const FString& FilePathAndName)
+{
+	FAssetImportInfo Info;
+	Info.Insert(FAssetImportInfo::FSourceFile(FilePathAndName));
+	SourceData = Info;
+
+	if (!FPaths::FileExists(FilePathAndName))
+	{
+		return;
+	}
+
+	RawZipFile.Reset();
+	FFileHelper::LoadFileToArray(RawZipFile, *FilePathAndName);
+}
+#endif // WITH_EDITOR
diff --git a/Engine/Plugins/VirtualProduction/DMX/DMXEngine/Source/DMXRuntime/Private/MVR/DMXMVRGeneralSceneDescription.cpp b/Engine/Plugins/VirtualProduction/DMX/DMXEngine/Source/DMXRuntime/Private/MVR/DMXMVRGeneralSceneDescription.cpp
--- a/Engine/Plugins/VirtualProduction/DMX/DMXEngine/Source/DMXRuntime/Private/MVR/DMXMVRGeneralSceneDescription.cpp
+++ b/Engine/Plugins/VirtualProduction/DMX/DMXEngine/Source/DMXRuntime/Private/MVR/DMXMVRGeneralSceneDescription.cpp
@@ -6,14 +6,25 @@
 #include "DMXRuntimeLog.h"
 #include "Library/DMXEntityFixturePatch.h"
 #include "Library/DMXEntityFixtureType.h"
+#include "Library/DMXGDTFAssetImportData.h"
 #include "Library/DMXImportGDTF.h"
 #include "Library/DMXLibrary.h"
+#include "MVR/DMXMVRAssetImportData.h"
 
 #include "XmlFile.h"
 #include "XmlNode.h"
+#include "EditorFramework/AssetImportData.h"
 #include "Misc/Paths.h"
 
 
+UDMXMVRGeneralSceneDescription::UDMXMVRGeneralSceneDescription()
+{
+#if WITH_EDITORONLY_DATA
+	MVRAssetImportData = NewObject<UDMXMVRAssetImportData>(this, TEXT("MVRAssetImportData"), RF_Public);
+#endif
+}
+
+#if WITH_EDITOR
 UDMXMVRGeneralSceneDescription* UDMXMVRGeneralSceneDescription::CreateFromXmlFile(TSharedRef<FXmlFile> GeneralSceneDescriptionXml, UObject* Outer, FName Name, EObjectFlags Flags)
 {
 	UDMXMVRGeneralSceneDescription* GeneralSceneDescription = NewObject<UDMXMVRGeneralSceneDescription>(Outer, Name, Flags);
@@ -21,7 +32,9 @@ UDMXMVRGeneralSceneDescription* UDMXMVRGeneralSceneDescription::CreateFromXmlFil
 
 	return GeneralSceneDescription;
 }
+#endif // WITH_EDITOR
 
+#if WITH_EDITOR
 UDMXMVRGeneralSceneDescription* UDMXMVRGeneralSceneDescription::CreateFromDMXLibrary(const UDMXLibrary& DMXLibrary, UObject* Outer, FName Name, EObjectFlags Flags)
 {
 	UDMXMVRGeneralSceneDescription* GeneralSceneDescription = NewObject<UDMXMVRGeneralSceneDescription>(Outer, Name, Flags);
@@ -29,6 +42,7 @@ UDMXMVRGeneralSceneDescription* UDMXMVRGeneralSceneDescription::CreateFromDMXLib
 
 	return GeneralSceneDescription;
 }
+#endif // WITH_EDITOR
 
 FDMXMVRFixture* UDMXMVRGeneralSceneDescription::FindMVRFixture(const FGuid& MVRFixtureUUID)
 {
@@ -54,6 +68,7 @@ void UDMXMVRGeneralSceneDescription::AddMVRFixture(FDMXMVRFixture& MVRFixture)
 	MVRFixtures.Add(MVRFixture);
 }
 
+#if WITH_EDITOR
 void UDMXMVRGeneralSceneDescription::WriteDMXLibraryToGeneralSceneDescription(const UDMXLibrary& DMXLibrary)
 {
 	TArray<FDMXMVRFixture> MVRFixturesInDMXLibrary;
@@ -89,6 +104,7 @@ void UDMXMVRGeneralSceneDescription::WriteDMXLibraryToGeneralSceneDescription(co
 		}
 	}
 }
+#endif // WITH_EDITOR
 
 void UDMXMVRGeneralSceneDescription::Serialize(FArchive& Ar)
 {
@@ -97,6 +113,7 @@ void UDMXMVRGeneralSceneDescription::Serialize(FArchive& Ar)
 	Ar << MVRFixtures;
 }
 
+#if WITH_EDITOR
 void UDMXMVRGeneralSceneDescription::WriteFixturePatchToGeneralSceneDescription(const UDMXEntityFixturePatch& FixturePatch)
 {
 	if (const UDMXLibrary* DMXLibrary = FixturePatch.GetParentLibrary())
@@ -127,13 +144,21 @@ void UDMXMVRGeneralSceneDescription::WriteFixturePatchToGeneralSceneDescription(
 		bool bSetGDTFSpec = false;
 		if (UDMXEntityFixtureType* FixtureType = FixturePatch.GetFixtureType())
 		{
-			if (FixtureType->GDTF)
+			if (UDMXImportGDTF* GDTF = FixtureType->GDTF)
 			{
-				MVRFixture.GDTFSpec = FPaths::GetBaseFilename(FixtureType->GDTF->SourceFilename);
+				const FString SourceFilename = [GDTF]()
+				{
+					if (GDTF && GDTF->GetGDTFAssetImportData())
+					{
+						return GDTF->GetGDTFAssetImportData()->GetSourceFilePathAndName();
+					}
+					return FString();
+				}();
+				MVRFixture.GDTFSpec = FPaths::GetBaseFilename(SourceFilename);
 			}
 
 			const int32 ModeIndex = FixturePatch.GetActiveModeIndex();
-			if (ensureAlwaysMsgf(FixtureType->Modes.IsValidIndex(ModeIndex), TEXT("Trying to write Active Mode of %s to General Scene Description, but the Mode Index is invalid")))
+			if (ensureAlwaysMsgf(FixtureType->Modes.IsValidIndex(ModeIndex), TEXT("Trying to write Active Mode of to General Scene Description, but the Mode Index is invalid")))
 			{
 				MVRFixture.GDTFMode = FixtureType->Modes[ModeIndex].ModeName;
 			}
@@ -156,7 +181,9 @@ void UDMXMVRGeneralSceneDescription::WriteFixturePatchToGeneralSceneDescription(
 		}
 	}
 }
+#endif // WITH_EDITOR
 
+#if WITH_EDITOR
 void UDMXMVRGeneralSceneDescription::ParseGeneralSceneDescriptionXml(const TSharedRef<FXmlFile>& GeneralSceneDescription)
 {
 	if (const FXmlNode* RootNode = GeneralSceneDescription->GetRootNode())
@@ -197,6 +224,7 @@ void UDMXMVRGeneralSceneDescription::ParseGeneralSceneDescriptionXml(const TShar
 		}
 	}
 }
+#endif // WITH_EDITOR
 
 TArray<int32> UDMXMVRGeneralSceneDescription::GetUnitNumbersInUse(const UDMXLibrary& DMXLibrary)
 {
diff --git a/Engine/Plugins/VirtualProduction/DMX/DMXEngine/Source/DMXRuntime/Private/MVR/DMXMVRSceneActor.cpp b/Engine/Plugins/VirtualProduction/DMX/DMXEngine/Source/DMXRuntime/Private/MVR/DMXMVRSceneActor.cpp
--- a/Engine/Plugins/VirtualProduction/DMX/DMXEngine/Source/DMXRuntime/Private/MVR/DMXMVRSceneActor.cpp
+++ b/Engine/Plugins/VirtualProduction/DMX/DMXEngine/Source/DMXRuntime/Private/MVR/DMXMVRSceneActor.cpp
@@ -77,6 +77,7 @@ void ADMXMVRSceneActor::PostRegisterAllComponents()
 {
 	Super::PostRegisterAllComponents();
 
+#if WITH_EDITOR
 	// If the actor was created as a Datasmith Element, set the library from there
 	const FString DMXLibraryPathString = UDatasmithAssetUserData::GetDatasmithUserDataValueForKey(this, TEXT("DMXLibraryPath"));
 	if (!DMXLibraryPathString.IsEmpty() && !DMXLibrary)
@@ -90,6 +91,7 @@ void ADMXMVRSceneActor::PostRegisterAllComponents()
 	}
 
 	EnsureMVRUUIDsForRelatedActors();
+#endif // WITH_EDITOR
 }
 
 #if WITH_EDITOR
@@ -137,6 +139,7 @@ void ADMXMVRSceneActor::EnsureMVRUUIDsForRelatedActors()
 	}
 }
 
+#if WITH_EDITOR
 void ADMXMVRSceneActor::SetDMXLibrary(UDMXLibrary* NewDMXLibrary)
 {
 	if (!ensureAlwaysMsgf(!DMXLibrary, TEXT("Tried to set the DMXLibrary for %s, but it already has one set. Changing the library is not supported."), *GetName()))
@@ -206,6 +209,7 @@ void ADMXMVRSceneActor::SetDMXLibrary(UDMXLibrary* NewDMXLibrary)
 		SpawnMVRActor(ActorClass, FixturePatch, Transform);
 	}
 }
+#endif // WITH_EDITOR
 
 #if WITH_EDITOR
 void ADMXMVRSceneActor::OnMapChange(uint32 MapEventFlags)
diff --git a/Engine/Plugins/VirtualProduction/DMX/DMXEngine/Source/DMXRuntime/Public/Library/DMXGDTFAssetImportData.h b/Engine/Plugins/VirtualProduction/DMX/DMXEngine/Source/DMXRuntime/Public/Library/DMXGDTFAssetImportData.h
new file mode 100644
--- /dev/null
+++ b/Engine/Plugins/VirtualProduction/DMX/DMXEngine/Source/DMXRuntime/Public/Library/DMXGDTFAssetImportData.h
@@ -0,0 +1,40 @@
+// Copyright Epic Games, Inc. All Rights Reserved.
+
+#pragma once
+
+#include "CoreMinimal.h"
+#include "EditorFramework/AssetImportData.h"
+
+#include "DMXGDTFAssetImportData.generated.h"
+
+class UDMXImportGDTF;
+
+
+UCLASS()
+class DMXRUNTIME_API UDMXGDTFAssetImportData 
+	: public UAssetImportData
+{
+	GENERATED_BODY()
+
+public:
+#if WITH_EDITOR
+	/** Returns path and name of the source file the GDTF asset originated from */
+	FString GetSourceFilePathAndName() const;
+
+	/** Creates new asset GDTF Asset Import Data from the File. Returns the GDTF Asset Import Data or nullptr if not a valid GDTF File */
+	void SetSourceFile(const FString& FilePathAndName);
+
+	/** Returns the raw zip file the asset was generated from */
+	FORCEINLINE const TArray<uint8>& GetRawZipFile() const { return RawZipFile; }
+#endif // WITH_EDITOR 
+
+private:
+	/** We need to remember the literal File Name to */
+
+
+#if WITH_EDITORONLY_DATA
+	/** The raw GDTF zip file as byte array */
+	UPROPERTY()
+	TArray<uint8> RawZipFile;
+#endif // WITH_EDITORONLY_DATA
+};
diff --git a/Engine/Plugins/VirtualProduction/DMX/DMXEngine/Source/DMXRuntime/Public/Library/DMXImportGDTF.h b/Engine/Plugins/VirtualProduction/DMX/DMXEngine/Source/DMXRuntime/Public/Library/DMXImportGDTF.h
--- a/Engine/Plugins/VirtualProduction/DMX/DMXEngine/Source/DMXRuntime/Public/Library/DMXImportGDTF.h
+++ b/Engine/Plugins/VirtualProduction/DMX/DMXEngine/Source/DMXRuntime/Public/Library/DMXImportGDTF.h
@@ -3,10 +3,14 @@
 #pragma once
 
 #include "Library/DMXImport.h"
+
 #include "DMXImportGDTF.generated.h"
 
+class UDMXGDTFAssetImportData;
+
 class UTexture2D;
 
+
 UENUM(BlueprintType)
 enum class EDMXImportGDTFType : uint8
 {
@@ -825,19 +829,34 @@ public:
 
 UCLASS(BlueprintType, Blueprintable)
 class DMXRUNTIME_API UDMXImportGDTF
-: public UDMXImport
+	: public UDMXImport
 {
     GENERATED_BODY()
 
 public:
+	/** Constructor */
+	UDMXImportGDTF();
+
+	//~ Begin UObject interface
+	virtual void PostLoad() override;
+	//~ End UObject interface
+
 	UFUNCTION(BlueprintPure, Category = "DMXGDTF|Import Data")
-	UDMXImportGDTFDMXModes* GetDMXModes()
-	{
-		return Cast<UDMXImportGDTFDMXModes>(DMXModes);
-	};
+	UDMXImportGDTFDMXModes* GetDMXModes() const { return Cast<UDMXImportGDTFDMXModes>(DMXModes); }
 
-public:
-	/** The filename of the dmx we were created from. This may not always exist on disk, as we may have previously loaded and cached the font data inside this asset. */
+	/** DEPRECATED 5.1 in favor of AssetImportData */
+	UPROPERTY(Meta = (DeprecatedProperty, DeprecationMessage = "Deprecated in favor of GDTFAssetImportData, see UDMXImportGDTF::GetGDTFAssetImportData."))
+	FString SourceFilename_DEPRECATED;
+
+#if WITH_EDITORONLY_DATA
+	/** Returns GDTF Asset Import Data for this GDTF */
+	FORCEINLINE UDMXGDTFAssetImportData* GetGDTFAssetImportData() const { return GDTFAssetImportData; }
+#endif 
+
+private:
+#if WITH_EDITORONLY_DATA
+	/** The Asset Import Data used to generate the GDTF asset or nullptr, if not generated from a GDTF file */
 	UPROPERTY()
-	FString SourceFilename;
-};
\ No newline at end of file
+	UDMXGDTFAssetImportData* GDTFAssetImportData;
+#endif
+};
diff --git a/Engine/Plugins/VirtualProduction/DMX/DMXEngine/Source/DMXRuntime/Public/Library/DMXLibrary.h b/Engine/Plugins/VirtualProduction/DMX/DMXEngine/Source/DMXRuntime/Public/Library/DMXLibrary.h
--- a/Engine/Plugins/VirtualProduction/DMX/DMXEngine/Source/DMXRuntime/Public/Library/DMXLibrary.h
+++ b/Engine/Plugins/VirtualProduction/DMX/DMXEngine/Source/DMXRuntime/Public/Library/DMXLibrary.h
@@ -210,13 +210,15 @@ public:
 
 	/** 
 	 * Sets the MVR General Scene Description of the Library. 
-	 * Note, this will not add any patches occur in the new General Scene Description. 
+	 * Note, this will not add any patches that occur in the new General Scene Description. 
 	 * Note, this will remove any patches that do not occur in the new General Scene Description. 
 	 */
 	void SetMVRGeneralSceneDescription(UDMXMVRGeneralSceneDescription* NewGeneralSceneDescription);
 
+#if WITH_EDITOR
 	/** Updates the General Scene Description to reflect the Library */
 	UDMXMVRGeneralSceneDescription* UpdateGeneralSceneDescription();
+#endif
 
 	/** 
 	 * Returns the General Scene Description of the Library. 
diff --git a/Engine/Plugins/VirtualProduction/DMX/DMXEngine/Source/DMXRuntime/Public/MVR/DMXMVRAssetImportData.h b/Engine/Plugins/VirtualProduction/DMX/DMXEngine/Source/DMXRuntime/Public/MVR/DMXMVRAssetImportData.h
new file mode 100644
--- /dev/null
+++ b/Engine/Plugins/VirtualProduction/DMX/DMXEngine/Source/DMXRuntime/Public/MVR/DMXMVRAssetImportData.h
@@ -0,0 +1,41 @@
+// Copyright Epic Games, Inc. All Rights Reserved.
+
+#pragma once
+
+#include "CoreMinimal.h"
+#include "EditorFramework/AssetImportData.h"
+
+#include "DMXMVRAssetImportData.generated.h"
+
+class UDMXMVRGeneralSceneDescription;
+
+
+UCLASS()
+class DMXRUNTIME_API UDMXMVRAssetImportData 
+	: public UAssetImportData
+{
+	GENERATED_BODY()
+
+public:
+#if WITH_EDITOR
+	/** Sets the source file and initializes members from it */
+	void SetSourceFile(const FString& FilePathAndName);
+
+	/** Returns the Asset Import Data */
+	FORCEINLINE UAssetImportData* GetAssetImportData() { return AssetImportData;};
+
+	/** Returns the raw zip file the asset was generated from */
+	FORCEINLINE const TArray<uint8>& GetRawZipFile() const { return RawZipFile; }
+#endif // WITH_EDITOR 
+
+private:
+#if WITH_EDITORONLY_DATA
+	/** The Asset Import Data used to generate the GDTF asset or nullptr, if not generated from a GDTF file */
+	UPROPERTY()
+	UAssetImportData* AssetImportData;
+
+	/** The raw GDTF zip file as byte array */
+	UPROPERTY()
+	TArray<uint8> RawZipFile;
+#endif // WITH_EDITORONLY_DATA
+};
diff --git a/Engine/Plugins/VirtualProduction/DMX/DMXEngine/Source/DMXRuntime/Public/MVR/DMXMVRGeneralSceneDescription.h b/Engine/Plugins/VirtualProduction/DMX/DMXEngine/Source/DMXRuntime/Public/MVR/DMXMVRGeneralSceneDescription.h
--- a/Engine/Plugins/VirtualProduction/DMX/DMXEngine/Source/DMXRuntime/Public/MVR/DMXMVRGeneralSceneDescription.h
+++ b/Engine/Plugins/VirtualProduction/DMX/DMXEngine/Source/DMXRuntime/Public/MVR/DMXMVRGeneralSceneDescription.h
@@ -13,6 +13,7 @@
 
 class UDMXEntityFixturePatch;
 class UDMXLibrary;
+class UDMXMVRAssetImportData;
 
 class FXmlFile;
 
@@ -25,11 +26,16 @@ class DMXRUNTIME_API UDMXMVRGeneralSceneDescription
 	GENERATED_BODY()
 
 public:
+	/** Constructor */
+	UDMXMVRGeneralSceneDescription();
+
+#if WITH_EDITOR
 	/** Creates an MVR General Scene Description from an Xml File */
 	static UDMXMVRGeneralSceneDescription* CreateFromXmlFile(TSharedRef<FXmlFile> GeneralSceneDescriptionXml, UObject* Outer, FName Name, EObjectFlags Flags = RF_NoFlags);
 
 	/** Creates an MVR General Scene Description from a DMX Library */
 	static UDMXMVRGeneralSceneDescription* CreateFromDMXLibrary(const UDMXLibrary& DMXLibrary, UObject* Outer, FName Name, EObjectFlags Flags = RF_NoFlags);
+#endif 
 
 	/** Returns a pointer to the MVR Fixture with corresponding UUID or nullptr if not found */
 	FDMXMVRFixture* FindMVRFixture(const FGuid& MVRFixtureUUID);
@@ -37,6 +43,7 @@ public:
 	/** Adds an MVR Fixture to the General Scene Description */
 	void AddMVRFixture(FDMXMVRFixture& MVRFixture);
 
+#if WITH_EDITOR
 	/**
 	 * Writes the Library to the General Scene Description, effectively removing inexisting and adding
 	 * new MVR Fixtures, according to what MVR Fixture UUIDs the Fixture Patches of the Library contain.
@@ -46,21 +53,33 @@ public:
 	/** Returns the MVR Fixtures of the General Scene Description */
 	FORCEINLINE const TArray<FDMXMVRFixture>& GetMVRFixtures() const { return MVRFixtures; }
 
+	/** Returns MVR Asset Import Data for this asset */
+	FORCEINLINE UDMXMVRAssetImportData* GetMVRAssetImportData() const { return MVRAssetImportData; }
+#endif
+
 protected:
 	//~ Begin UObject interface
 	virtual void Serialize(FArchive& Ar) override;
 	//~ End UObject interface
 
 private:
-	/** Parses a General Scene Description Xml File. Only ment to be used for initialization (ensured) */
-	void ParseGeneralSceneDescriptionXml(const TSharedRef<FXmlFile>& GeneralSceneDescriptionXml);
-
+#if WITH_EDITOR
 	/** Writes the Fixture Patch to the General Scene Description, adding a new MVR Fixture if required */
 	void WriteFixturePatchToGeneralSceneDescription(const UDMXEntityFixturePatch& FixturePatch);
 
+	/** Parses a General Scene Description Xml File. Only ment to be used for initialization (ensured) */
+	void ParseGeneralSceneDescriptionXml(const TSharedRef<FXmlFile>& GeneralSceneDescriptionXml);
+#endif
+
 	/** Returns the Unit Numbers currently in use, sorted from lowest to highest */
 	TArray<int32> GetUnitNumbersInUse(const UDMXLibrary& DMXLibrary);
 
 	/** The MVR Fixtures of the General Scene Description */
 	TArray<FDMXMVRFixture> MVRFixtures;
+
+#if WITH_EDITORONLY_DATA
+	/** Import Data for this asset */
+	UPROPERTY()
+	UDMXMVRAssetImportData* MVRAssetImportData;
+#endif
 };
diff --git a/Engine/Plugins/VirtualProduction/DMX/DMXEngine/Source/DMXRuntime/Public/MVR/DMXMVRSceneActor.h b/Engine/Plugins/VirtualProduction/DMX/DMXEngine/Source/DMXRuntime/Public/MVR/DMXMVRSceneActor.h
--- a/Engine/Plugins/VirtualProduction/DMX/DMXEngine/Source/DMXRuntime/Public/MVR/DMXMVRSceneActor.h
+++ b/Engine/Plugins/VirtualProduction/DMX/DMXEngine/Source/DMXRuntime/Public/MVR/DMXMVRSceneActor.h
@@ -48,13 +48,14 @@ public:
 #if WITH_EDITOR
 	virtual void PreEditChange(FProperty* PropertyAboutToChange) override;
 	virtual void PostEditChangeProperty(FPropertyChangedEvent& PropertyChangedEvent) override;
-#endif // WITH_EDITOR
+#endif
 	//~ End AActor interface
 
 public:
+#if WITH_EDITOR
 	/** Sets the dmx library for this MVR actor. Should only be called once, further calls will have no effect and hit an ensure condition */
 	void SetDMXLibrary(UDMXLibrary* NewDMXLibrary);
-
+#endif
 	/** Returns the DMX Library of this MVR Scene Actor */
 	FORCEINLINE UDMXLibrary* GetDMXLibrary() const { return DMXLibrary; }
 
@@ -67,7 +68,7 @@ public:
 	static FName GetRelatedAcctorsPropertyNameChecked() { return GET_MEMBER_NAME_CHECKED(ADMXMVRSceneActor, RelatedActors); }
 	static FName GetGDTFToDefaultActorClassesPropertyNameChecked() { return GET_MEMBER_NAME_CHECKED(ADMXMVRSceneActor, GDTFToDefaultActorClasses); }
 	static FName GetMVRSceneRootPropertyNameChecked() { return GET_MEMBER_NAME_CHECKED(ADMXMVRSceneActor, MVRSceneRoot); }
-#endif // WITH_EDITOR
+#endif
 
 private:
 	/** Set MVR UUIDs for related Actors */
@@ -108,7 +109,7 @@ private:
 #if WITH_EDITORONLY_DATA
 	/** The GDTFToDefaultActorClassMap cached of PreEditChange, to find changes. */
 	TArray<FDMXMVRSceneGDTFToActorClassPair> GDTFToDefaultActorClasses_PreEditChange;
-#endif // WITH_EDITORONLY_DATA
+#endif
 
 	/** The root component to which all actors are attached initially */
 	UPROPERTY(VisibleAnywhere, Category = "Actor", AdvancedDisplay, Meta = (AllowPrivateAccess = true))
