DMX - Fix various smaller issues with the initial MVR Editor submit

- Undo add/duplicate fixtures doesn't cause a fixture patch list refresh
- Fixture Patch rows have colored tints to separate even/odd rows and for selection, causing weird colors
- Conflicting Fixture Patches now correctly show a warning when their channels conflict
- New Fixture Patches aren't assigned to a new universe if they would exceed the available DMX channels

diff --git a/Engine/Plugins/VirtualProduction/DMX/DMXEngine/Source/DMXEditor/Private/DMXEditorStyle.cpp b/Engine/Plugins/VirtualProduction/DMX/DMXEngine/Source/DMXEditor/Private/DMXEditorStyle.cpp
--- a/Engine/Plugins/VirtualProduction/DMX/DMXEngine/Source/DMXEditor/Private/DMXEditorStyle.cpp
+++ b/Engine/Plugins/VirtualProduction/DMX/DMXEngine/Source/DMXEditor/Private/DMXEditorStyle.cpp
@@ -141,11 +141,11 @@ FDMXEditorStyle::FDMXEditorStyle()
 
 	{
 		Set("MVRFixtureList.Row", FTableRowStyle()
-			.SetEvenRowBackgroundBrush(IMAGE_BRUSH("Common/Selection", Icon8x8, FLinearColor(.95f, 1.f, 1.f, .3f)))
-			.SetOddRowBackgroundBrush(IMAGE_BRUSH("Common/Selection", Icon8x8, FLinearColor(.85f, .9f, 9.f, .3f)))
-			.SetEvenRowBackgroundHoveredBrush(IMAGE_BRUSH("Common/Selection", Icon8x8, FLinearColor(.95f, 1.f, 1.f, 0.6f)))
-			.SetOddRowBackgroundHoveredBrush(IMAGE_BRUSH("Common/Selection", Icon8x8, FLinearColor(.85f, 9.f, 9.f, 0.6f)))
-			.SetSelectorFocusedBrush(BORDER_BRUSH("Common/Selector", FMargin(4.f / 16.f), SelectorColor))
+			.SetEvenRowBackgroundBrush(IMAGE_BRUSH("Common/Selection", Icon8x8, FLinearColor(1.f, 1.f, 1.f, .2f)))
+			.SetOddRowBackgroundBrush(IMAGE_BRUSH("Common/Selection", Icon8x8, FLinearColor(1.f, 1.f, 1.f, .3f)))
+			.SetEvenRowBackgroundHoveredBrush(IMAGE_BRUSH("Common/Selection", Icon8x8, FLinearColor(1.f, 1.f, 1.f, .4f)))
+			.SetOddRowBackgroundHoveredBrush(IMAGE_BRUSH("Common/Selection", Icon8x8, FLinearColor(1.f, 1.f, 1.f, .5f)))
+			.SetSelectorFocusedBrush(BORDER_BRUSH("Common/Selector", FMargin(4.f / 16.f), FLinearColor(1.f, 1.f, 1.f, .8f)))
 		);
 	}
 
diff --git a/Engine/Plugins/VirtualProduction/DMX/DMXEngine/Source/DMXEditor/Private/Widgets/FixturePatch/SDMXMVRFixtureList.cpp b/Engine/Plugins/VirtualProduction/DMX/DMXEngine/Source/DMXEditor/Private/Widgets/FixturePatch/SDMXMVRFixtureList.cpp
--- a/Engine/Plugins/VirtualProduction/DMX/DMXEngine/Source/DMXEditor/Private/Widgets/FixturePatch/SDMXMVRFixtureList.cpp
+++ b/Engine/Plugins/VirtualProduction/DMX/DMXEngine/Source/DMXEditor/Private/Widgets/FixturePatch/SDMXMVRFixtureList.cpp
@@ -76,7 +76,7 @@ private:
 			: Item(InItem)
 		{
 			Universe = Item->GetUniverse();
-			AddressRange = TRange<int32>(Item->GetAddress(), Item->GetNumChannels());
+			AddressRange = TRange<int32>(Item->GetAddress(), Item->GetAddress() + Item->GetNumChannels());
 		}
 
 		FText GetConfictsWith(const FItemPatch& Other) const
@@ -253,6 +253,16 @@ SDMXMVRFixtureList::~SDMXMVRFixtureList()
 	SaveHeaderRowSettings();
 }
 
+void SDMXMVRFixtureList::PostUndo(bool bSuccess)
+{
+	RequestListRefresh();
+}
+
+void SDMXMVRFixtureList::PostRedo(bool bSuccess)
+{
+	RequestListRefresh();
+}
+
 void SDMXMVRFixtureList::Construct(const FArguments& InArgs, TWeakPtr<FDMXEditor> InDMXEditor)
 {
 	if (!InDMXEditor.IsValid())
@@ -313,12 +323,18 @@ void SDMXMVRFixtureList::Construct(const FArguments& InArgs, TWeakPtr<FDMXEditor
 
 void SDMXMVRFixtureList::RequestListRefresh()
 {
-	GEditor->GetTimerManager()->SetTimerForNextTick(
-		FTimerDelegate::CreateLambda([this]()
-			{
-				RefreshList();
-			})
-	);
+	if (RequestListRefreshTimerHandle.IsValid())
+	{
+		return;
+	}
+
+	RequestListRefreshTimerHandle =
+		GEditor->GetTimerManager()->SetTimerForNextTick(
+			FTimerDelegate::CreateLambda([this]()
+				{
+					RefreshList();
+				})
+		);
 }
 
 void SDMXMVRFixtureList::EnterFixturePatchNameEditingMode()
@@ -357,6 +373,8 @@ void SDMXMVRFixtureList::OnSearchChanged()
 
 void SDMXMVRFixtureList::RefreshList()
 {	
+	RequestListRefreshTimerHandle.Invalidate();
+
 	SaveHeaderRowSettings();
 
 	// Clear cached data
@@ -940,6 +958,7 @@ void SDMXMVRFixtureList::OnDuplicateItems()
 	TGuardValue<bool>(bChangingDMXLibrary, true);
 
 	const TArray<TSharedPtr<FDMXMVRFixtureListItem>> SelectedItems = ListView->GetSelectedItems();
+
 	FDMXMVRFixtureListItem::DuplicateItems(WeakDMXEditor, SelectedItems);
 
 	RequestListRefresh();
@@ -958,7 +977,7 @@ void SDMXMVRFixtureList::OnDeleteItems()
 
 	const TArray<TSharedPtr<FDMXMVRFixtureListItem>> SelectedItems = ListView->GetSelectedItems();
 
-	// Make a beautiful selection invariant to ordering of the List
+	// Make a meaningful selection invariant to ordering of the List
 	TSharedPtr<FDMXMVRFixtureListItem> NewSelection;
 	for (int32 ItemIndex = 0; ItemIndex < ListSource.Num(); ItemIndex++)
 	{
diff --git a/Engine/Plugins/VirtualProduction/DMX/DMXEngine/Source/DMXEditor/Private/Widgets/FixturePatch/SDMXMVRFixtureList.h b/Engine/Plugins/VirtualProduction/DMX/DMXEngine/Source/DMXEditor/Private/Widgets/FixturePatch/SDMXMVRFixtureList.h
--- a/Engine/Plugins/VirtualProduction/DMX/DMXEngine/Source/DMXEditor/Private/Widgets/FixturePatch/SDMXMVRFixtureList.h
+++ b/Engine/Plugins/VirtualProduction/DMX/DMXEngine/Source/DMXEditor/Private/Widgets/FixturePatch/SDMXMVRFixtureList.h
@@ -3,6 +3,8 @@
 #pragma once
 
 #include "CoreMinimal.h"
+#include "EditorUndoClient.h"
+#include "Engine/EngineTypes.h"
 #include "Widgets/SCompoundWidget.h"
 #include "Widgets/Views/SHeaderRow.h"
 
@@ -21,6 +23,7 @@ class SHeaderRow;
 template <typename ItemType> class SListView;
 class STableViewBase;
 
+
 /** Collumn IDs in the Fixture Patch List */
 struct FDMXMVRFixtureListCollumnIDs
 {
@@ -33,12 +36,12 @@ struct FDMXMVRFixtureListCollumnIDs
 	static const FName Patch;
 };
 
-
 using FDMXMVRFixtureListType = SListView<TSharedPtr<FDMXMVRFixtureListItem>>;
 
 /** Sortable, editable List of Fixture Patches in the library */
 class SDMXMVRFixtureList
 	: public SCompoundWidget
+	, public FSelfRegisteringEditorUndoClient
 {
 public:
 	SLATE_BEGIN_ARGS(SDMXMVRFixtureList)
@@ -51,7 +54,12 @@ public:
 
 	/** Destructor */
 	virtual ~SDMXMVRFixtureList();
-	
+
+	//~Begin EditorUndoClient interface
+	virtual void PostUndo(bool bSuccess) override;
+	virtual void PostRedo(bool bSuccess) override;
+	//~End EditorUndoClient interface
+
 	/** Constructs this widget */
 	void Construct(const FArguments& InArgs, TWeakPtr<FDMXEditor> InDMXEditor);
 
@@ -139,6 +147,9 @@ private:
 	/** Rows of Mode widgets in the List */
 	TArray<TSharedPtr<SDMXMVRFixtureListRow>> Rows;
 
+	/** Timer handle for the Request List Refresh method */
+	FTimerHandle RequestListRefreshTimerHandle;
+
 	/** Shared Data for Fixture Patch Editors */
 	TSharedPtr<FDMXFixturePatchSharedData> FixturePatchSharedData;
 
diff --git a/Engine/Plugins/VirtualProduction/DMX/DMXEngine/Source/DMXEditor/Private/Widgets/FixturePatch/SDMXMVRFixtureListToolbar.cpp b/Engine/Plugins/VirtualProduction/DMX/DMXEngine/Source/DMXEditor/Private/Widgets/FixturePatch/SDMXMVRFixtureListToolbar.cpp
--- a/Engine/Plugins/VirtualProduction/DMX/DMXEngine/Source/DMXEditor/Private/Widgets/FixturePatch/SDMXMVRFixtureListToolbar.cpp
+++ b/Engine/Plugins/VirtualProduction/DMX/DMXEngine/Source/DMXEditor/Private/Widgets/FixturePatch/SDMXMVRFixtureListToolbar.cpp
@@ -328,9 +328,9 @@ void SDMXMVRFixtureListToolbar::OnAddNewMVRFixtureClicked(UDMXEntity* InSelected
 		{
 			const bool bUniverseIsSmaller = FixturePatchA.GetUniverseID() < FixturePatchB.GetUniverseID();
 			const bool bUniverseIsEqual = FixturePatchA.GetUniverseID() == FixturePatchB.GetUniverseID();
-			const bool bAddressIsSmaller = FixturePatchA.GetStartingChannel() < FixturePatchB.GetStartingChannel();
+			const bool bLastAddressIsSmaller = FixturePatchA.GetEndingChannel() < FixturePatchB.GetEndingChannel();
 			
-			return bUniverseIsSmaller || (bUniverseIsEqual && bAddressIsSmaller);
+			return bUniverseIsSmaller || (bUniverseIsEqual && bLastAddressIsSmaller);
 		});
 
 	int32 Universe = 1; 
@@ -338,7 +338,7 @@ void SDMXMVRFixtureListToolbar::OnAddNewMVRFixtureClicked(UDMXEntity* InSelected
 	if (FixturePatches.Num() > 0)
 	{
 		const int32 ChannelSpan = FixtureType->Modes.Num() > 0 ? FixtureType->Modes[0].ChannelSpan : 1;
-		if (FixturePatches.Last()->GetEndingChannel() + 1 + ChannelSpan <= DMX_MAX_ADDRESS)
+		if (FixturePatches.Last()->GetEndingChannel() + ChannelSpan > DMX_MAX_ADDRESS)
 		{
 			Universe = FixturePatches.Last()->GetUniverseID() + 1;
 			Address = 1;
