DMX - Properly deprecate the DMXEntityFixtureType::DMXImport property

diff --git a/Engine/Plugins/VirtualProduction/DMX/DMXEngine/Source/DMXRuntime/Private/DMXRuntimeMainStreamObjectVersion.h b/Engine/Plugins/VirtualProduction/DMX/DMXEngine/Source/DMXRuntime/Private/DMXRuntimeMainStreamObjectVersion.h
--- a/Engine/Plugins/VirtualProduction/DMX/DMXEngine/Source/DMXRuntime/Private/DMXRuntimeMainStreamObjectVersion.h
+++ b/Engine/Plugins/VirtualProduction/DMX/DMXEngine/Source/DMXRuntime/Private/DMXRuntimeMainStreamObjectVersion.h
@@ -27,9 +27,6 @@ struct FDMXRuntimeMainStreamObjectVersion
 		// Update Fixture Patch to use an MVR UUID
 		DMXFixturePatchHasMVRUUID,
 
-		// Update DMX Entity Fixture Type to use the correct GDTF type instead of its base class
-		DMXFixtureTypeUsesCorrectGDTFType,
-
 		VersionPlusOne,
 		LatestVersion = VersionPlusOne - 1
 	};
diff --git a/Engine/Plugins/VirtualProduction/DMX/DMXEngine/Source/DMXRuntime/Private/Library/DMXEntityFixtureType.cpp b/Engine/Plugins/VirtualProduction/DMX/DMXEngine/Source/DMXRuntime/Private/Library/DMXEntityFixtureType.cpp
--- a/Engine/Plugins/VirtualProduction/DMX/DMXEngine/Source/DMXRuntime/Private/Library/DMXEntityFixtureType.cpp
+++ b/Engine/Plugins/VirtualProduction/DMX/DMXEngine/Source/DMXRuntime/Private/Library/DMXEntityFixtureType.cpp
@@ -238,19 +238,27 @@ void UDMXEntityFixtureType::Serialize(FArchive& Ar)
 				PRAGMA_ENABLE_DEPRECATION_WARNINGS
 			}
 		}
+	}
+}
+
+void UDMXEntityFixtureType::PostLoad()
+{
+	Super::PostLoad();
 
-		if (Ar.CustomVer(FDMXRuntimeMainStreamObjectVersion::GUID) < FDMXRuntimeMainStreamObjectVersion::DMXFixtureTypeUsesCorrectGDTFType)
+#if WITH_EDITOR
+	// Upgrade the GDTF type from the base DMXImport to the actual DMXImportGDTF type which was always used.
+	// Name the new variable 'GDTF' to make the purpose of the variable more clear, namely holding a GDTF asset.
+	PRAGMA_DISABLE_DEPRECATION_WARNINGS
+	if(DMXImport_DEPRECATED)
+	{
+		if (UDMXImportGDTF* OldGDTF = Cast<UDMXImportGDTF>(DMXImport_DEPRECATED))
 		{
-			// Upgrade the GDTF type from the base DMXImport to the actual DMXImportGDTF type which was always used.
-			// Name the new variable 'GDTF' to make the purpose of the variable more clear, namely holding a GDTF asset.
-			PRAGMA_DISABLE_DEPRECATION_WARNINGS
-			if (UDMXImportGDTF* OldGDTF = Cast<UDMXImportGDTF>(DMXImport))
-			{
-				GDTF = OldGDTF;
-			}
-			PRAGMA_ENABLE_DEPRECATION_WARNINGS
+			GDTF = OldGDTF;
+			DMXImport_DEPRECATED = nullptr;
 		}
 	}
+	PRAGMA_ENABLE_DEPRECATION_WARNINGS
+#endif
 }
 
 #if WITH_EDITOR
diff --git a/Engine/Plugins/VirtualProduction/DMX/DMXEngine/Source/DMXRuntime/Public/Library/DMXEntityFixtureType.h b/Engine/Plugins/VirtualProduction/DMX/DMXEngine/Source/DMXRuntime/Public/Library/DMXEntityFixtureType.h
--- a/Engine/Plugins/VirtualProduction/DMX/DMXEngine/Source/DMXRuntime/Public/Library/DMXEntityFixtureType.h
+++ b/Engine/Plugins/VirtualProduction/DMX/DMXEngine/Source/DMXRuntime/Public/Library/DMXEntityFixtureType.h
@@ -302,9 +302,8 @@ public:
 	static void RemoveFixtureTypeFromLibrary(FDMXEntityFixtureTypeRef FixtureTypeRef);
 
 	//~ Begin UObject interface
-protected:
 	virtual void Serialize(FArchive& Ar) override;
-public:
+	virtual void PostLoad() override;
 #if WITH_EDITOR
 	virtual bool Modify(bool bAlwaysMarkDirty = true) override;
 	virtual void PostEditChangeProperty(FPropertyChangedEvent& PropertyChangedEvent) override;
@@ -326,10 +325,12 @@ public:
 	/** Returns a delegate that is and should be broadcast whenever a Fixture Type changed */
 	static FDMXOnFixtureTypeChangedDelegate& GetOnFixtureTypeChanged();
 
+#if WITH_EDITORONLY_DATA
 	/** DEPRECATED 5.0 */
 	UE_DEPRECATED(5.1, "Deprecated in favor of the GDTF property, which is of the GDTF type instead of just a base class.")
-	UPROPERTY(EditAnywhere, BlueprintReadOnly, Category = "Fixture Settings", Meta = (DeprecatedProperty, DeprecationMessage = "Deprecated in favor of the GDTF property, which is of the GDTF type instead of just a base class."))
-	UDMXImport* DMXImport;
+	UPROPERTY(Meta = (DeprecatedProperty, DeprecationMessage = "Deprecated in favor of the GDTF property, which is of the GDTF type instead of just a base class."))
+	UDMXImport* DMXImport_DEPRECATED;
+#endif 
 
 	/** The GDTF file from which the Fixture Type was setup */
 	UPROPERTY(EditAnywhere, BlueprintReadOnly, Category = "Fixture Settings")
