DMX - Fix MVR Universe is held as an Index but for UE Universe is held by ID (e.g. MVR Universe 0 is UE Universe 1)

diff --git a/Engine/Plugins/VirtualProduction/DMX/DMXEngine/Source/DMXEditor/Private/Factories/DMXLibraryFromMVRFactory.cpp b/Engine/Plugins/VirtualProduction/DMX/DMXEngine/Source/DMXEditor/Private/Factories/DMXLibraryFromMVRFactory.cpp
--- a/Engine/Plugins/VirtualProduction/DMX/DMXEngine/Source/DMXEditor/Private/Factories/DMXLibraryFromMVRFactory.cpp
+++ b/Engine/Plugins/VirtualProduction/DMX/DMXEngine/Source/DMXEditor/Private/Factories/DMXLibraryFromMVRFactory.cpp
@@ -256,8 +256,8 @@ void UDMXLibraryFromMVRFactory::InitDMXLibrary(UDMXLibrary* DMXLibrary, const TA
 				const FDMXFixtureMode* ActiveModePtr = FixturePatch->GetActiveMode();
 				return
 					FixturePatch->GetFixtureType() == FixtureType &&
-					FixturePatch->GetUniverseID() == MVRFixture.Addresses.Universe &&
-					FixturePatch->GetStartingChannel() == MVRFixture.Addresses.Address &&
+					FixturePatch->GetUniverseID() == MVRFixture.GetUniverseID() &&
+					FixturePatch->GetStartingChannel() == MVRFixture.GetStartingChannel() &&
 					ActiveModePtr && ActiveModePtr->ModeName == MVRFixture.GDTFMode;
 			});
 
@@ -286,8 +286,8 @@ void UDMXLibraryFromMVRFactory::InitDMXLibrary(UDMXLibrary* DMXLibrary, const TA
 				FDMXEntityFixturePatchConstructionParams FixturePatchConstructionParams;
 				FixturePatchConstructionParams.ActiveMode = ActiveModeIndex;
 				FixturePatchConstructionParams.FixtureTypeRef = FDMXEntityFixtureTypeRef(FixtureType);
-				FixturePatchConstructionParams.UniverseID = MVRFixture.Addresses.Universe;
-				FixturePatchConstructionParams.StartingAddress = MVRFixture.Addresses.Address;
+				FixturePatchConstructionParams.UniverseID = MVRFixture.GetUniverseID();
+				FixturePatchConstructionParams.StartingAddress = MVRFixture.GetStartingChannel();
 				FixturePatchConstructionParams.MVRFixtureUUID = MVRFixture.UUID;
 
 				UDMXEntityFixturePatch* FixturePatch = UDMXEntityFixturePatch::CreateFixturePatchInLibrary(FixturePatchConstructionParams, MVRFixture.Name);
diff --git a/Engine/Plugins/VirtualProduction/DMX/DMXEngine/Source/DMXRuntime/Private/MVR/DMXMVRGeneralSceneDescription.cpp b/Engine/Plugins/VirtualProduction/DMX/DMXEngine/Source/DMXRuntime/Private/MVR/DMXMVRGeneralSceneDescription.cpp
--- a/Engine/Plugins/VirtualProduction/DMX/DMXEngine/Source/DMXRuntime/Private/MVR/DMXMVRGeneralSceneDescription.cpp
+++ b/Engine/Plugins/VirtualProduction/DMX/DMXEngine/Source/DMXRuntime/Private/MVR/DMXMVRGeneralSceneDescription.cpp
@@ -138,8 +138,8 @@ void UDMXMVRGeneralSceneDescription::WriteFixturePatchToGeneralSceneDescription(
 			MVRFixture.UnitNumber = UnitNumbersInUse.Last() + 1;
 		}
 
-		MVRFixture.Addresses.Universe = FixturePatch.GetUniverseID();
-		MVRFixture.Addresses.Address = FixturePatch.GetStartingChannel();
+		MVRFixture.SetUniverseID(FixturePatch.GetUniverseID());
+		MVRFixture.SetStartingChannel(FixturePatch.GetStartingChannel());
 
 		bool bSetGDTFSpec = false;
 		if (UDMXEntityFixtureType* FixtureType = FixturePatch.GetFixtureType())
diff --git a/Engine/Plugins/VirtualProduction/DMX/DMXEngine/Source/DMXRuntime/Public/MVR/DMXMVRFixture.h b/Engine/Plugins/VirtualProduction/DMX/DMXEngine/Source/DMXRuntime/Public/MVR/DMXMVRFixture.h
--- a/Engine/Plugins/VirtualProduction/DMX/DMXEngine/Source/DMXRuntime/Public/MVR/DMXMVRFixture.h
+++ b/Engine/Plugins/VirtualProduction/DMX/DMXEngine/Source/DMXRuntime/Public/MVR/DMXMVRFixture.h
@@ -148,10 +148,24 @@ public:
 
 	/** The unit number of the lighting fixture in a position. */
 	int32 UnitNumber = 0;
+	
+	/** Sets the Universe ID, in UE form */
+	FORCEINLINE void SetUniverseID(int32 UniverseID) { Addresses.Universe = UniverseID - 1; }
+
+	/** Returns the Universe ID, in UE form */
+	FORCEINLINE int32 GetUniverseID() const { return Addresses.Universe + 1; }
+
+	/** Sets the Address of the MVR Fixture */
+	FORCEINLINE void SetStartingChannel(int32 Address) { Addresses.Address = Address; }
 
+	/** Returns the Address of the MVR Fixture */
+	FORCEINLINE int32 GetStartingChannel() const { return Addresses.Address; }
+
+private:
 	/** The container for DMX Addresses for this fixture. */
 	FDMXMVRFixtureAddresses Addresses;
 
+public:
 	/** A color assigned to a fixture. If it is not defined, there is no color for the fixture. */
 	TOptional<FDMXMVRColorCIE> CIEColor;
 
