DMX - Add DMX icon for Sequencer Track Menu

diff --git a/Engine/Plugins/VirtualProduction/DMX/DMXEngine/Source/DMXEditor/Private/Sequencer/DMXLibraryTrackEditor.cpp b/Engine/Plugins/VirtualProduction/DMX/DMXEngine/Source/DMXEditor/Private/Sequencer/DMXLibraryTrackEditor.cpp
--- a/Engine/Plugins/VirtualProduction/DMX/DMXEngine/Source/DMXEditor/Private/Sequencer/DMXLibraryTrackEditor.cpp
+++ b/Engine/Plugins/VirtualProduction/DMX/DMXEngine/Source/DMXEditor/Private/Sequencer/DMXLibraryTrackEditor.cpp
@@ -123,9 +123,8 @@ void FDMXLibraryTrackEditor::BuildAddTrackMenu(FMenuBuilder& MenuBuilder)
 		LOCTEXT("AddDMXTrack", "DMX Library Track"),
 		LOCTEXT("AddDMXTrackToolTip", "Adds a new track that controls a DMX Library's Patches functions."),
 		FNewMenuDelegate::CreateLambda(SubMenuCallback),
-		false//,
-		//FSlateIconFinder::FindIconForClass(UDMXLibrary::StaticClass())
-		// TODO Add DMX Library icon
+		false,
+		FSlateIconFinder::FindIconForClass(UDMXLibrary::StaticClass())
 	);
 }
 
