DMX - Fix MVR Fixture List no longer clears selection when an empty area is clicked

diff --git a/Engine/Plugins/VirtualProduction/DMX/DMXEngine/Source/DMXEditor/Private/Widgets/FixturePatch/SDMXMVRFixtureList.cpp b/Engine/Plugins/VirtualProduction/DMX/DMXEngine/Source/DMXEditor/Private/Widgets/FixturePatch/SDMXMVRFixtureList.cpp
--- a/Engine/Plugins/VirtualProduction/DMX/DMXEngine/Source/DMXEditor/Private/Widgets/FixturePatch/SDMXMVRFixtureList.cpp
+++ b/Engine/Plugins/VirtualProduction/DMX/DMXEngine/Source/DMXEditor/Private/Widgets/FixturePatch/SDMXMVRFixtureList.cpp
@@ -464,26 +464,52 @@ void SDMXMVRFixtureList::OnSelectionChanged(TSharedPtr<FDMXMVRFixtureListItem> I
 		return;
 	}
  
-	TArray<TWeakObjectPtr<UDMXEntityFixturePatch>> FixturePatches;
 	const TArray<TSharedPtr<FDMXMVRFixtureListItem>> SelectedItems = ListView->GetSelectedItems();
-	for (const TSharedPtr<FDMXMVRFixtureListItem>& Item : SelectedItems)
+	if (SelectedItems.Num() > 0)
 	{
-		if (UDMXEntityFixturePatch* FixturePatch = Item->GetFixturePatch())
+		TArray<TWeakObjectPtr<UDMXEntityFixturePatch>> FixturePatchesToSelect;
+		for (const TSharedPtr<FDMXMVRFixtureListItem>& Item : SelectedItems)
 		{
-			FixturePatches.AddUnique(FixturePatch);
+			if (UDMXEntityFixturePatch* FixturePatch = Item->GetFixturePatch())
+			{
+				FixturePatchesToSelect.AddUnique(FixturePatch);
+			}
 		}
-	}
 
-	FixturePatchSharedData->SelectFixturePatches(FixturePatches);
+		FixturePatchSharedData->SelectFixturePatches(FixturePatchesToSelect);
 
-	if (SelectedItems.Num() > 0)
+		if (SelectedItems.Num() > 0)
+		{
+			const int32 SelectedUniverse = FixturePatchSharedData->GetSelectedUniverse();
+			const int32 UniverseOfFirstItem = SelectedItems[0]->GetUniverse();
+			if (SelectedUniverse != UniverseOfFirstItem)
+			{
+				FixturePatchSharedData->SelectUniverse(UniverseOfFirstItem);
+			}
+		}
+	}
+	else
 	{
-		const int32 SelectedUniverse = FixturePatchSharedData->GetSelectedUniverse();
-		const int32 UniverseOfFirstItem = SelectedItems[0]->GetUniverse();
-		if (SelectedUniverse != UniverseOfFirstItem)
+		// Restore selection when nothing got selected
+		const TArray<TWeakObjectPtr<UDMXEntityFixturePatch>> SelectedFixturePatches = FixturePatchSharedData->GetSelectedFixturePatches();
+		TArray<TSharedPtr<FDMXMVRFixtureListItem>> ItemsToSelect;
+		for (TWeakObjectPtr<UDMXEntityFixturePatch> FixturePatch : SelectedFixturePatches)
 		{
-			FixturePatchSharedData->SelectUniverse(UniverseOfFirstItem);
+			if (FixturePatch.IsValid())
+			{
+				const TSharedPtr<FDMXMVRFixtureListItem>* ItemPtr = ListSource.FindByPredicate([&FixturePatch](const TSharedPtr<FDMXMVRFixtureListItem>& Item)
+					{
+						return Item->GetFixturePatch() == FixturePatch;
+					});
+				if (ItemPtr)
+				{
+					ItemsToSelect.Add(*ItemPtr);
+				}
+			}
 		}
+		ListView->ClearSelection();
+		constexpr bool bSelected = true;
+		ListView->SetItemSelection(ItemsToSelect, bSelected);
 	}
 }
 
