DMX localization key duplicate fix

diff --git a/Engine/Plugins/VirtualProduction/DMX/DMXEngine/Source/DMXEditor/Private/Widgets/FixtureType/SDMXFixtureFunctionEditor.cpp b/Engine/Plugins/VirtualProduction/DMX/DMXEngine/Source/DMXEditor/Private/Widgets/FixtureType/SDMXFixtureFunctionEditor.cpp
--- a/Engine/Plugins/VirtualProduction/DMX/DMXEngine/Source/DMXEditor/Private/Widgets/FixtureType/SDMXFixtureFunctionEditor.cpp
+++ b/Engine/Plugins/VirtualProduction/DMX/DMXEngine/Source/DMXEditor/Private/Widgets/FixtureType/SDMXFixtureFunctionEditor.cpp
@@ -55,11 +55,11 @@ void SDMXFixtureFunctionEditor::NotifyPreChange(FProperty* PropertyAboutToChange
 		}
 		else if (PropertyName == GET_MEMBER_NAME_CHECKED(FDMXFixtureFunction, DataType))
 		{
-			Transaction = MakeUnique<FScopedTransaction>(LOCTEXT("ChangeDefaultValueTransaction", "Data Type of Function"));
+			Transaction = MakeUnique<FScopedTransaction>(LOCTEXT("SetDataTypeTransaction", "Data Type of Function"));
 		}
 		else if (PropertyName == GET_MEMBER_NAME_CHECKED(FDMXFixtureFunction, DefaultValue))
 		{
-			Transaction = MakeUnique<FScopedTransaction>(LOCTEXT("ChangeDefaultValueTransaction2", "Default Value of Function"));
+			Transaction = MakeUnique<FScopedTransaction>(LOCTEXT("SetDefaultValueTransaction", "Default Value of Function"));
 		}
 
 		FixtureType->Modify();
