DMX - Update MustImplement SoftClassPath of ActorClass member in DMXMVRSceneActor to the full path of DMXRuntime's MVRFixtureActorInterface, as required with 5.1

diff --git a/Engine/Plugins/VirtualProduction/DMX/DMXEngine/Source/DMXRuntime/Public/MVR/DMXMVRSceneActor.h b/Engine/Plugins/VirtualProduction/DMX/DMXEngine/Source/DMXRuntime/Public/MVR/DMXMVRSceneActor.h
--- a/Engine/Plugins/VirtualProduction/DMX/DMXEngine/Source/DMXRuntime/Public/MVR/DMXMVRSceneActor.h
+++ b/Engine/Plugins/VirtualProduction/DMX/DMXEngine/Source/DMXRuntime/Public/MVR/DMXMVRSceneActor.h
@@ -25,7 +25,7 @@ struct DMXRUNTIME_API FDMXMVRSceneGDTFToActorClassPair
 	TSoftObjectPtr<UDMXImportGDTF> GDTF;
 	
 	/** The Actor Class that should be or was spawned */
-	UPROPERTY(EditAnywhere, BlueprintReadOnly, Category = MVR, NoClear, Meta = (MustImplement = "DMXMVRFixtureActorInterface"))
+	UPROPERTY(EditAnywhere, BlueprintReadOnly, Category = MVR, NoClear, Meta = (MustImplement = "/Script/DMXRuntime.DMXMVRFixtureActorInterface"))
 	TSoftClassPtr<AActor> ActorClass;
 };
 
