DMX - Add a status warning when a patch exceeds the available channels in a Universe

diff --git a/Engine/Plugins/VirtualProduction/DMX/DMXEngine/Source/DMXEditor/Private/Widgets/FixturePatch/SDMXMVRFixtureList.cpp b/Engine/Plugins/VirtualProduction/DMX/DMXEngine/Source/DMXEditor/Private/Widgets/FixturePatch/SDMXMVRFixtureList.cpp
--- a/Engine/Plugins/VirtualProduction/DMX/DMXEngine/Source/DMXEditor/Private/Widgets/FixturePatch/SDMXMVRFixtureList.cpp
+++ b/Engine/Plugins/VirtualProduction/DMX/DMXEngine/Source/DMXEditor/Private/Widgets/FixturePatch/SDMXMVRFixtureList.cpp
@@ -36,22 +36,31 @@ public:
 		: Items(InItems)
 	{}
 
-	/** Generates warning texts. Returns a map of those Items that need a warning set along with the Warning Text */
+	/** Generates warning texts. Returns a map of those Items that need a warning set along with the warning Text */
 	TMap<TSharedPtr<FDMXMVRFixtureListItem>, FText> GenerateWarningTexts() const
 	{
-		TMap<TSharedPtr<FDMXMVRFixtureListItem>, FText> WarningTextMap;
+		TMap<TSharedPtr<FDMXMVRFixtureListItem>, FText> AccumulatedConflicts;
 
-		return WarningTextMap;
-	}
-
-	/** Generates warning texts. Returns a map of those Items that need a warning set along with the Error Text */
-	TMap<TSharedPtr<FDMXMVRFixtureListItem>, FText> GenerateErrorTexts() const
-	{
 		TMap<TSharedPtr<FDMXMVRFixtureListItem>, FText> UnitNumberConflicts = GetUnitNumberConflicts();
-		TMap<TSharedPtr<FDMXMVRFixtureListItem>, FText> ChannelConflicts = GetChannelConflicts();
+		AccumulatedConflicts.Append(UnitNumberConflicts);
+
+		TMap<TSharedPtr<FDMXMVRFixtureListItem>, FText> ChannelExcessConflicts = GetChannelExcessConflicts();
+		for (const TTuple<TSharedPtr<FDMXMVRFixtureListItem>, FText>& ChannelConflict : ChannelExcessConflicts)
+		{
+			if (AccumulatedConflicts.Contains(ChannelConflict.Key))
+			{
+				const FText LineTerminator = FText::FromString(LINE_TERMINATOR);
+				const FText AccumulatedErrorText = FText::Format(FText::FromString(TEXT("{0}{1}{2}{3}")), AccumulatedConflicts[ChannelConflict.Key], LineTerminator, LineTerminator, ChannelConflict.Value);
+				AccumulatedConflicts[ChannelConflict.Key] = AccumulatedErrorText;
+			}
+			else
+			{
+				AccumulatedConflicts.Add(ChannelConflict.Key, ChannelConflict.Value);
+			}
+		}
 
-		TMap<TSharedPtr<FDMXMVRFixtureListItem>, FText> AccumulatedConflicts = UnitNumberConflicts;
-		for (const TTuple<TSharedPtr<FDMXMVRFixtureListItem>, FText>& ChannelConflict : ChannelConflicts)
+		TMap<TSharedPtr<FDMXMVRFixtureListItem>, FText> ChannelOverlapConflicts = GetChannelOverlapConflicts();
+		for (const TTuple<TSharedPtr<FDMXMVRFixtureListItem>, FText>& ChannelConflict : ChannelOverlapConflicts)
 		{
 			if (AccumulatedConflicts.Contains(ChannelConflict.Key))
 			{
@@ -79,7 +88,7 @@ private:
 			AddressRange = TRange<int32>(Item->GetAddress(), Item->GetAddress() + Item->GetNumChannels());
 		}
 
-		FText GetConfictsWith(const FItemPatch& Other) const
+		FText GetConfictsWithOther(const FItemPatch& Other) const
 		{
 			// No conflict with self
 			if (Other.Item == Item)
@@ -120,18 +129,18 @@ private:
 			{
 				// Modes confict
 				check(Item->GetModeIndex() != Other.Item->GetModeIndex());
-				return FText::Format(LOCTEXT("ModeConflict", "Using same Address and Fixture Type as Fixture {1}, but Modes differ."), UnitNumberText, OtherUnitNumberText);
+				return FText::Format(LOCTEXT("ModeConflict", "Uses same Address and Fixture Type as Fixture {1}, but Modes differ."), UnitNumberText, OtherUnitNumberText);
 			}
 			else if (AddressRange.GetLowerBound() == Other.AddressRange.GetLowerBound())
 			{
 				// Fixture Types conflict
 				check(Item->GetFixtureType() != Other.Item->GetFixtureType());
-				return FText::Format(LOCTEXT("FixtureTypeConflict", "Using same Address as Fixture {1}, but Fixture Types differ."), UnitNumberText, OtherUnitNumberText);
+				return FText::Format(LOCTEXT("FixtureTypeConflict", "Uses same Address as Fixture {1}, but Fixture Types differ."), UnitNumberText, OtherUnitNumberText);
 			}
 			else
 			{
 				// Addresses conflict
-				return FText::Format(LOCTEXT("AddressConflict", "Overlapping Addresses with Fixture {1}"), UnitNumberText, OtherUnitNumberText);
+				return FText::Format(LOCTEXT("AddressConflict", "Overlaps Addresses with Fixture {1}"), UnitNumberText, OtherUnitNumberText);
 			}
 		}
 
@@ -145,8 +154,36 @@ private:
 		TSharedPtr<FDMXMVRFixtureListItem> Item;
 	};
 
-	/** Returns an Map of Items to Channel conflict Texts */
-	TMap<TSharedPtr<FDMXMVRFixtureListItem>, FText> GetChannelConflicts() const
+	/** Returns a Map of Items to Channels exceeding the DMX address range Texts */
+	TMap<TSharedPtr<FDMXMVRFixtureListItem>, FText> GetChannelExcessConflicts() const
+	{
+		TMap<TSharedPtr<FDMXMVRFixtureListItem>, FText> ItemToConflictMap;
+		for (const TSharedPtr<FDMXMVRFixtureListItem>& Item : Items)
+		{
+			const int32 EndingAddress = Item->GetAddress() + Item->GetNumChannels() - 1;
+			if (Item->GetAddress() < 1 &&
+				EndingAddress > DMX_MAX_ADDRESS)
+			{
+				const FText ConflictText = FText::Format(LOCTEXT("ChannelExcessConflict", "Exceeds available DMX Address range. Staring Address is {0} but min Address is 1. Ending Address is {1} but max Address is 512."), Item->GetAddress(), EndingAddress);
+				ItemToConflictMap.Add(Item, ConflictText);
+			}
+			else if (Item->GetAddress() < 1)
+			{
+				const FText ConflictText = FText::Format(LOCTEXT("ChannelExcessConflict", "Exceeds available DMX Address range. Staring Address is {0} but min Address is 1."), Item->GetAddress());
+				ItemToConflictMap.Add(Item, ConflictText);
+			}
+			else if (EndingAddress > DMX_MAX_ADDRESS)
+			{
+				const FText ConflictText = FText::Format(LOCTEXT("ChannelExcessConflict", "Exceeds available DMX Address range. Ending Address is {0} but max Address is 512."), EndingAddress);
+				ItemToConflictMap.Add(Item, ConflictText);
+			}			
+		}
+
+		return ItemToConflictMap;
+	}
+
+	/** Returns a Map of Items to overlapping Channel conflict Texts */
+	TMap<TSharedPtr<FDMXMVRFixtureListItem>, FText> GetChannelOverlapConflicts() const
 	{
 		TArray<FItemPatch> ItemPatches;
 		ItemPatches.Reserve(Items.Num());
@@ -161,17 +198,17 @@ private:
 		{
 			for (const FItemPatch& Other : ItemPatches)
 			{
-				const FText ConflictText = ItemPatch.GetConfictsWith(Other);
-				if (!ConflictText.IsEmpty())
+				const FText ConflictWithOtherText = ItemPatch.GetConfictsWithOther(Other);
+				if (!ConflictWithOtherText.IsEmpty())
 				{
 					if (ItemToConflictMap.Contains(ItemPatch.GetItem()))
 					{
-						FText AppendConflictText = FText::Format(FText::FromString(TEXT("{0}{1}{2}")), ItemToConflictMap[ItemPatch.GetItem()], FText::FromString(FString(LINE_TERMINATOR)), ConflictText);
+						FText AppendConflictText = FText::Format(FText::FromString(TEXT("{0}{1}{2}")), ItemToConflictMap[ItemPatch.GetItem()], FText::FromString(FString(LINE_TERMINATOR)), ConflictWithOtherText);
 						ItemToConflictMap[ItemPatch.GetItem()] = AppendConflictText;
 					}
 					else
 					{
-						ItemToConflictMap.Add(ItemPatch.GetItem(), ConflictText);
+						ItemToConflictMap.Add(ItemPatch.GetItem(), ConflictWithOtherText);
 					}
 				}
 			}
@@ -434,12 +471,6 @@ void SDMXMVRFixtureList::GenereateStatusText()
 	{
 		ItemToWarningTextPair.Key->WarningStatusText = ItemToWarningTextPair.Value;
 	}
-
-	const TMap<TSharedPtr<FDMXMVRFixtureListItem>, FText> ErrorTextMap = StatusTextGenerator.GenerateErrorTexts();
-	for (const TTuple<TSharedPtr<FDMXMVRFixtureListItem>, FText>& ItemToErrorTextPair : ErrorTextMap)
-	{
-		ItemToErrorTextPair.Key->WarningStatusText = ItemToErrorTextPair.Value;
-	}
 }
 
 TSharedRef<ITableRow> SDMXMVRFixtureList::OnGenerateRow(TSharedPtr<FDMXMVRFixtureListItem> InItem, const TSharedRef<STableViewBase>& OwnerTable)
@@ -851,6 +882,11 @@ TSharedPtr<SWidget> SDMXMVRFixtureList::OnContextMenuOpening()
 
 void SDMXMVRFixtureList::RegisterCommands()
 {
+	if (CommandList.IsValid())
+	{
+		return;
+	}
+
 	// listen to common editor shortcuts for copy/paste etc
 	CommandList = MakeShared<FUICommandList>();
 	CommandList->MapAction(FGenericCommands::Get().Cut, 
