DMX - Rework Fixture Patcher to be in support of MVR

diff --git a/Engine/Plugins/VirtualProduction/DMX/DMXEngine/Source/DMXEditor/Private/Customizations/DMXEntityFixturePatchDetails.cpp b/Engine/Plugins/VirtualProduction/DMX/DMXEngine/Source/DMXEditor/Private/Customizations/DMXEntityFixturePatchDetails.cpp
--- a/Engine/Plugins/VirtualProduction/DMX/DMXEngine/Source/DMXEditor/Private/Customizations/DMXEntityFixturePatchDetails.cpp
+++ b/Engine/Plugins/VirtualProduction/DMX/DMXEngine/Source/DMXEditor/Private/Customizations/DMXEntityFixturePatchDetails.cpp
@@ -3,6 +3,7 @@
 #include "Customizations/DMXEntityFixturePatchDetails.h"
 
 #include "DMXEditorUtils.h"
+#include "DMXFixturePatchSharedData.h"
 #include "Library/DMXEntityFixturePatch.h"
 #include "Widgets/SDMXEntityDropdownMenu.h"
 
@@ -31,11 +32,12 @@ void FDMXEntityFixturePatchDetails::CustomizeDetails(IDetailLayoutBuilder& Detai
 	ParentFixtureTypeHandle = DetailBuilder.GetProperty(UDMXEntityFixturePatch::GetParentFixtureTypeTemplatePropertyNameChecked());
 	ActiveModeHandle = DetailBuilder.GetProperty(UDMXEntityFixturePatch::GetActiveModePropertyNameChecked());
 
-	// Bind to auto assign address changes to assign channels when it gets enabled
-	FSimpleDelegate OnAutoAssignAddressChangedDelegate = FSimpleDelegate::CreateSP(this, &FDMXEntityFixturePatchDetails::OnAutoAssignAddressChanged);
-	AutoAssignAddressHandle->SetOnPropertyValueChanged(OnAutoAssignAddressChangedDelegate);
+	// Handle Property changes
+	AutoAssignAddressHandle->SetOnPropertyValueChanged(FSimpleDelegate::CreateSP(this, &FDMXEntityFixturePatchDetails::OnAutoAssignAddressChanged));
+	UniverseIDHandle = DetailBuilder.GetProperty(UDMXEntityFixturePatch::GetUniverseIDPropertyNameChecked());
+	UniverseIDHandle->SetOnPropertyValueChanged(FSimpleDelegate::CreateSP(this, &FDMXEntityFixturePatchDetails::OnUniverseIDChanged));
 
-	// Handle mode changes of the parent fixture type
+	// Handle Fixture Type changes
 	UDMXEntityFixtureType::GetOnFixtureTypeChanged().AddSP(this, &FDMXEntityFixturePatchDetails::OnFixtureTypeChanged);
 
 	// Make a Fixture Types dropdown for the Fixture Type template property
@@ -141,6 +143,21 @@ void FDMXEntityFixturePatchDetails::OnFixtureTypeChanged(const UDMXEntityFixture
 	}
 }
 
+void FDMXEntityFixturePatchDetails::OnUniverseIDChanged()
+{
+	const TSharedPtr<FDMXEditor> DMXEditor = DMXEditorPtr.Pin();
+	if (!DMXEditor.IsValid())
+	{
+		return;
+	}
+
+	int32 UniverseID;
+	if (UniverseIDHandle->GetValue(UniverseID) == FPropertyAccess::Success)
+	{
+		DMXEditor->GetFixturePatchSharedData()->SelectUniverse(UniverseID);
+	}
+}
+
 void FDMXEntityFixturePatchDetails::OnActiveModeChanged(const TSharedPtr<uint32> InSelectedMode, ESelectInfo::Type SelectInfo)
 {
 	if (InSelectedMode.IsValid())
diff --git a/Engine/Plugins/VirtualProduction/DMX/DMXEngine/Source/DMXEditor/Private/Customizations/DMXEntityFixturePatchDetails.h b/Engine/Plugins/VirtualProduction/DMX/DMXEngine/Source/DMXEditor/Private/Customizations/DMXEntityFixturePatchDetails.h
--- a/Engine/Plugins/VirtualProduction/DMX/DMXEngine/Source/DMXEditor/Private/Customizations/DMXEntityFixturePatchDetails.h
+++ b/Engine/Plugins/VirtualProduction/DMX/DMXEngine/Source/DMXEditor/Private/Customizations/DMXEntityFixturePatchDetails.h
@@ -39,13 +39,16 @@ private:
 
 	/** Called when the ParentFixtureType property changed */
 	void OnParentFixtureTypeChanged(UDMXEntity* NewTemplate) const;
-	
-	/** Called when the ActiveMode property changed */
-	void OnActiveModeChanged(const TSharedPtr<uint32> InSelectedMode, ESelectInfo::Type SelectInfo);
 
 	/** Called when a Mode in the Modes Array property changed */
 	void OnFixtureTypeChanged(const UDMXEntityFixtureType* FixtureType);
 
+	/** Called when the UniverseID property changed */
+	void OnUniverseIDChanged();
+
+	/** Called when the ActiveMode property changed */
+	void OnActiveModeChanged(const TSharedPtr<uint32> InSelectedMode, ESelectInfo::Type SelectInfo);
+
 	/** Called when the bAutoAssignAddress property changed */
 	void OnAutoAssignAddressChanged();
 
@@ -70,6 +73,9 @@ private:
 	/** Handle for the ParentFixtureType property */
 	TSharedPtr<IPropertyHandle> ParentFixtureTypeHandle;
 
+	/** Handle for the UniverseID property */
+	TSharedPtr<IPropertyHandle> UniverseIDHandle;
+
 	/** Handle for the ActiveMode property */
 	TSharedPtr<IPropertyHandle> ActiveModeHandle;
 
diff --git a/Engine/Plugins/VirtualProduction/DMX/DMXEngine/Source/DMXEditor/Private/Widgets/FixturePatch/DMXFixturePatchNode.cpp b/Engine/Plugins/VirtualProduction/DMX/DMXEngine/Source/DMXEditor/Private/Widgets/FixturePatch/DMXFixturePatchNode.cpp
--- a/Engine/Plugins/VirtualProduction/DMX/DMXEngine/Source/DMXEditor/Private/Widgets/FixturePatch/DMXFixturePatchNode.cpp
+++ b/Engine/Plugins/VirtualProduction/DMX/DMXEngine/Source/DMXEditor/Private/Widgets/FixturePatch/DMXFixturePatchNode.cpp
@@ -8,6 +8,7 @@
 #include "SDMXFixturePatchFragment.h"
 #include "SDMXPatchedUniverse.h"
 #include "Library/DMXEntityFixturePatch.h"
+#include "Library/DMXLibrary.h"
 
 #include "ScopedTransaction.h"
 #include "Widgets/Layout/SGridPanel.h"
@@ -15,175 +16,353 @@
 
 #define LOCTEXT_NAMESPACE "DMXFixturePatchNode"
 
-void FDMXFixturePatchFragment::CreateFragments(TWeakObjectPtr<UDMXEntityFixturePatch> InFixturePatch, int32 ChannelID, TArray<TSharedPtr<FDMXFixturePatchFragment>>& OutFragments)
+
+/** A fixture patch fragment in a grid. CreateFragments creates DMXFixturePatchFragment from a DMXFixturePatchNode. */
+class FDMXFixturePatchFragment
+	: public TSharedFromThis<FDMXFixturePatchFragment>
 {
-	OutFragments.Reset();
-	if (!InFixturePatch.IsValid())
+public:
+	/** Creates fragments of a patch spread across a grid */
+	static void CreateFragments(TWeakObjectPtr<UDMXEntityFixturePatch> InFixturePatch, int32 ChannelID, TArray<TSharedPtr<FDMXFixturePatchFragment>>& OutFragments)
 	{
-		return;
-	}
+		OutFragments.Reset();
 
-	// Remove ID offset
-	ChannelID = ChannelID - FDMXChannelGridSpecs::ChannelIDOffset;
-
-	int32 ChannelSpan = InFixturePatch->GetChannelSpan();
-	check(ChannelSpan > 0);
+		UDMXEntityFixturePatch* FixturePatch = InFixturePatch.Get();
+		if (!FixturePatch)
+		{
+			return;
+		}
 
-	// Create a node for each row, at least one node
-	TSharedPtr<FDMXFixturePatchFragment> PrevFragment;	
-	do
-	{
-		TSharedPtr<FDMXFixturePatchFragment> NewFragment = MakeShared<FDMXFixturePatchFragment>();
+		UDMXLibrary* DMXLibrary = FixturePatch->GetParentLibrary();
+		if (!DMXLibrary)
+		{
+			return;
+		}
 
-		int32 Column = ChannelID % FDMXChannelGridSpecs::NumColumns;
-		int32 Row = ChannelID / FDMXChannelGridSpecs::NumColumns;
-		int32 ColumnSpan = Column + ChannelSpan < FDMXChannelGridSpecs::NumColumns ? ChannelSpan : FDMXChannelGridSpecs::NumColumns - Column;
+		TMap<TRange<int32>, int32> RangeToNumConflictsMap;
+		BuildRangeToNumConflictsMap(*FixturePatch, RangeToNumConflictsMap);
+		if (RangeToNumConflictsMap.IsEmpty())
+		{
+			return;
+		}
 
-		NewFragment->ID = ChannelID;
-		NewFragment->Column = Column;
-		NewFragment->Row = Row;
-		NewFragment->ColumnSpan = ColumnSpan;
+		// Create a node for each range, at least one node per range
+		TSharedPtr<FDMXFixturePatchFragment> PrevFragment;
+		for (const TTuple<TRange<int32>, int32>& RangeToNumConflictsPair : RangeToNumConflictsMap)
+		{
+			const TRange<int32>& Range = RangeToNumConflictsPair.Key;
+			const int32 NumConflicts = RangeToNumConflictsPair.Value;
+			int32 ChannelIndex = Range.GetLowerBoundValue();
+			int32 ChannelSpan = Range.GetUpperBoundValue() - Range.GetLowerBoundValue();
 
-		OutFragments.Add(NewFragment);
+			do
+			{
+				TSharedPtr<FDMXFixturePatchFragment> NewFragment = MakeShared<FDMXFixturePatchFragment>();
+
+				const int32 Row = ChannelIndex / FDMXChannelGridSpecs::NumColumns;
+				const int32 Column = ChannelIndex % FDMXChannelGridSpecs::NumColumns;
+				int32 ColumnSpan = Column + ChannelSpan < FDMXChannelGridSpecs::NumColumns ? ChannelSpan : FDMXChannelGridSpecs::NumColumns - Column;
+
+				NewFragment->FixturePatch = InFixturePatch;
+				NewFragment->ChannelIndex = ChannelIndex;
+				NewFragment->Column = Column;
+				NewFragment->Row = Row;
+				NewFragment->ColumnSpan = ColumnSpan;
+				NewFragment->bIsConflict = NumConflicts > 0;
+				NewFragment->bIsText = false;
+
+				OutFragments.Add(NewFragment);
+
+				// Link rhs
+				if (PrevFragment.IsValid())
+				{
+					PrevFragment->RhsFragment = NewFragment;
+				}
+
+				// Update for next node
+				ChannelSpan -= ColumnSpan;
+				ChannelIndex += ColumnSpan;
+				PrevFragment = NewFragment;
+			} while (ChannelSpan > 0);
+		}
 
-		if (PrevFragment.IsValid())
+		// Create a text fragment for the first range, from the start to the first increment of conflicts
+		int32 PreviousNumConflicts = TNumericLimits<int32>::Max();
+		int32 TextStartingIndex = INDEX_NONE;
+		int32 TextEndingIndex = INDEX_NONE;
+		for (const TTuple<TRange<int32>, int32>& RangeToNumConflictsPair : RangeToNumConflictsMap)
 		{
-			PrevFragment->RhsFragment = NewFragment;
+			if (TextStartingIndex == INDEX_NONE)
+			{
+				TextStartingIndex = RangeToNumConflictsPair.Key.GetLowerBoundValue();
+			}
+
+			const int32 NumConflicts = RangeToNumConflictsPair.Value;
+			if (PreviousNumConflicts <= NumConflicts)
+			{
+				break;
+			}
+
+			TextEndingIndex = RangeToNumConflictsPair.Key.GetUpperBoundValue() - 1;
+			PreviousNumConflicts = NumConflicts;
 		}
 
-		// Update for next node
-		ChannelSpan -= ColumnSpan;
-		ChannelID += ColumnSpan;
-	} while (ChannelSpan > 0);
+		TSharedPtr<FDMXFixturePatchFragment> NewTextFragment = MakeShared<FDMXFixturePatchFragment>();
 
-	// link lhs
-	TSharedPtr<FDMXFixturePatchFragment> LhsLink = OutFragments[0];
-	for (int idxNode = 1; idxNode < OutFragments.Num(); idxNode++)
-	{
-		OutFragments[idxNode]->LhsFragment = LhsLink;
-		LhsLink = OutFragments[idxNode];
-	}
-}
+		const int32 Row = TextStartingIndex / FDMXChannelGridSpecs::NumColumns;
+		const int32 Column = TextStartingIndex % FDMXChannelGridSpecs::NumColumns;
+		int32 ChannelSpan = TextEndingIndex + 1 - TextStartingIndex;
+		int32 ColumnSpan = Column + ChannelSpan < FDMXChannelGridSpecs::NumColumns ? ChannelSpan : FDMXChannelGridSpecs::NumColumns - Column;
 
-TArray<TSharedPtr<FDMXFixturePatchFragment>> FDMXFixturePatchFragment::GetFragments()
-{
-	TArray<TSharedPtr<FDMXFixturePatchFragment>> Fragments;
-	TSharedPtr<FDMXFixturePatchFragment> Root = SharedThis(this);
-	while(Root->LhsFragment.IsValid())
-	{
-		Root = LhsFragment;
+		NewTextFragment->ChannelIndex = TextStartingIndex;
+		NewTextFragment->Column = Column;
+		NewTextFragment->Row = Row;
+		NewTextFragment->ColumnSpan = ColumnSpan;
+		NewTextFragment->bIsConflict = false;
+		NewTextFragment->bIsText = true;
+		OutFragments.Add(NewTextFragment);
+
+		// Link lhs
+		TSharedPtr<FDMXFixturePatchFragment> LhsLink = OutFragments[0];
+		for (int IndexOfFragment = 1; IndexOfFragment < OutFragments.Num(); IndexOfFragment++)
+		{
+			OutFragments[IndexOfFragment]->LhsFragment = LhsLink;
+			LhsLink = OutFragments[IndexOfFragment];
+		}
 	}
 
-	TSharedPtr<FDMXFixturePatchFragment> Fragment = Root;
-	while (Fragment.IsValid())
+	/** Gets all fragments of a patch */
+	TArray<TSharedPtr<FDMXFixturePatchFragment>> GetFragments()
 	{
-		Fragments.Add(Fragment);
-		Fragment = Fragment->RhsFragment;
+		TArray<TSharedPtr<FDMXFixturePatchFragment>> Fragments;
+		TSharedPtr<FDMXFixturePatchFragment> Root = SharedThis(this);
+		while (Root->LhsFragment.IsValid())
+		{
+			Root = LhsFragment;
+		}
+
+		TSharedPtr<FDMXFixturePatchFragment> Fragment = Root;
+		while (Fragment.IsValid())
+		{
+			Fragments.Add(Fragment);
+			Fragment = Fragment->RhsFragment;
+		}
+
+		return Fragments;
 	}
 
-	return Fragments;
-}
 
-TSharedPtr<FDMXFixturePatchNode> FDMXFixturePatchNode::Create(TWeakPtr<FDMXEditor> InDMXEditor, TWeakObjectPtr<UDMXEntityFixturePatch> InFixturePatch)
-{
-	TSharedPtr<FDMXFixturePatchNode> NewNode = MakeShared<FDMXFixturePatchNode>();
+	bool IsHead() const { return !LhsFragment.IsValid(); }
+	bool IsTail() const { return !RhsFragment.IsValid(); }
 
-	NewNode->DMXEditor = InDMXEditor;
-	NewNode->FixturePatch = InFixturePatch;
-	NewNode->LastTransactedUniverseID = InFixturePatch->GetUniverseID();
-	NewNode->LastTransactedChannelID = InFixturePatch->GetStartingChannel();
+	/** Channel Index of the fragment */
+	int32 ChannelIndex = INDEX_NONE;
+
+	/** Row of the fragment */
+	int32 Row = -1;
+
+	/** Column of the fragment */
+	int32 Column = -1;
+
+	/** Columnspan of the fragment */
+	int32 ColumnSpan = -1;
+
+	/** True if this fragment conflicts with others */
+	bool bIsConflict = false;
+
+	/** True if this fragment is the text fragment */
+	bool bIsText = false;
+
+	/** The fragment to the left of this fragment */
+	TSharedPtr<FDMXFixturePatchFragment> LhsFragment;
+
+	/** The fragment to the right of this fragment */
+	TSharedPtr<FDMXFixturePatchFragment> RhsFragment;
+
+	/** The fixture patch of which the fragment is part of */
+	TWeakObjectPtr<UDMXEntityFixturePatch> FixturePatch;
 
-	// Bind to shared data's selection event 
-	if (TSharedPtr<FDMXEditor> PinnedDMXEditor = NewNode->DMXEditor.Pin())
+private:
+	/** Breaks up the patch in an array of ranges along with the number of conflicting patches */
+	static void BuildRangeToNumConflictsMap(const UDMXEntityFixturePatch& InFixturePatch, TMap<TRange<int32>, int32>& OutRangeToNumConflictsMap)
 	{
-		TSharedPtr<FDMXFixturePatchSharedData> SharedData = PinnedDMXEditor->GetFixturePatchSharedData();
-		check(SharedData.IsValid());
+		TMap<int32, int32> IndexToNumConflictsMap;
+		IndexToNumConflictsMap.Reserve(InFixturePatch.GetChannelSpan());
 
-		SharedData->OnFixturePatchSelectionChanged.AddSP(NewNode.ToSharedRef(), &FDMXFixturePatchNode::OnSelectionChanged);
+		const int32 StartingIndex = InFixturePatch.GetStartingChannel() - 1;
+		const int32 EndingIndex = InFixturePatch.GetEndingChannel() - 1;
+		const TArray<const UDMXEntityFixturePatch*> ConflictingFixturePatches = FindConflictingFixturePatches(InFixturePatch);
+		for (const UDMXEntityFixturePatch* ConflictingFixturePatch : ConflictingFixturePatches)
+		{
+			const int32 ConflictStartingIndex = FMath::Max(StartingIndex, ConflictingFixturePatch->GetStartingChannel() - 1);
+			const int32 ConflictEndingIndex = FMath::Min(EndingIndex, ConflictingFixturePatch->GetEndingChannel() - 1);
+			for (int32 ConflictIndex = ConflictStartingIndex; ConflictIndex <= ConflictEndingIndex; ConflictIndex++)
+			{
+				IndexToNumConflictsMap.FindOrAdd(ConflictIndex, 0)++;
+			}
+		}
+		for (int32 Index = StartingIndex; Index <= EndingIndex; Index++)
+		{
+			if (!IndexToNumConflictsMap.Contains(Index))
+			{
+				IndexToNumConflictsMap.Add(Index, 0);
+			}
+		}
+		if (IndexToNumConflictsMap.IsEmpty())
+		{
+			return;
+		}
 
-		// Select the new node if it is selected in the entity list
-		if (InFixturePatch.IsValid())
+		TRange<int32> Range(StartingIndex, StartingIndex);
+		int32 NumConflicts = IndexToNumConflictsMap[StartingIndex];
+		for (const TTuple<int32, int32>& IndexToNumConflictsPair : IndexToNumConflictsMap)
 		{
-			if (SharedData->GetSelectedFixturePatches().Contains(NewNode->FixturePatch))
+			const int32 NextIndex = Range.GetUpperBoundValue() + 1;
+			Range.SetUpperBoundValue(NextIndex);
+
+			const int32 NextNumConflicts = IndexToNumConflictsMap.Contains(NextIndex) ? IndexToNumConflictsMap[NextIndex] : -1;
+			if (NumConflicts != NextNumConflicts)
 			{
-				NewNode->bSelected = true;
+				OutRangeToNumConflictsMap.Add(Range, NumConflicts);
+				Range = TRange<int32>(Range.GetUpperBoundValue(), Range.GetUpperBoundValue());
+				NumConflicts = NextNumConflicts;
 			}
-		}	
+		}
 	}
 
-	return NewNode;
-}
+	/** Returns the fixture patches that conflict with the specified fixture patch */
+	static TArray<const UDMXEntityFixturePatch*> FindConflictingFixturePatches(const UDMXEntityFixturePatch& InFixturePatch)
+	{
+		TArray<const UDMXEntityFixturePatch*> Result;
+
+		if (InFixturePatch.GetChannelSpan() == 0)
+		{
+			return Result;
+		}
+
+		UDMXLibrary* DMXLibrary = InFixturePatch.GetParentLibrary();
+		if (!DMXLibrary)
+		{
+			return Result;
+		}
+
+		// Get the other Fixture Patches in the DMX Library
+		const UDMXEntityFixturePatch* FixturePatch = &InFixturePatch;
+		const TArray<const UDMXEntityFixturePatch*> OtherFixturePatches = [FixturePatch, DMXLibrary]()
+		{
+			TArray<const UDMXEntityFixturePatch*> FixturePatches = DMXLibrary->GetEntitiesTypeCast<const UDMXEntityFixturePatch>();
+			FixturePatches.Remove(FixturePatch);
+			return FixturePatches;
+		}();
+
+		// Find the conflicts
+		TArray<const UDMXEntityFixturePatch*> ConflictingFixturePatches = [FixturePatch, &OtherFixturePatches]()
+		{
+			TArray<const UDMXEntityFixturePatch*> Result;
+			for (const UDMXEntityFixturePatch* Other : OtherFixturePatches)
+			{
+				if (!Other)
+				{
+					continue;
+				}
+
+				if (Other->GetUniverseID() != FixturePatch->GetUniverseID())
+				{
+					continue;
+				}
+
+				// Equal patches are not conflicting
+				if (Other->GetFixtureType() == FixturePatch->GetFixtureType() &&
+					Other->GetActiveModeIndex() == FixturePatch->GetActiveModeIndex() &&
+					Other->GetStartingChannel() == FixturePatch->GetStartingChannel())
+				{
+					continue;
+				}
+
+				const int32 OtherChannelIndex = Other->GetStartingChannel() - 1;
+				const int32 OtherChannelSpan = Other->GetChannelSpan();
+				const TRange<int32> OtherRange(OtherChannelIndex, OtherChannelIndex + OtherChannelSpan);
+
+				if (!OtherRange.Overlaps(OtherRange))
+				{
+					continue;
+				}
+
+				Result.Add(Other);
+			}
+			return Result;
+		}();
 
-void FDMXFixturePatchNode::Update(TSharedPtr<SDMXPatchedUniverse> NewUniverseWidget, int32 NewStartingChannel, int32 NewChannelSpan)
+		return ConflictingFixturePatches;
+	}
+};
+
+TSharedPtr<FDMXFixturePatchNode> FDMXFixturePatchNode::Create(TWeakPtr<FDMXEditor> InDMXEditor, const TWeakObjectPtr<UDMXEntityFixturePatch>& InFixturePatch)
 {
-	check(NewUniverseWidget.IsValid());
-	check(NewStartingChannel > 0 && NewStartingChannel <= DMX_UNIVERSE_SIZE);
-	check(NewChannelSpan > 0);
+	if (!ensureMsgf(InFixturePatch.IsValid(), TEXT("Trying to create Fixture Patch Node. But no Fixture Patch is not valid.")))
+	{
+		return nullptr;
+	}
 
-	UniverseWidget = NewUniverseWidget;
-	StartingChannel = NewStartingChannel;
-	ChannelSpan = NewChannelSpan;
+	TSharedPtr<FDMXFixturePatchNode> NewNode = MakeShared<FDMXFixturePatchNode>();
+	NewNode->WeakDMXEditor = InDMXEditor;
+	NewNode->FixturePatch= InFixturePatch;
 
-	// Fragment the patch to fit the ID
-	TArray<TSharedPtr<FDMXFixturePatchFragment>> NewFragments;
-	FDMXFixturePatchFragment::CreateFragments(FixturePatch, NewStartingChannel, NewFragments);
+	TSharedPtr<FDMXEditor> DMXEditor = NewNode->WeakDMXEditor.Pin();
+	if (!DMXEditor.IsValid())
+	{
+		return nullptr;
+	}
 
+	NewNode->LastTransactedUniverseID = InFixturePatch->GetUniverseID();
+	NewNode->LastTransactedChannelID = InFixturePatch->GetStartingChannel();
 
-	// Update the Fragmented Widgets, reuse existing if possible
-	if (NewFragments.Num() == FragmentedWidgets.Num())
+	TSharedPtr<FDMXFixturePatchSharedData> SharedData = DMXEditor->GetFixturePatchSharedData();
+	if (!SharedData.IsValid())
 	{
-		for (int32 IdxFragment = 0; IdxFragment < NewFragments.Num(); IdxFragment++)
-		{
-			TSharedPtr<SDMXFixturePatchFragment> Widget = FragmentedWidgets[IdxFragment];
-			Widget->SetColumn(NewFragments[IdxFragment]->Column);
-			Widget->SetRow(NewFragments[IdxFragment]->Row);
-			Widget->SetColumnSpan(NewFragments[IdxFragment]->ColumnSpan);
-		}
+		return nullptr;
 	}
-	else
+
+	// Select the new node if it is selected in Fixture Patch Shared Data
+	if (SharedData->GetSelectedFixturePatches().Contains(InFixturePatch))
 	{
-		FragmentedWidgets.Reset();
-		for (const TSharedPtr<FDMXFixturePatchFragment>& Fragment : NewFragments)
-		{
-			FragmentedWidgets.Add(
-				SNew(SDMXFixturePatchFragment, SharedThis(this))
-				.DMXEditor(DMXEditor)
-				.Text(FText::FromString(FixturePatch->GetDisplayName()))
-				.Column(Fragment->Column)
-				.Row(Fragment->Row)
-				.ColumnSpan(Fragment->ColumnSpan)
-				.bHighlight(bSelected)
-				.Visibility(EVisibility::HitTestInvisible)
-			);
-		}
+		NewNode->bSelected = true;
 	}
+
+	// Listen to selection changes 
+	SharedData->OnFixturePatchSelectionChanged.AddSP(NewNode.ToSharedRef(), &FDMXFixturePatchNode::OnSelectionChanged);
+
+	return NewNode;
 }
 
-void FDMXFixturePatchNode::CommitPatch(bool bTransacted)
+void FDMXFixturePatchNode::SetAddresses(const TSharedRef<SDMXPatchedUniverse>& OwningUniverseWidget, int32 NewStartingChannel, int32 NewChannelSpan, bool bTransacted)
 {
 	if (!FixturePatch.IsValid())
-	{		
+	{
 		return;
 	}
-	
+
+	UniverseWidget = OwningUniverseWidget;
+	StartingChannel = NewStartingChannel;
+	ChannelSpan = NewChannelSpan;
+
 	if (bTransacted)
-	{	
+	{
 		// As we allow the asset to update even when not transacted,
 		// We have to set the last known values to enable proper undo
 		FixturePatch->SetUniverseID(LastTransactedUniverseID);
 		FixturePatch->SetManualStartingAddress(LastTransactedChannelID);
 
+		LastTransactedUniverseID = UniverseWidget->GetUniverseID();
+		LastTransactedChannelID = StartingChannel;
+
 		const FScopedTransaction Transaction = FScopedTransaction(LOCTEXT("FixturePatchAssigned", "Patched Fixture"));
-		FixturePatch->Modify();
 		FixturePatch->PreEditChange(nullptr);
 
 		FixturePatch->SetUniverseID(UniverseWidget->GetUniverseID());
 		FixturePatch->SetManualStartingAddress(StartingChannel);
 
 		FixturePatch->PostEditChange();
-	
-		LastTransactedUniverseID = UniverseWidget->GetUniverseID();
-		LastTransactedChannelID = StartingChannel;
 	}
 	else
 	{
@@ -192,6 +371,47 @@ void FDMXFixturePatchNode::CommitPatch(bool bTransacted)
 	}
 }
 
+TArray<TSharedPtr<SDMXFixturePatchFragment>> FDMXFixturePatchNode::GenerateWidgets(const TArray<TSharedPtr<FDMXFixturePatchNode>>& FixturePatchNodeGroup)
+{
+
+	// Fragment the patch to fit the Addresses
+	TArray<TSharedPtr<FDMXFixturePatchFragment>> NewFragments;
+	FDMXFixturePatchFragment::CreateFragments(FixturePatch, StartingChannel, NewFragments);
+
+	FixturePatchFragmentWidgets.Reset();
+	for (const TSharedPtr<FDMXFixturePatchFragment>& Fragment : NewFragments)
+	{
+		const bool bIsHead = !Fragment->LhsFragment.IsValid();
+		const bool bIsTail = !Fragment->RhsFragment.IsValid();
+
+		const FText DisplayName = FixturePatch.IsValid() ? FText::FromString(FixturePatch->GetDisplayName()) : LOCTEXT("InvalidFixturePatchName", "Invalid Fixture Patch");
+		FixturePatchFragmentWidgets.Add(
+			SNew(SDMXFixturePatchFragment, SharedThis(this), FixturePatchNodeGroup)
+			.DMXEditor(WeakDMXEditor)
+			.IsHead(bIsHead)
+			.IsTail(bIsTail)
+			.IsText(Fragment->bIsText)
+			.IsConflicting(Fragment->bIsConflict)
+			.Column(Fragment->Column)
+			.Row(Fragment->Row)
+			.ColumnSpan(Fragment->ColumnSpan)
+			.bHighlight(bSelected)
+		);
+	}
+
+	return FixturePatchFragmentWidgets;
+}
+
+bool FDMXFixturePatchNode::OccupiesChannels(int32 Channel, int32 Span) const
+{
+	check(Span != 0);
+	if ((Channel + Span <= StartingChannel) || (Channel >= StartingChannel + ChannelSpan))
+	{
+		return false;
+	}
+	return true;
+}
+
 bool FDMXFixturePatchNode::IsPatched() const
 {
 	return UniverseWidget.IsValid();
@@ -214,40 +434,32 @@ int32 FDMXFixturePatchNode::GetUniverseID() const
 
 void FDMXFixturePatchNode::OnSelectionChanged()
 {
-	if (TSharedPtr<FDMXEditor> PinnedDMXEditor = DMXEditor.Pin())
+	const TSharedPtr<FDMXEditor> DMXEditor = WeakDMXEditor.Pin();
+	if (!DMXEditor.IsValid())
 	{
-		TSharedPtr<FDMXFixturePatchSharedData> SharedData = PinnedDMXEditor->GetFixturePatchSharedData();
-		check(SharedData.IsValid());
+		return;
+	}
 
-		if (SharedData->GetSelectedFixturePatches().Contains(FixturePatch))
+	const TSharedPtr<FDMXFixturePatchSharedData> FixturePatchSharedData = DMXEditor->GetFixturePatchSharedData();
+	TArray<TWeakObjectPtr<UDMXEntityFixturePatch>> SelectedFixturePatches = FixturePatchSharedData->GetSelectedFixturePatches();
+	if (SelectedFixturePatches.Contains(FixturePatch))
+	{
+		for (const TSharedPtr<SDMXFixturePatchFragment>& Widget : FixturePatchFragmentWidgets)
 		{
-			for (const TSharedPtr<SDMXFixturePatchFragment>& Fragment : FragmentedWidgets)
-			{
-				Fragment->SetHighlight(true);
-			}
-
-			bSelected = true;
+			Widget->SetHighlight(true);
 		}
-		else
-		{
-			for (const TSharedPtr<SDMXFixturePatchFragment>& Fragment : FragmentedWidgets)
-			{
-				Fragment->SetHighlight(false);
-			}
 
-			bSelected = false;
-		}
+		bSelected = true;
 	}
-}
-
-bool FDMXFixturePatchNode::OccupiesChannels(int32 Channel, int32 Span) const
-{
-	check(Span != 0);
-	if ((Channel + Span <= StartingChannel) || (Channel >= StartingChannel + ChannelSpan))
+	else
 	{
-		return false;
-	} 
-	return true;
+		for (const TSharedPtr<SDMXFixturePatchFragment>& Widget : FixturePatchFragmentWidgets)
+		{
+			Widget->SetHighlight(false);
+		}
+
+		bSelected = false;
+	}
 }
 
 #undef LOCTEXT_NAMESPACE
diff --git a/Engine/Plugins/VirtualProduction/DMX/DMXEngine/Source/DMXEditor/Private/Widgets/FixturePatch/DMXFixturePatchNode.h b/Engine/Plugins/VirtualProduction/DMX/DMXEngine/Source/DMXEditor/Private/Widgets/FixturePatch/DMXFixturePatchNode.h
--- a/Engine/Plugins/VirtualProduction/DMX/DMXEngine/Source/DMXEditor/Private/Widgets/FixturePatch/DMXFixturePatchNode.h
+++ b/Engine/Plugins/VirtualProduction/DMX/DMXEngine/Source/DMXEditor/Private/Widgets/FixturePatch/DMXFixturePatchNode.h
@@ -11,38 +11,9 @@ class SDMXFixturePatchFragment;
 class SDMXPatchedUniverse;
 class UDMXEntityFixturePatch;
 
-class SGridPanel;
-struct EVisibility;
-
-/** A fixture patch fragment in a grid. CreateFragments creates DMXFixturePatchFragment from a DMXFixturePatchNode. */
-class FDMXFixturePatchFragment
-	: public TSharedFromThis<FDMXFixturePatchFragment>
-{
-public:
-	/** Creates fragments of a patch spread across a grid */
-	static void CreateFragments(TWeakObjectPtr<UDMXEntityFixturePatch> InFixturePatch, int32 ChannelID, TArray<TSharedPtr<FDMXFixturePatchFragment>>& OutFragments);
-
-	/** Gets all fragments of a patch */
-	TArray<TSharedPtr<FDMXFixturePatchFragment>> GetFragments();
-
-	bool IsHead() const { return !LhsFragment.IsValid(); }
-	bool IsTail() const { return !RhsFragment.IsValid(); }
-
-	int32 ID;
-	int32 Row;
-	int32 Column;
-	int32 ColumnSpan;
-
-	TSharedPtr<FDMXFixturePatchFragment> LhsFragment;
-	TSharedPtr<FDMXFixturePatchFragment> RhsFragment;
-
-	TWeakObjectPtr<UDMXEntityFixturePatch> FixturePatch;
-};
-
-////////////////////////////////////////////////////////////////////////////////
 
 /** 
- * A fixture patch in a grid, consists of several fragments, displayed as SDMXFixturePatchNodes. 
+ * A fixture patch in a grid, consists of several fragments, displayed as SDMXFixturePatchNodes (see GenerateWidgets). 
  * The patch node does not necessarily have to present the state of the object, e.g. when
  * dragged or when containing an occuppied universe/channel. 
  */
@@ -50,14 +21,17 @@ class FDMXFixturePatchNode
 	: public TSharedFromThis<FDMXFixturePatchNode>
 {
 public:
-	/** Creates a new node from patch */
-	static TSharedPtr<FDMXFixturePatchNode> Create(TWeakPtr<FDMXEditor> InDMXEditor, TWeakObjectPtr<UDMXEntityFixturePatch> InFixturePatch);
-	
-	/** Updates the patch node. NewUniverse can be nullptr */
-	void Update(TSharedPtr<SDMXPatchedUniverse> NewUniverseWidget, int32 NewStartingChannel, int32 NewChannelSpan);
+	/** Creates a new node from Patch */
+	static TSharedPtr<FDMXFixturePatchNode> Create(TWeakPtr<FDMXEditor> InDMXEditor, const TWeakObjectPtr<UDMXEntityFixturePatch>& InFixturePatch);
+
+	/** Sets the Addresses of the Node */
+	void SetAddresses(const TSharedRef<SDMXPatchedUniverse>& OwningUniverseWidget, int32 NewStartingChannel, int32 NewChannelSpan, bool bTransacted);
+
+	/** Generates widgets. Requires Addresses to be set via SetAddresses */
+	TArray<TSharedPtr<SDMXFixturePatchFragment>> GenerateWidgets(const TArray<TSharedPtr<FDMXFixturePatchNode>>& FixturePatchNodeGroup);
 
-	/** Commits the patch to the object */
-	void CommitPatch(bool bTransacted);
+	/** Returns widgets */
+	const TArray<TSharedPtr<SDMXFixturePatchFragment>>& GetGeneratedWidgets() { return FixturePatchFragmentWidgets; };
 
 	/** Returns wether the patch uses specified channesl */
 	bool OccupiesChannels(int32 Channel, int32 Span) const;
@@ -68,15 +42,12 @@ public:
 	/** Returns true if the node is selected */
 	bool IsSelected() const;
 
-	/** Returns the fixture patch this node connects */
+	/** Returns the fixture patches this node holds */
 	TWeakObjectPtr<UDMXEntityFixturePatch> GetFixturePatch() const { return FixturePatch; }
 
 	/** Returns the universe this node resides in */
 	const TSharedPtr<SDMXPatchedUniverse>& GetUniverseWidget() const { return UniverseWidget; }
 
-	/** Returns fragmented widgets to visualize the node in a grid */
-	const TArray<TSharedPtr<SDMXFixturePatchFragment>>& GetFragmentedWidgets() const { return FragmentedWidgets; }
-
 	/** Returns the Universe ID the Node currently resides in. Returns a negative Value if not assigned to a Universe. */
 	int32 GetUniverseID() const;
 
@@ -86,6 +57,12 @@ public:
 	/** Returns the Channel Span the Node currently occupies */
 	int32 GetChannelSpan() const { return ChannelSpan; }
 
+	/** Sets the ZOrder of the Node */
+	void SetZOrder(int32 NewZOrder) { ZOrder = NewZOrder; }
+
+	/** Gets the ZOrder of the Node */
+	int32 GetZOrder() const { return ZOrder; }
+
 private:
 	/** Called when a patch node was selected */
 	void OnSelectionChanged();
@@ -93,7 +70,7 @@ private:
 	/** Universe the patch is assigned to */
 	TSharedPtr<SDMXPatchedUniverse> UniverseWidget;
 
-	/** Starting channel of the patch */
+	/** Starting channel of the patch */	
 	int32 StartingChannel = 0;
 
 	/** Channel span of the patch */
@@ -108,12 +85,15 @@ private:
 	/** If true the node is selected */
 	bool bSelected = false;
 
-	/** The widget fragments that visualize the patch */
-	TArray<TSharedPtr<SDMXFixturePatchFragment>> FragmentedWidgets;
+	/** ZOrder of this node */
+	int32 ZOrder = 0;
+
+	/** Keep track of Fragment Widgets, useful to handle selection */
+	TArray<TSharedPtr<SDMXFixturePatchFragment>> FixturePatchFragmentWidgets;
 
-	/** Weak refrence of the actual patch in the library */
+	/** The Fixture Patch this Node stands for */
 	TWeakObjectPtr<UDMXEntityFixturePatch> FixturePatch;
 
 	/** Weak DMXEditor refrence */
-	TWeakPtr<FDMXEditor> DMXEditor;
+	TWeakPtr<FDMXEditor> WeakDMXEditor;
 };
diff --git a/Engine/Plugins/VirtualProduction/DMX/DMXEngine/Source/DMXEditor/Private/Widgets/FixturePatch/SDMXChannelConnector.cpp b/Engine/Plugins/VirtualProduction/DMX/DMXEngine/Source/DMXEditor/Private/Widgets/FixturePatch/SDMXChannelConnector.cpp
--- a/Engine/Plugins/VirtualProduction/DMX/DMXEngine/Source/DMXEditor/Private/Widgets/FixturePatch/SDMXChannelConnector.cpp
+++ b/Engine/Plugins/VirtualProduction/DMX/DMXEngine/Source/DMXEditor/Private/Widgets/FixturePatch/SDMXChannelConnector.cpp
@@ -4,7 +4,6 @@
 
 #include "DMXEditor.h"
 #include "DMXFixturePatchSharedData.h"
-#include "DragDrop/DMXEntityFixturePatchDragDropOp.h"
 #include "Library/DMXLibrary.h"
 #include "Library/DMXEntityFixturePatch.h"
 
@@ -13,6 +12,16 @@
 
 void SDMXChannelConnector::Construct(const FArguments& InArgs)
 {
+	check(InArgs._OnMouseButtonDownOnChannel.IsBound());
+	check(InArgs._OnMouseButtonUpOnChannel.IsBound());
+	check(InArgs._OnDragDetectedOnChannel.IsBound());
+	check(InArgs._OnDragEnterChannel.IsBound());
+	check(InArgs._OnDragLeaveChannel.IsBound());
+	check(InArgs._OnDropOntoChannel.IsBound());
+
+	OnMouseButtonDownOnChannel = InArgs._OnMouseButtonDownOnChannel;
+	OnMouseButtonUpOnChannel = InArgs._OnMouseButtonUpOnChannel;
+	OnDragDetectedOnChannel = InArgs._OnDragDetectedOnChannel;
 	OnDragEnterChannel = InArgs._OnDragEnterChannel;
 	OnDragLeaveChannel = InArgs._OnDragLeaveChannel;
 	OnDropOntoChannel = InArgs._OnDropOntoChannel;
@@ -33,44 +42,17 @@ void SDMXChannelConnector::Construct(const FArguments& InArgs)
 
 FReply SDMXChannelConnector::OnMouseButtonDown(const FGeometry& MyGeometry, const FPointerEvent& MouseEvent)
 {
-	if (MouseEvent.GetEffectingButton() == EKeys::LeftMouseButton)
-	{
-		if(UDMXEntityFixturePatch* Patch = GetFixturePatch())
-		{
-			SelectFixturePatch(Patch);
-
-			FReply Reply = FReply::Handled();
-
-			// SelectFixturePatch calls SetKeyboardFocus at some point; the call to SetKeyboardFocus is not enough for Slate: 
-			// We must also return the widget using FReply.
-			TSharedPtr<SWidget> WidgetThatWasFocusedBy_SelectFixturePatch = FSlateApplication::Get().GetKeyboardFocusedWidget();
-			if(WidgetThatWasFocusedBy_SelectFixturePatch.IsValid())
-			{
-				Reply.SetUserFocus(WidgetThatWasFocusedBy_SelectFixturePatch.ToSharedRef());
-			}
-			return Reply.DetectDrag(SharedThis(this), EKeys::LeftMouseButton);
-		}
-	}
-	return FReply::Unhandled();
+	return OnMouseButtonDownOnChannel.Execute(ChannelValueWidget->GetChannelID(), MouseEvent);
 }
 
-FReply SDMXChannelConnector::OnDragDetected(const FGeometry& MyGeometry, const FPointerEvent& MouseEvent)
+FReply SDMXChannelConnector::OnMouseButtonUp(const FGeometry& MyGeometry, const FPointerEvent& MouseEvent)
 {
-	UDMXLibrary* DMXLibrary = GetDMXLibrary();
-
-	if (DMXLibrary)
-	{
-		if (UDMXEntityFixturePatch* Patch = GetFixturePatch())
-		{
-			int32 ChannelOffset = ChannelID - Patch->GetStartingChannel();
-
-			TSharedRef<FDMXEntityFixturePatchDragDropOperation> DragDropOp = MakeShared<FDMXEntityFixturePatchDragDropOperation>(DMXLibrary, TArray<TWeakObjectPtr<UDMXEntity>>{ Patch }, ChannelOffset);
-
-			return FReply::Handled().BeginDragDrop(DragDropOp);
-		}
-	}
+	return OnMouseButtonUpOnChannel.Execute(ChannelValueWidget->GetChannelID(), MouseEvent);
+}
 
-	return FReply::Unhandled();
+FReply SDMXChannelConnector::OnDragDetected(const FGeometry& MyGeometry, const FPointerEvent& MouseEvent)
+{
+	return OnDragDetectedOnChannel.Execute(ChannelValueWidget->GetChannelID(), MouseEvent);
 }
 
 void SDMXChannelConnector::OnDragEnter(const FGeometry& MyGeometry, const FDragDropEvent& DragDropEvent)
@@ -88,48 +70,6 @@ FReply SDMXChannelConnector::OnDrop(const FGeometry& MyGeometry, const FDragDrop
 	return OnDropOntoChannel.Execute(ChannelValueWidget->GetChannelID(), DragDropEvent);
 }
 
-void SDMXChannelConnector::SelectFixturePatch(UDMXEntityFixturePatch* FixturePatch)
-{
-	if (TSharedPtr<FDMXEditor> DMXEditor = DMXEditorPtr.Pin())
-	{
-		TSharedPtr<FDMXFixturePatchSharedData> SharedData = DMXEditor->GetFixturePatchSharedData();
-		check(SharedData.IsValid());
-
-		if (FSlateApplication::Get().GetModifierKeys().IsShiftDown())
-		{
-			SharedData->AddFixturePatchToSelection(FixturePatch);
-		}
-		else
-		{
-			SharedData->SelectFixturePatch(FixturePatch);
-		}
-	}
-}
-
-UDMXEntityFixturePatch* SDMXChannelConnector::GetFixturePatch() const
-{
-	if (UDMXLibrary* DMXLibrary = GetDMXLibrary())
-	{
-		TArray<UDMXEntityFixturePatch*> FixturePatches = DMXLibrary->GetEntitiesTypeCast<UDMXEntityFixturePatch>();
-		UDMXEntityFixturePatch** PatchOnChannel = FixturePatches.FindByPredicate([this](UDMXEntityFixturePatch* TestedPatch) {
-			int32 PatchUniverseID = TestedPatch->GetUniverseID();
-			int32 PatchStartingChannel = TestedPatch->GetStartingChannel();
-			int32 PatchEndingChannel = PatchStartingChannel + TestedPatch->GetChannelSpan() - 1;
-			return
-				PatchUniverseID == UniverseID &&
-				PatchStartingChannel <= ChannelID &&
-				PatchEndingChannel >= ChannelID;
-			});
-
-		if (PatchOnChannel)
-		{
-			return *PatchOnChannel;
-		}
-	}
-
-	return nullptr;
-}
-
 UDMXLibrary* SDMXChannelConnector::GetDMXLibrary() const
 {
 	if (TSharedPtr<FDMXEditor> DMXEditor = DMXEditorPtr.Pin())
diff --git a/Engine/Plugins/VirtualProduction/DMX/DMXEngine/Source/DMXEditor/Private/Widgets/FixturePatch/SDMXChannelConnector.h b/Engine/Plugins/VirtualProduction/DMX/DMXEngine/Source/DMXEditor/Private/Widgets/FixturePatch/SDMXChannelConnector.h
--- a/Engine/Plugins/VirtualProduction/DMX/DMXEngine/Source/DMXEditor/Private/Widgets/FixturePatch/SDMXChannelConnector.h
+++ b/Engine/Plugins/VirtualProduction/DMX/DMXEngine/Source/DMXEditor/Private/Widgets/FixturePatch/SDMXChannelConnector.h
@@ -15,9 +15,11 @@ class UDMXEntityFixturePatch;
 class SDMXChannelConnector
 	: public SCompoundWidget
 {
-	DECLARE_DELEGATE_TwoParams(FOnDragOverChannel, int32, const FDragDropEvent&);
+	DECLARE_DELEGATE_RetVal_TwoParams(FReply, FOnMouseEventWithReply, uint32 /** ChannelID */, const FPointerEvent&);
 
-	DECLARE_DELEGATE_RetVal_TwoParams(FReply, FOnDropOntoChannel, int32, const FDragDropEvent&);
+	DECLARE_DELEGATE_TwoParams(FOnDragEvent, uint32 /** ChannelID */, const FDragDropEvent&);
+
+	DECLARE_DELEGATE_RetVal_TwoParams(FReply, FOnDragEventWithReply, uint32 /** ChannelID */, const FDragDropEvent&);
 
 public:
 	SLATE_BEGIN_ARGS(SDMXChannelConnector)
@@ -38,14 +40,23 @@ public:
 		/** The value of the channel */
 		SLATE_ARGUMENT(uint8, Value)
 
-		/** Called when drag enters the widget */
-		SLATE_EVENT(FOnDragOverChannel, OnDragEnterChannel)
+		/** Called when Mouse Button was pressed on the Channel */
+		SLATE_EVENT(FOnMouseEventWithReply, OnMouseButtonDownOnChannel)
+
+		/** Called when Mouse Button was released on the Channel */
+		SLATE_EVENT(FOnMouseEventWithReply, OnMouseButtonUpOnChannel)
+
+		/** Called when it was detected that the Channel was dragged */
+		SLATE_EVENT(FOnMouseEventWithReply, OnDragDetectedOnChannel)
 
-		/** Called when drag leaves the widget */
-		SLATE_EVENT(FOnDragOverChannel, OnDragLeaveChannel)
+		/** Called when drag enters the Widget */
+		SLATE_EVENT(FOnDragEvent, OnDragEnterChannel)
 
-		/** Called when dropped onto the channel */
-		SLATE_EVENT(FOnDropOntoChannel, OnDropOntoChannel)
+		/** Called when drag leaves the Widget */
+		SLATE_EVENT(FOnDragEvent, OnDragLeaveChannel)
+
+		/** Called when dropped onto the Channel */
+		SLATE_EVENT(FOnDragEventWithReply, OnDropOntoChannel)
 
 		SLATE_ARGUMENT(TWeakPtr<FDMXEditor>, DMXEditor)
 
@@ -66,34 +77,33 @@ public:
 protected:
 	// Begin SWidget interface	
 	virtual FReply OnMouseButtonDown(const FGeometry& MyGeometry, const FPointerEvent& MouseEvent) override;
-	virtual FReply OnDragDetected(const FGeometry& MyGeometry, const FPointerEvent& MouseEvent) override;	
+	virtual FReply OnMouseButtonUp(const FGeometry& MyGeometry, const FPointerEvent& PointerEvent) override;
+	virtual FReply OnDragDetected(const FGeometry& MyGeometry, const FPointerEvent& MouseEvent) override;
 	virtual void OnDragEnter(const FGeometry& MyGeometry, const FDragDropEvent& DragDropEvent) override;
 	virtual void OnDragLeave(const FDragDropEvent& DragDropEvent) override;
 	virtual FReply OnDrop(const FGeometry& MyGeometry, const FDragDropEvent& DragDropEvent) override;
 	// End SWidget interface
 
-	TSharedPtr<class SDMXChannel> ChannelValueWidget;
-
-	int32 ChannelID;
-
-	int32 UniverseID;
-
-protected:
-	/** Selects the fixture patch in Shared Data */
-	void SelectFixturePatch(UDMXEntityFixturePatch* FixturePatch);
-
-	/** Returns the Fixture Patch residing on this channel */
-	UDMXEntityFixturePatch* GetFixturePatch() const;
-
+private:
 	/** Returns the DMXLibrary or nullptr if not available */
 	UDMXLibrary* GetDMXLibrary() const;
 
 	int32 Column = INDEX_NONE;
 	int32 Row = INDEX_NONE;
 
-	FOnDragOverChannel OnDragEnterChannel;
-	FOnDragOverChannel OnDragLeaveChannel;
-	FOnDropOntoChannel OnDropOntoChannel;
+	TSharedPtr<class SDMXChannel> ChannelValueWidget;
+
+	int32 ChannelID;
+
+	int32 UniverseID;
 
 	TWeakPtr<FDMXEditor> DMXEditorPtr;
+
+	// Slate Arguments
+	FOnMouseEventWithReply OnMouseButtonDownOnChannel;
+	FOnMouseEventWithReply OnMouseButtonUpOnChannel;
+	FOnMouseEventWithReply OnDragDetectedOnChannel;
+	FOnDragEvent OnDragEnterChannel;
+	FOnDragEvent OnDragLeaveChannel;
+	FOnDragEventWithReply OnDropOntoChannel;
 };
diff --git a/Engine/Plugins/VirtualProduction/DMX/DMXEngine/Source/DMXEditor/Private/Widgets/FixturePatch/SDMXFixturePatchFragment.cpp b/Engine/Plugins/VirtualProduction/DMX/DMXEngine/Source/DMXEditor/Private/Widgets/FixturePatch/SDMXFixturePatchFragment.cpp
--- a/Engine/Plugins/VirtualProduction/DMX/DMXEngine/Source/DMXEditor/Private/Widgets/FixturePatch/SDMXFixturePatchFragment.cpp
+++ b/Engine/Plugins/VirtualProduction/DMX/DMXEngine/Source/DMXEditor/Private/Widgets/FixturePatch/SDMXFixturePatchFragment.cpp
@@ -3,8 +3,10 @@
 #include "SDMXFixturePatchFragment.h"
 
 #include "DMXEditor.h"
+#include "DMXEditorStyle.h"
 #include "DMXFixturePatchNode.h"
 #include "DMXFixturePatchEditorDefinitions.h"
+#include "DMXFixturePatchSharedData.h"
 #include "Library/DMXEntityFixturePatch.h"
 
 #include "Styling/AppStyle.h"
@@ -14,148 +16,277 @@
 #include "Widgets/SOverlay.h"
 #include "Widgets/Layout/SBorder.h"
 #include "Widgets/Layout/SSpacer.h"
+#include "Widgets/Layout/SWidgetSwitcher.h"
 
 #define LOCTEXT_NAMESPACE "SDMXFixturePatchFragment"
 
-BEGIN_SLATE_FUNCTION_BUILD_OPTIMIZATION
-void SDMXFixturePatchFragment::Construct(const FArguments& InArgs, TSharedPtr<FDMXFixturePatchNode> InPatchNode)
+
+void SDMXFixturePatchFragment::Construct(const FArguments& InArgs, const TSharedPtr<FDMXFixturePatchNode>& InFixturePatchNode, const TArray<TSharedPtr<FDMXFixturePatchNode>>& InFixturePatchNodeGroup)
 {
-	check(InPatchNode.IsValid());
+	check(InFixturePatchNode.IsValid());
+	check(InFixturePatchNodeGroup.Contains(InFixturePatchNode));
 	check(Column != -1);
 	check(Row != -1);
 	check(ColumnSpan != -1);
 
-	PatchNode = InPatchNode;
-		
+	ColorBrush = MakeShared<FSlateColorBrush>(FLinearColor::Red);
+
+	FixturePatchNode = InFixturePatchNode;
+	FixturePatchNodeGroup = InFixturePatchNodeGroup;
 	DMXEditorPtr = InArgs._DMXEditor;
+	bIsHead = InArgs._IsHead;
+	bIsTail = InArgs._IsTail;
+	bIsText = InArgs._IsText;
+	bIsConflicting = InArgs._IsConflicting;
 	Column = InArgs._Column;
 	Row = InArgs._Row;
 	ColumnSpan = InArgs._ColumnSpan;
 	bHighlight = InArgs._bHighlight;
 
-	SetVisibility(EVisibility::HitTestInvisible);
-
-	FMargin MinimalTextMargin = FMargin(3.0f, 2.0f, 4.0f, 1.0f);
-	ShadowSize = FAppStyle::GetVector(TEXT("Graph.Node.ShadowSize"));
+	if (bIsText)
+	{
+		SAssignNew(ContentBorder, SBorder)
+			.HAlign(HAlign_Fill)
+			.VAlign(VAlign_Fill)
+			.Padding(FMargin(3.0f, 2.0f, 0.f, 1.0f))
+			.Visibility(EVisibility::SelfHitTestInvisible);
+	}
+	else
+	{
+		SAssignNew(ContentBorder, SBorder)
+			.HAlign(HAlign_Fill)
+			.VAlign(VAlign_Fill)
+			.Visibility(EVisibility::SelfHitTestInvisible)
+			.BorderImage_Lambda([this]()
+				{
+					if (UDMXEntityFixturePatch* FixturePatch = FixturePatchNode->GetFixturePatch().Get())
+					{
+						if (bIsConflicting)
+						{
+							ColorBrush->TintColor = FixturePatch->EditorColor.CopyWithNewOpacity(.25f / FixturePatchNodeGroup.Num());
+						}
+						else if (bIsTopmost)
+						{
+							ColorBrush->TintColor = FixturePatch->EditorColor.CopyWithNewOpacity(.5f);
+						}
+						else
+						{
+							ColorBrush->TintColor = FLinearColor::Transparent;
+						}
+					}
+					check(ColorBrush.IsValid());
+					return ColorBrush.Get();
+				});
+	}
 
-	// We do not need graph node feats, but mimic its visuals
 	ChildSlot
 		[
-			SNew(SBox)
-			.ToolTipText(this, &SDMXFixturePatchFragment::GetToolTipText)
-			.MaxDesiredHeight(1.0f)
-			.MaxDesiredWidth(1.0f)			
+			SNew(SBox) 	
+			.Visibility(EVisibility::SelfHitTestInvisible)
+			.MaxDesiredHeight(1.0f) // Prevents from some wiggling issues
+			.MaxDesiredWidth(1.0f) // Prevents from some wiggling issues
+			.Padding_Lambda([this]()
+				{
+					return BorderPadding;
+				})
 			[
-				SNew(SOverlay)
-				+ SOverlay::Slot()
-				[
-					SNew(SImage)
-					.Image( FAppStyle::GetBrush("Graph.VarNode.Body") )		
-					.ColorAndOpacity( this, &SDMXFixturePatchFragment::GetColor )
-				]
-				+ SOverlay::Slot()
-				.VAlign(VAlign_Top)
-				[
-					SNew(SImage)
-					.Image( FAppStyle::GetBrush("Graph.VarNode.ColorSpill") )
-					.ColorAndOpacity( this, &SDMXFixturePatchFragment::GetColor )
-				]
-				+ SOverlay::Slot()
-				[
-					SNew(SImage)
-					.Image( FAppStyle::GetBrush("Graph.VarNode.Gloss") )
-					.ColorAndOpacity( this, &SDMXFixturePatchFragment::GetColor )
-				]
-				+ SOverlay::Slot()
-				[		
-					SNew(SBorder)
-					.HAlign(HAlign_Fill)
-					.VAlign(VAlign_Fill)
-					.BorderImage(FAppStyle::GetBrush("NoBorder"))
-					.Padding(MinimalTextMargin)
-					.BorderBackgroundColor( this, &SDMXFixturePatchFragment::GetColor )
-					.OnMouseButtonDown(this, &SDMXFixturePatchFragment::OnMouseButtonDown)
-					[	
-						SNew(STextBlock)
-						.Text(this, &SDMXFixturePatchFragment::GetText)
-						.TextStyle(FAppStyle::Get(), "SmallText")
-						.ColorAndOpacity(FLinearColor::White)						
-					]
-				]
-				+SOverlay::Slot()
-				.VAlign(VAlign_Top)
-				[
-					SNew(SBorder)
-					.Visibility(EVisibility::HitTestInvisible)			
-					.BorderImage( FAppStyle::GetBrush( "Graph.Node.TitleHighlight" ) )
-					.BorderBackgroundColor( FLinearColor::White )
-					[
-						SNew(SSpacer)
-						.Size(FVector2D(20,20))
-					]
+				SNew(SBorder)
+				.HAlign(HAlign_Fill)
+				.VAlign(VAlign_Fill)
+				.Visibility(EVisibility::SelfHitTestInvisible)
+				.BorderImage(&GetBorderImage())
+				.BorderBackgroundColor_Lambda([this]()
+					{
+						return bHighlight ? FLinearColor(1.f, 1.f, 1.0f) : FLinearColor(.2f, .2f, .2f);
+					})
+				[	
+					ContentBorder.ToSharedRef()
 				]
 			]
 		];
 
+	SetVisibility(EVisibility::SelfHitTestInvisible);
+	Refresh(InFixturePatchNodeGroup);
 }
-END_SLATE_FUNCTION_BUILD_OPTIMIZATION
 
-void SDMXFixturePatchFragment::SetHighlight(bool bEnabled)
+void SDMXFixturePatchFragment::Refresh(const TArray<TSharedPtr<FDMXFixturePatchNode>>& InFixturePatchNodeGroup)
 {
-	bHighlight = bEnabled;
+	FixturePatchNodeGroup = InFixturePatchNodeGroup;
+	
+	// Update if it is topmost first so further methods can rely on bIsTopmost
+	UpdateIsTopmost();
+
+	UpdateFixturePatchNameText();
+	UpdateNodeGroupText();
+	UpdateBorderPadding();
+
+	if (bIsText)
+	{
+		// Redraw text
+		ContentBorder->SetContent(CreateTextWidget());
+	}
 }
 
-int32 SDMXFixturePatchFragment::OnPaint(const FPaintArgs& Args, const FGeometry& AllottedGeometry, const FSlateRect& MyCullingRect, FSlateWindowElementList& OutDrawElements, int32 LayerId, const FWidgetStyle& InWidgetStyle, bool bParentEnabled) const
+void SDMXFixturePatchFragment::OnFixturePatchSharedDataSelectedFixturePatch()
 {
-	SCompoundWidget::OnPaint(Args, AllottedGeometry, MyCullingRect, OutDrawElements, LayerId, InWidgetStyle, bParentEnabled);
-
-	// Draw a shadow	
-	const FSlateBrush* ShadowBrush = GetShadowBrush(bHighlight);
-	FSlateDrawElement::MakeBox(
-		OutDrawElements,
-		LayerId,
-		AllottedGeometry.ToInflatedPaintGeometry(ShadowSize),
-		ShadowBrush
-	);
-
-	return LayerId;
+	TSharedPtr<FDMXEditor> DMXEditor = DMXEditorPtr.Pin();
+	if (!DMXEditor.IsValid())
+	{
+		return;
+	}
+
+	const TSharedPtr<FDMXFixturePatchSharedData> SharedData = DMXEditor->GetFixturePatchSharedData();
+	if (!SharedData.IsValid())
+	{
+		return;
+	}
+	const TArray<TWeakObjectPtr<UDMXEntityFixturePatch>>& SelectedFixturePatches = SharedData->GetSelectedFixturePatches();
 }
 
-FText SDMXFixturePatchFragment::GetText() const
+void SDMXFixturePatchFragment::UpdateFixturePatchNameText()
 {
-	check(PatchNode.IsValid());
-
-	FText Text;
-	if (UDMXEntityFixturePatch* Patch = PatchNode->GetFixturePatch().Get())
+	if (UDMXEntityFixturePatch* Patch = FixturePatchNode->GetFixturePatch().Get())
 	{
-		Text = FText::FromString(Patch->GetDisplayName());
+		FixturePatchNameText = FText::FromString(Patch->GetDisplayName());
 	}
+}
 
-	return Text;
+void SDMXFixturePatchFragment::UpdateNodeGroupText()
+{
+	NodeGroupText =
+		FixturePatchNodeGroup.Num() > 1 ?
+		FText::Format(LOCTEXT("NodeGroupText", "(and {0} more)"), FixturePatchNodeGroup.Num() - 1) :
+		FText::GetEmpty();
 }
 
-FText SDMXFixturePatchFragment::GetToolTipText() const
-{		
-	// Useful when the fragment cannot show the whole 
-	// name of the patch on top of it.
-	return GetText();
+void SDMXFixturePatchFragment::UpdateIsTopmost()
+{
+	if (!ensureMsgf(FixturePatchNodeGroup.Num() > 0, TEXT("Unexpected empty node group for Fixture Patch Fragment Widget.")))
+	{
+		return;
+	}
+
+	bIsTopmost = !FixturePatchNodeGroup.ContainsByPredicate([this](const TSharedPtr<FDMXFixturePatchNode>& OtherNode)
+		{
+			return OtherNode->GetZOrder() > FixturePatchNode->GetZOrder();
+		});
+
+	if (FixturePatchNodeGroup.Num() > 1)
+	{
+		UpdateFixturePatchNameText();
+		UE_LOG(LogTemp, Warning, TEXT("%s: I'm top!"), *FixturePatchNameText.ToString());
+	}
 }
 
-FSlateColor SDMXFixturePatchFragment::GetColor() const
+void SDMXFixturePatchFragment::UpdateBorderPadding()
 {
-	check(PatchNode.IsValid());
+	if (bIsHead && FixturePatchNodeGroup.Num() > 1)
+	{
+		// Sort by ZOrder ascending
+		FixturePatchNodeGroup.StableSort([](const TSharedPtr<FDMXFixturePatchNode>& NodeA, const TSharedPtr<FDMXFixturePatchNode>& NodeB)
+			{
+				return 
+					NodeA->GetZOrder() < NodeB->GetZOrder();
+			});
+
+		const int32 IndexOfThis = FixturePatchNodeGroup.IndexOfByKey(FixturePatchNode);
 
-	if (PatchNode->GetFixturePatch().IsValid())
+		constexpr int32 OffsetPerIndex = 3;
+		const int32 ClampedIndexOfThis = FMath::Min(IndexOfThis, 5);
+		const int32 OffsetLeft = ClampedIndexOfThis * OffsetPerIndex;
+
+		BorderPadding = FMargin(OffsetLeft, 0.f, 0.f, 0.f);
+	}
+	else
 	{
-		return PatchNode->GetFixturePatch()->EditorColor;
+		BorderPadding = FMargin(0.f);
 	}
+}
 
-	return FLinearColor::White;
+const FSlateBrush& SDMXFixturePatchFragment::GetBorderImage() const
+{
+	const FSlateBrush* BorderImage = nullptr;
+	if (bIsText)
+	{
+		BorderImage = FCoreStyle::Get().GetBrush("NoBorder");
+	}
+	else if (bIsConflicting)
+	{
+		BorderImage = FAppStyle::GetBrush("Graph.Node.DevelopmentBanner");
+	}
+	else if (bIsHead && bIsTail)
+	{
+		BorderImage = FCoreStyle::Get().GetBrush("Border");
+	}
+	else if (bIsHead)
+	{
+		BorderImage = FAppStyle::GetBrush("Profiler.EventGraph.Border.L");
+	}
+	else if (bIsTail)
+	{
+		BorderImage = FAppStyle::GetBrush("Profiler.EventGraph.Border.R");
+	}
+	else
+	{
+		BorderImage = FAppStyle::GetBrush("Profiler.EventGraph.Border.TB");
+	}
+	check(BorderImage);
+
+	return *BorderImage;
 }
 
-const FSlateBrush* SDMXFixturePatchFragment::GetShadowBrush(bool bSelected) const
+TSharedRef<SWidget> SDMXFixturePatchFragment::CreateTextWidget()
 {
-	return bSelected ? FAppStyle::GetBrush(TEXT("Graph.VarNode.ShadowSelected")) : FAppStyle::GetBrush(TEXT("Graph.VarNode.Shadow"));
+	if (bIsTopmost)
+	{
+		return
+			SNew(SBox)
+			.HAlign(HAlign_Left)
+			.VAlign(VAlign_Top)
+			.Visibility(EVisibility::HitTestInvisible)
+			[
+				SNew(SBorder)
+				.HAlign(HAlign_Left)
+				.VAlign(VAlign_Top)
+				.Padding(FMargin(2.f, 2.f, 0.f, 2.f))
+				.BorderImage(FDMXEditorStyle::Get().GetBrush("DMXEditor.RoundedPropertyBorder"))
+				[
+					SNew(SHorizontalBox)
+					.Clipping(EWidgetClipping::ClipToBoundsAlways)
+
+					+ SHorizontalBox::Slot()
+					.HAlign(HAlign_Left)
+					.VAlign(VAlign_Top)
+					.AutoWidth()
+					[
+						SNew(STextBlock)
+						.Text_Lambda([this]()
+							{
+								return FixturePatchNameText;
+							})
+						.TextStyle(FAppStyle::Get(), "SmallText")
+						.ColorAndOpacity(FLinearColor::White)
+					]
+				
+					+ SHorizontalBox::Slot()
+					.HAlign(HAlign_Left)
+					.VAlign(VAlign_Top)
+					.AutoWidth()
+					.Padding(FMargin(10.f, 0.f, 0.f, 0.f))
+					[
+						SNew(STextBlock)
+						.Text_Lambda([this]()
+							{
+								return NodeGroupText;
+							})
+						.TextStyle(FAppStyle::Get(), "SmallText")
+						.ColorAndOpacity(FLinearColor(.5f, .5f, .5f))
+					]
+				]
+			];
+	}
+
+	return SNullWidget::NullWidget;
 }
 
 #undef LOCTEXT_NAMESPACE
diff --git a/Engine/Plugins/VirtualProduction/DMX/DMXEngine/Source/DMXEditor/Private/Widgets/FixturePatch/SDMXFixturePatchFragment.h b/Engine/Plugins/VirtualProduction/DMX/DMXEngine/Source/DMXEditor/Private/Widgets/FixturePatch/SDMXFixturePatchFragment.h
--- a/Engine/Plugins/VirtualProduction/DMX/DMXEngine/Source/DMXEditor/Private/Widgets/FixturePatch/SDMXFixturePatchFragment.h
+++ b/Engine/Plugins/VirtualProduction/DMX/DMXEngine/Source/DMXEditor/Private/Widgets/FixturePatch/SDMXFixturePatchFragment.h
@@ -14,14 +14,21 @@ class FDMXFixturePatchNode;
 class UDMXEntityFixturePatch;
 class UDMXLibrary;
 
+struct FSlateColorBrush;
+class SBorder;
+template <typename OptionType> class SComboBox;
+
 
 class SDMXFixturePatchFragment
-	: public SCompoundWidget
+	: public SCompoundWidget 
 {
 public:
 	SLATE_BEGIN_ARGS(SDMXFixturePatchFragment)
 		: _DMXEditor(nullptr)
-		, _Text(FText::GetEmpty())
+		, _IsHead(true)
+		, _IsTail(true)
+		, _IsText(false)
+		, _IsConflicting(false)
 		, _Column(-1)
 		, _Row(-1)
 		, _ColumnSpan(1)
@@ -29,7 +36,13 @@ public:
 	{}
 		SLATE_ARGUMENT(TWeakPtr<FDMXEditor>, DMXEditor)
 
-		SLATE_ATTRIBUTE(FText, Text)
+		SLATE_ARGUMENT(bool, IsHead)
+
+		SLATE_ARGUMENT(bool, IsTail)
+
+		SLATE_ARGUMENT(bool, IsText)
+
+		SLATE_ARGUMENT(bool, IsConflicting)
 
 		SLATE_ARGUMENT(int32, Column)
 
@@ -42,11 +55,17 @@ public:
 	SLATE_END_ARGS()
 
 	/** Constructs this widget */
-	void Construct(const FArguments& InArgs, TSharedPtr<FDMXFixturePatchNode> InPatchNode);
+	void Construct(const FArguments& InArgs, const TSharedPtr<FDMXFixturePatchNode>& InFixturePatchNode, const TArray<TSharedPtr<FDMXFixturePatchNode>>& InFixturePatchNodeGroup);
+
+	/** Refreshes the widget */
+	void Refresh(const TArray<TSharedPtr<FDMXFixturePatchNode>>& InFixturePatchNodeGroup);
 
 	/** Sets wether the fragment should be highlit */
-	void SetHighlight(bool bEnabled);
-	
+	void SetHighlight(bool bEnabled) { bHighlight = bEnabled; }
+
+	/** Sets if the fragment is conflicting with others */
+	void SetConflicting(bool bConflicting) { bIsConflicting = bConflicting; if (Column == 0 && Row == 0 && bIsConflicting) UE_LOG(LogTemp, Warning, TEXT("%s set to conflicting"), *FixturePatchNameText.ToString()); }
+
 	/** Returns the Column of the fragment */
 	int32 GetColumn() const { return Column; }
 
@@ -65,40 +84,75 @@ public:
 	/** Sets the Column span of the fragment */
 	void SetColumnSpan(int32 NewColumnSpan) { ColumnSpan = NewColumnSpan; }
 
-protected:
-	// Begin SWidget interface
-	virtual int32 OnPaint(const FPaintArgs& Args, const FGeometry& AllottedGeometry, const FSlateRect& MyCullingRect, FSlateWindowElementList& OutDrawElements, int32 LayerId, const FWidgetStyle& InWidgetStyle, bool bParentEnabled) const override;
-	// End SWidget interface
+private:
+	/** Called when Fixture Patch Shared Data selected a Fixture Patch */
+	void OnFixturePatchSharedDataSelectedFixturePatch();
+
+	/** Updates the Fixture Patch Name Text */
+	void UpdateFixturePatchNameText();
+
+	/** Updates the Node Group Text */
+	void UpdateNodeGroupText();
 
-	/** Gets the text displayed on the fragment */
-	FText GetText() const;
+	/** Updates the Padding of the outermost Border */
+	void UpdateBorderPadding();
 
-	/** Gets the tooltip text displayed on the fragment */
-	FText GetToolTipText() const;
+	/** Updates the bTopmost memeber from the ZOrder of the node */
+	void UpdateIsTopmost();
 
-	/** Gets the color of the fragment */
-	FSlateColor GetColor() const;
+	/** Returns the Border image */
+	const FSlateBrush& GetBorderImage() const;
 
-	/** Gets the shadow brush of the fragment, depending on wether it's selected */
-	const FSlateBrush* GetShadowBrush(bool bSelected) const;
+	/** Gets the Text Widget */
+	TSharedRef<SWidget> CreateTextWidget();
 
-	/** Column of the fragment */
+	/** If true, the patch has the highest ZOrder in its node group */
+	bool bIsTopmost = false;
+
+	/** If true this is the Head of the Fixture Patch */
+	bool bIsHead = false;
+
+	/** If true this is the Tail of the Fixture Patch */
+	bool bIsTail = false;
+
+	/** If true, this fragment shows text */
+	bool bIsText;
+
+	/** Column of the Fragment */
 	int32 Column;
 
-	/** Row of the fragment */
+	/** Row of the Fragment */
 	int32 Row;
 
-	/** Column span of the fragment */
+	/** Column span of the Fragment */
 	int32 ColumnSpan;
 
-	/** If true, the widget shows a highlight */
+	/** If true, the Widget shows a highlight */
 	bool bHighlight = false;
 
-	/** Size of the shadow */
-	FVector2D ShadowSize;
+	/** If true, the Widget displays a conflict background */
+	bool bIsConflicting = false;
+
+	/** The Fixture Patch Name Text the Node displays */
+	FText FixturePatchNameText;
+
+	/** The Node Group Text the Node displays */
+	FText NodeGroupText;
+
+	/** The padding of the outermost border */
+	FMargin BorderPadding;
+	
+	/** The border that holds the color or text content */
+	TSharedPtr<SBorder> ContentBorder;
+
+	/** Color Brush the widget uses */
+	TSharedPtr<FSlateColorBrush> ColorBrush;
 
 	/** Fixture Patch Node being displayed by this widget */
-	TSharedPtr<FDMXFixturePatchNode> PatchNode;
+	TSharedPtr<FDMXFixturePatchNode> FixturePatchNode;
+
+	/** The group of nodes (same Universe, same Channel) this widget belongs to */
+	TArray<TSharedPtr<FDMXFixturePatchNode>> FixturePatchNodeGroup;
 
 	/** Weak DMXEditor refrence */
 	TWeakPtr<FDMXEditor> DMXEditorPtr;
diff --git a/Engine/Plugins/VirtualProduction/DMX/DMXEngine/Source/DMXEditor/Private/Widgets/FixturePatch/SDMXFixturePatcher.cpp b/Engine/Plugins/VirtualProduction/DMX/DMXEngine/Source/DMXEditor/Private/Widgets/FixturePatch/SDMXFixturePatcher.cpp
--- a/Engine/Plugins/VirtualProduction/DMX/DMXEngine/Source/DMXEditor/Private/Widgets/FixturePatch/SDMXFixturePatcher.cpp
+++ b/Engine/Plugins/VirtualProduction/DMX/DMXEngine/Source/DMXEditor/Private/Widgets/FixturePatch/SDMXFixturePatcher.cpp
@@ -31,146 +31,145 @@ void SDMXFixturePatcher::Construct(const FArguments& InArgs)
 {
 
 	DMXEditorPtr = InArgs._DMXEditor;
-
-	if (TSharedPtr<FDMXEditor> DMXEditor = DMXEditorPtr.Pin())
+	if (!DMXEditorPtr.IsValid())
 	{
-		SharedData = DMXEditor->GetFixturePatchSharedData();
-		check(SharedData.IsValid());
-		
-		const FLinearColor BackgroundTint(0.6f, 0.6f, 0.6f, 1.0f);
+		return;
+	}
 
-		ChildSlot			
+	SharedData = DMXEditorPtr.Pin()->GetFixturePatchSharedData();
+		
+	const FLinearColor BackgroundTint(0.6f, 0.6f, 0.6f, 1.0f);
+	ChildSlot			
+		[
+			SNew(SBox)
+			.HAlign(HAlign_Left)
+			.ToolTipText(this, &SDMXFixturePatcher::GetTooltipText)
 			[
-				SNew(SBox)
-				.HAlign(HAlign_Left)
-				.ToolTipText(this, &SDMXFixturePatcher::GetTooltipText)
-				[
-					SNew(SVerticalBox)				
+				SNew(SVerticalBox)				
 
-					// Settings area
+				// Settings area
+
+				+ SVerticalBox::Slot()
+				.HAlign(HAlign_Fill)
+				.AutoHeight()
+				[
 
-					+ SVerticalBox::Slot()
+					SNew(SBorder)					
 					.HAlign(HAlign_Fill)
-					.AutoHeight()
+					.BorderBackgroundColor(BackgroundTint)
+					.BorderImage(FAppStyle::GetBrush("DetailsView.CategoryTop"))
 					[
+						SNew(SHorizontalBox)			
 
-						SNew(SBorder)					
-						.HAlign(HAlign_Fill)
-						.BorderBackgroundColor(BackgroundTint)
-						.BorderImage(FAppStyle::GetBrush("DetailsView.CategoryTop"))
+						+ SHorizontalBox::Slot()						
+						.AutoWidth()
+						.VAlign(VAlign_Center)
+						.Padding(FMargin(4.0f, 4.0f, 15.0f, 4.0f))
 						[
-							SNew(SHorizontalBox)			
+							SNew(STextBlock)							
+							.MinDesiredWidth(75.0f)
+							.Font(FAppStyle::GetFontStyle("PropertyWindow.NormalFont"))
+							.TextStyle(FAppStyle::Get(), "DetailsView.CategoryTextStyle")
+							.IsEnabled(this, &SDMXFixturePatcher::IsUniverseSelectionEnabled)
+							.Text(LOCTEXT("UniverseSelectorLabel", "Universe"))
+						]
 
-							+ SHorizontalBox::Slot()						
-							.AutoWidth()
-							.VAlign(VAlign_Center)
-							.Padding(FMargin(4.0f, 4.0f, 15.0f, 4.0f))
+						+ SHorizontalBox::Slot()
+						.AutoWidth()
+						.VAlign(VAlign_Center)
+						.Padding(FMargin(4.0f, 4.0f, 15.0f, 4.0f))
+						[
+							SNew(SBox)
+							.MinDesiredWidth(210.0f)
+							.MaxDesiredWidth(420.0f)
 							[
-								SNew(STextBlock)							
-								.MinDesiredWidth(75.0f)
-								.Font(FAppStyle::GetFontStyle("PropertyWindow.NormalFont"))
-								.TextStyle(FAppStyle::Get(), "DetailsView.CategoryTextStyle")
+								SNew(SSpinBox<int32>)								
+								.SliderExponent(1000.0f)								
+								.MinSliderValue(0)
+								.MaxSliderValue(DMX_MAX_UNIVERSE - 1)
+								.MinValue(0)
+								.MaxValue(DMX_MAX_UNIVERSE - 1)
 								.IsEnabled(this, &SDMXFixturePatcher::IsUniverseSelectionEnabled)
-								.Text(LOCTEXT("UniverseSelectorLabel", "Universe"))
-							]
-
-							+ SHorizontalBox::Slot()
-							.AutoWidth()
-							.VAlign(VAlign_Center)
-							.Padding(FMargin(4.0f, 4.0f, 15.0f, 4.0f))
-							[
-								SNew(SBox)
-								.MinDesiredWidth(210.0f)
-								.MaxDesiredWidth(420.0f)
-								[
-									SNew(SSpinBox<int32>)								
-									.SliderExponent(1000.0f)								
-									.MinSliderValue(0)
-									.MaxSliderValue(DMX_MAX_UNIVERSE - 1)
-									.MinValue(0)
-									.MaxValue(DMX_MAX_UNIVERSE - 1)
-									.IsEnabled(this, &SDMXFixturePatcher::IsUniverseSelectionEnabled)
-									.Value(this, &SDMXFixturePatcher::GetSelectedUniverse)
-									.OnValueChanged(this, &SDMXFixturePatcher::SelectUniverse)
-								]
+								.Value(this, &SDMXFixturePatcher::GetSelectedUniverse)
+								.OnValueChanged(this, &SDMXFixturePatcher::SelectUniverse)
 							]
+						]
 						
-							+ SHorizontalBox::Slot()
-							.AutoWidth()
-							.VAlign(VAlign_Center)
-							.Padding(FMargin(4.0f, 4.0f, 15.0f, 4.0f))						
-							[
-								SNew(SSeparator)
-								.Orientation(EOrientation::Orient_Vertical)
-							]
+						+ SHorizontalBox::Slot()
+						.AutoWidth()
+						.VAlign(VAlign_Center)
+						.Padding(FMargin(4.0f, 4.0f, 15.0f, 4.0f))						
+						[
+							SNew(SSeparator)
+							.Orientation(EOrientation::Orient_Vertical)
+						]
 
-							+ SHorizontalBox::Slot()
-							.AutoWidth()
-							.VAlign(VAlign_Center)
-							.Padding(FMargin(4.0f, 4.0f, 15.0f, 4.0f))
-							[
-								SNew(STextBlock)
-								.Font(FAppStyle::GetFontStyle("PropertyWindow.NormalFont"))
-								.Text(LOCTEXT("UniverseDisplayAllText", "Show all patched Universes"))
-							]
+						+ SHorizontalBox::Slot()
+						.AutoWidth()
+						.VAlign(VAlign_Center)
+						.Padding(FMargin(4.0f, 4.0f, 15.0f, 4.0f))
+						[
+							SNew(STextBlock)
+							.Font(FAppStyle::GetFontStyle("PropertyWindow.NormalFont"))
+							.Text(LOCTEXT("UniverseDisplayAllText", "Show all patched Universes"))
+						]
 
-							+ SHorizontalBox::Slot()
-							.AutoWidth()
-							.VAlign(VAlign_Center)
-							.Padding(FMargin(4.0f, 4.0f, 15.0f, 4.0f))
-							[
-								SAssignNew(ShowAllUniversesCheckBox, SCheckBox)
-								.IsChecked(false)
-								.OnCheckStateChanged(this, &SDMXFixturePatcher::OnToggleDisplayAllUniverses)
-							]
+						+ SHorizontalBox::Slot()
+						.AutoWidth()
+						.VAlign(VAlign_Center)
+						.Padding(FMargin(4.0f, 4.0f, 15.0f, 4.0f))
+						[
+							SAssignNew(ShowAllUniversesCheckBox, SCheckBox)
+							.IsChecked(false)
+							.OnCheckStateChanged(this, &SDMXFixturePatcher::OnToggleDisplayAllUniverses)
 						]
 					]
+				]
 
-					// Patched Universes
+				// Patched Universes
 
-					+ SVerticalBox::Slot()
-					.HAlign(HAlign_Left)
-					.VAlign(VAlign_Fill)
-					[	
-						SAssignNew(PatchedUniverseScrollBox, SScrollBox)					
-						.Orientation(EOrientation::Orient_Vertical)	
-						.ScrollBarAlwaysVisible(true)
-						.ScrollBarThickness(FVector2D(5.f, 0.f))
-					]
+				+ SVerticalBox::Slot()
+				.HAlign(HAlign_Left)
+				.VAlign(VAlign_Fill)
+				[	
+					SAssignNew(PatchedUniverseScrollBox, SScrollBox)					
+					.Orientation(EOrientation::Orient_Vertical)	
+					.ScrollBarAlwaysVisible(true)
+					.ScrollBarThickness(FVector2D(5.f, 0.f))
 				]
-			];
+			]
+		];
 
-		// Bind to selection changes
-		SharedData->OnFixturePatchSelectionChanged.AddSP(this, &SDMXFixturePatcher::OnFixturePatchSelectionChanged);
-		SharedData->OnUniverseSelectionChanged.AddSP(this, &SDMXFixturePatcher::OnUniverseSelectionChanged);
+	// Bind to selection changes
+	SharedData->OnFixturePatchSelectionChanged.AddSP(this, &SDMXFixturePatcher::OnFixturePatchSelectionChanged);
+	SharedData->OnUniverseSelectionChanged.AddSP(this, &SDMXFixturePatcher::OnUniverseSelectionChanged);
 
-		// If the selected universe has no patches, try to find one with patches instead
-		UDMXLibrary* Library = GetDMXLibrary();
-		check(Library);
+	// If the selected universe has no patches, try to find one with patches instead
+	UDMXLibrary* Library = GetDMXLibrary();
+	check(Library);
 
-		// Bind to object changes
-		UDMXEntityFixturePatch::GetOnFixturePatchChanged().AddSP(this, &SDMXFixturePatcher::OnFixturePatchChanged);
-		UDMXEntityFixtureType::GetOnFixtureTypeChanged().AddSP(this, &SDMXFixturePatcher::OnFixtureTypeChanged);
-		Library->GetOnEntitiesAdded().AddSP(this, &SDMXFixturePatcher::OnEntitiesAddedOrRemoved);
-		Library->GetOnEntitiesRemoved().AddSP(this, &SDMXFixturePatcher::OnEntitiesAddedOrRemoved);
+	// Bind to object changes
+	UDMXEntityFixtureType::GetOnFixtureTypeChanged().AddSP(this, &SDMXFixturePatcher::OnFixtureTypeChanged);
+	Library->GetOnEntitiesAdded().AddSP(this, &SDMXFixturePatcher::OnEntitiesAddedOrRemoved);
+	Library->GetOnEntitiesRemoved().AddSP(this, &SDMXFixturePatcher::OnEntitiesAddedOrRemoved);
 
-		TArray<UDMXEntityFixturePatch*> Patches = Library->GetEntitiesTypeCast<UDMXEntityFixturePatch>();
-		UDMXEntityFixturePatch** ExistingPatchPtr = Patches.FindByPredicate([&](UDMXEntityFixturePatch* Patch) {
-			return Patch->GetUniverseID() == SharedData->GetSelectedUniverse();
-			});
-		if (!ExistingPatchPtr && Patches.Num() > 0)
-		{
-			SharedData->SelectUniverse(Patches[0]->GetUniverseID());
-		}		
+	TArray<UDMXEntityFixturePatch*> Patches = Library->GetEntitiesTypeCast<UDMXEntityFixturePatch>();
+	UDMXEntityFixturePatch** ExistingPatchPtr = Patches.FindByPredicate([&](UDMXEntityFixturePatch* Patch) {
+		return Patch->GetUniverseID() == SharedData->GetSelectedUniverse();
+		});
+	if (!ExistingPatchPtr && Patches.Num() > 0)
+	{
+		SharedData->SelectUniverse(Patches[0]->GetUniverseID());
+	}		
 
-		// Bind to tabs being switched
-		FGlobalTabmanager::Get()->OnActiveTabChanged_Subscribe(FOnActiveTabChanged::FDelegate::CreateSP(this, &SDMXFixturePatcher::OnActiveTabChanged));
+	// Bind to tabs being switched
+	FGlobalTabmanager::Get()->OnActiveTabChanged_Subscribe(FOnActiveTabChanged::FDelegate::CreateSP(this, &SDMXFixturePatcher::OnActiveTabChanged));
 
 
-		GEditor->RegisterForUndo(this);
+	GEditor->RegisterForUndo(this);
+
+	ShowSelectedUniverse();
 
-		ShowSelectedUniverse();
-	}
 }
 END_SLATE_FUNCTION_BUILD_OPTIMIZATION
 
@@ -184,43 +183,6 @@ void SDMXFixturePatcher::RefreshFromLibrary()
 	RefreshFixturePatchState = EDMXRefreshFixturePatcherState::RefreshFromLibrary;
 }
 
-void SDMXFixturePatcher::SelectUniverseThatContainsSelectedPatches()
-{
-	UDMXLibrary* Library = GetDMXLibrary();
-	if (Library)
-	{
-		// If the selected universe no longer contains a patch, select another universe with patches		
-		check(SharedData.IsValid());
-		const TArray<TWeakObjectPtr<UDMXEntityFixturePatch>>& SelectedFixturePatches = SharedData->GetSelectedFixturePatches();
-
-		if (SelectedFixturePatches.Num() == 0)
-		{
-			return;
-		}
-
-		const int32 SelectedUniverseID = GetSelectedUniverse();
-		const int32 PatchInSelectedUniverseIndex = SelectedFixturePatches.IndexOfByPredicate([SelectedUniverseID](TWeakObjectPtr<UDMXEntityFixturePatch> Patch) {
-			return 
-				Patch.IsValid() && 
-				Patch->GetUniverseID() == SelectedUniverseID;
-			});
-		
-		const bool bIsAnySelectedPatchInSelectedUniverse = PatchInSelectedUniverseIndex != INDEX_NONE;
-		if (!bIsAnySelectedPatchInSelectedUniverse)
-		{
-			// Select arbitrary valid universe in selection
-			for(TWeakObjectPtr<UDMXEntityFixturePatch> SelectedPatch : SelectedFixturePatches)
-			{
-				if(SelectedPatch->GetUniverseID() >= 0)
-				{
-					SharedData->SelectUniverse(SelectedPatch->GetUniverseID());
-				}
-				break;
-			}
-		}
-	}
-}
-
 void SDMXFixturePatcher::OnActiveTabChanged(TSharedPtr<SDockTab> PreviouslyActive, TSharedPtr<SDockTab> NewlyActivated)
 {
 	if (!NewlyActivated.IsValid() ||
@@ -302,11 +264,7 @@ void SDMXFixturePatcher::OnDragEnterChannel(int32 UniverseID, int32 ChannelID, c
 					return ClampStartingChannel(StartingChannel, ChannelSpan);
 				}();
 
-				// Update the ChannelOffset
-				const int32 NewChannelOffset = ChannelID - NewStartingChannel;
-				FixturePatchDragDropOp->SetChannelOffset(NewChannelOffset);
-
-				// Patch the node but do not transact it (transact on drop or leave instead)
+				// Patch the node but do not transact it (transact on drop instead)
 				const TSharedPtr<SDMXPatchedUniverse>& Universe = PatchedUniversesByID.FindChecked(UniverseID);
 
 				const bool bCreateTransaction = false;
@@ -376,15 +334,22 @@ FReply SDMXFixturePatcher::OnDropOntoChannel(int32 UniverseID, int32 ChannelID,
 
 		if (DraggedNode.IsValid())
 		{
-			// Compensate Drag Offset
-			ChannelID = ChannelID - FixturePatchDragDropOp->GetChannelOffset();
+			if (UDMXEntityFixturePatch* FixturePatch = Cast<UDMXEntityFixturePatch>(DraggedEntities[0]))
+			{
+				const int32 ChannelSpan = FixturePatch->GetChannelSpan();
+				const int32 NewStartingChannel = [this, ChannelID, &FixturePatchDragDropOp, ChannelSpan]()
+				{
+					int32 StartingChannel = ChannelID - FixturePatchDragDropOp->GetChannelOffset();
+					return ClampStartingChannel(StartingChannel, ChannelSpan);
+				}();
 
-			const TSharedPtr<SDMXPatchedUniverse>& Universe = PatchedUniversesByID.FindChecked(UniverseID);
+				const TSharedPtr<SDMXPatchedUniverse>& Universe = PatchedUniversesByID.FindChecked(UniverseID);
 
-			bool bCreateTransaction = true;
-			if (Universe->Patch(DraggedNode, ChannelID, bCreateTransaction))
-			{
-				return FReply::Handled().EndDragDrop();
+				bool bCreateTransaction = true;
+				if (Universe->Patch(DraggedNode, NewStartingChannel, bCreateTransaction))
+				{
+					return FReply::Handled().EndDragDrop();
+				}
 			}
 		}
 	}
@@ -398,6 +363,7 @@ TSharedPtr<FDMXFixturePatchNode> SDMXFixturePatcher::GetDraggedNode(const TArray
 	{
 		if (UDMXEntityFixturePatch* FixturePatch = Cast<UDMXEntityFixturePatch>(DraggedEntities[0]))
 		{
+			const TArray<TWeakObjectPtr<UDMXEntityFixturePatch>> FixturePatchArray = TArray<TWeakObjectPtr<UDMXEntityFixturePatch>>({ FixturePatch });
 			TSharedPtr<FDMXFixturePatchNode> DraggedNode = FindPatchNode(FixturePatch);
 
 			if (!DraggedNode.IsValid())
@@ -470,21 +436,19 @@ void SDMXFixturePatcher::PostRedo(bool bSuccess)
 	RefreshFromProperties();
 }
 
-TSharedPtr<FDMXFixturePatchNode> SDMXFixturePatcher::FindPatchNode(TWeakObjectPtr<UDMXEntityFixturePatch> Patch) const
+TSharedPtr<FDMXFixturePatchNode> SDMXFixturePatcher::FindPatchNode(const TWeakObjectPtr<UDMXEntityFixturePatch> FixturePatch) const
 {
-	if (Patch.IsValid())
+	TSharedPtr<FDMXFixturePatchNode> ExistingNode;
+	for (const TPair<int32, TSharedPtr<SDMXPatchedUniverse>>& UniverseByID : PatchedUniversesByID)
 	{
-		TSharedPtr<FDMXFixturePatchNode> ExistingNode;
-		for (const TPair<int32, TSharedPtr<SDMXPatchedUniverse>>& UniverseByID : PatchedUniversesByID)
-		{
-			ExistingNode = UniverseByID.Value->FindPatchNode(Patch);
+		ExistingNode = UniverseByID.Value->FindPatchNode(FixturePatch);
 
-			if (ExistingNode.IsValid())
-			{
-				return ExistingNode;
-			}
+		if (ExistingNode.IsValid())
+		{
+			return ExistingNode;
 		}
 	}
+
 	return nullptr;
 }
 
@@ -509,29 +473,6 @@ void SDMXFixturePatcher::OnEntitiesAddedOrRemoved(UDMXLibrary* DMXLibrary, TArra
 	}
 }
 
-void SDMXFixturePatcher::OnFixturePatchChanged(const UDMXEntityFixturePatch* FixturePatch)
-{
-	if (IsValid(FixturePatch))
-	{
-		// Const cast is ok here, as long as this function only updates the widget
-		UDMXEntityFixturePatch* NonConstFixturePatch = const_cast<UDMXEntityFixturePatch*>(FixturePatch);
-
-		if (const TSharedPtr<FDMXFixturePatchNode> PatchNode = FindPatchNode(NonConstFixturePatch))
-		{
-			if (FixturePatch->GetUniverseID() != PatchNode->GetUniverseID() ||
-				FixturePatch->GetStartingChannel() != PatchNode->GetStartingChannel() ||
-				FixturePatch->GetChannelSpan() != PatchNode->GetChannelSpan())
-			{
-				RefreshFromProperties();
-			}
-		}
-		else if (FixturePatch->GetParentLibrary() == GetDMXLibrary())
-		{
-			RefreshFromProperties();
-		}
-	}
-}
-
 void SDMXFixturePatcher::OnFixtureTypeChanged(const UDMXEntityFixtureType* FixtureType)
 {
 	RefreshFromLibrary();
@@ -539,44 +480,53 @@ void SDMXFixturePatcher::OnFixtureTypeChanged(const UDMXEntityFixtureType* Fixtu
 
 void SDMXFixturePatcher::OnFixturePatchSelectionChanged()
 {
-	check(SharedData.IsValid());
-	const TArray<TWeakObjectPtr<UDMXEntityFixturePatch>>& SelectedPatches = SharedData->GetSelectedFixturePatches();
+	const int32 SelectedUniverse = SharedData->GetSelectedUniverse();
+	if (PatchedUniversesByID.Contains(SelectedUniverse))
+	{
+		return;
+	}
 
-	// Only RefreshFromProperties if a node for a selected patch doesn't exist
-	// This avoids issues when a patch gets selected and detect drag is pending.	
-	for (TWeakObjectPtr<UDMXEntityFixturePatch> Patch : SelectedPatches)
+	if (IsUniverseSelectionEnabled())
+	{
+		constexpr bool bReconstructWidget = true;
+		ShowSelectedUniverse(bReconstructWidget);
+	}
+	else
 	{
-		if (!FindPatchNode(Patch).IsValid())
+		const TSharedPtr<SDMXPatchedUniverse>* const UniverseWidgetPtr = PatchedUniversesByID.Find(SelectedUniverse);
+		if (UniverseWidgetPtr)
 		{
-			RefreshFromProperties();
-			break;
+			PatchedUniverseScrollBox->ScrollDescendantIntoView(*UniverseWidgetPtr);
 		}
 	}
-
-	SelectUniverseThatContainsSelectedPatches();
 }
 
 void SDMXFixturePatcher::OnUniverseSelectionChanged()
 {
+	const int32 SelectedUniverse = SharedData->GetSelectedUniverse();
+	if (PatchedUniversesByID.Contains(SelectedUniverse))
+	{
+		return;
+	}
+
 	if (IsUniverseSelectionEnabled())
 	{
 		ShowSelectedUniverse();
 	}
-	else
+	else if(!PatchedUniversesByID.Contains(SharedData->GetSelectedUniverse()))
 	{
-		// The newly selected universe is not yet shown and may contain a patch.
-		// If so, show all universes anew, to include the newly selected universe.
-		check(SharedData.IsValid());
-		if (!PatchedUniversesByID.Contains(SharedData->GetSelectedUniverse()))
-		{
-			ShowAllPatchedUniverses();
-		}
+		ShowAllPatchedUniverses();
 	}
+
+	if (!ensureMsgf(PatchedUniversesByID.Contains(SelectedUniverse), TEXT("Tried to scroll universe into view after displaying it, but universe doesn't exist. This is not expected.")))
+	{
+		return;
+	}
+	PatchedUniverseScrollBox->ScrollDescendantIntoView(PatchedUniversesByID[SelectedUniverse]);
 }
 
 void SDMXFixturePatcher::SelectUniverse(int32 NewUniverseID)
 {
-	check(SharedData.IsValid());
 	UniverseToSetNextTick = NewUniverseID;
 }
 
@@ -590,6 +540,27 @@ void SDMXFixturePatcher::ShowSelectedUniverse(bool bForceReconstructWidget)
 {
 	int32 SelectedUniverseID = GetSelectedUniverse();
 
+	/** Select a patched universe if the selected Universe doesn't contain any universes */
+	UDMXLibrary* DMXLibrary = GetDMXLibrary();
+	if (!DMXLibrary)
+	{
+		return;
+	}
+	TArray<UDMXEntityFixturePatch*> FixturePatches = DMXLibrary->GetEntitiesTypeCast<UDMXEntityFixturePatch>();
+	const bool bIsSelectedUniversePatched = FixturePatches.ContainsByPredicate([SelectedUniverseID](const UDMXEntityFixturePatch* FixturePatch)
+		{
+			return FixturePatch->GetUniverseID() == SelectedUniverseID;
+		});
+	if (!bIsSelectedUniversePatched && FixturePatches.Num() > 0)
+	{
+		FixturePatches.Sort([](const UDMXEntityFixturePatch& FixturePatchA, const UDMXEntityFixturePatch& FixturePatchB)
+			{
+				return FixturePatchA.GetUniverseID() < FixturePatchB.GetUniverseID();
+			});
+		SelectedUniverseID = FixturePatches[0]->GetUniverseID();
+	}
+
+	// Force reconstruct if requested
 	if (bForceReconstructWidget)
 	{
 		PatchedUniverseScrollBox->ClearChildren();
@@ -599,7 +570,7 @@ void SDMXFixturePatcher::ShowSelectedUniverse(bool bForceReconstructWidget)
 	// Create a new patched universe if required
 	if (PatchedUniversesByID.Num() != 1)
 	{
-		TSharedRef<SDMXPatchedUniverse> NewPatchedUniverse =
+		const TSharedRef<SDMXPatchedUniverse> NewPatchedUniverse =
 			SNew(SDMXPatchedUniverse)
 			.DMXEditor(DMXEditorPtr)
 			.UniverseID(SelectedUniverseID)
@@ -724,7 +695,8 @@ void SDMXFixturePatcher::AddUniverse(int32 UniverseID)
 
 void SDMXFixturePatcher::OnToggleDisplayAllUniverses(ECheckBoxState CheckboxState)
 {
-	bool bForceReconstructWidget = true;
+	constexpr bool bForceReconstructWidget = true;
+	const int32 SelectedUniverse = SharedData->GetSelectedUniverse();
 
 	switch (ShowAllUniversesCheckBox->GetCheckedState())
 	{
@@ -733,7 +705,7 @@ void SDMXFixturePatcher::OnToggleDisplayAllUniverses(ECheckBoxState CheckboxStat
 		return;
 
 	case ECheckBoxState::Unchecked:
-		SelectUniverseThatContainsSelectedPatches();
+		SelectUniverse(SelectedUniverse);
 		ShowSelectedUniverse(bForceReconstructWidget);
 		return;
 
diff --git a/Engine/Plugins/VirtualProduction/DMX/DMXEngine/Source/DMXEditor/Private/Widgets/FixturePatch/SDMXFixturePatcher.h b/Engine/Plugins/VirtualProduction/DMX/DMXEngine/Source/DMXEditor/Private/Widgets/FixturePatch/SDMXFixturePatcher.h
--- a/Engine/Plugins/VirtualProduction/DMX/DMXEngine/Source/DMXEditor/Private/Widgets/FixturePatch/SDMXFixturePatcher.h
+++ b/Engine/Plugins/VirtualProduction/DMX/DMXEngine/Source/DMXEditor/Private/Widgets/FixturePatch/SDMXFixturePatcher.h
@@ -48,9 +48,6 @@ public:
 	/** Refreshes the whole view from the library, considers library changes */
 	void RefreshFromLibrary();
 
-	/** Selects a universe that contains selected patches, if possible */
-	void SelectUniverseThatContainsSelectedPatches();
-
 protected:
 	/** Called when the active tab changed */
 	void OnActiveTabChanged(TSharedPtr<SDockTab> PreviouslyActive, TSharedPtr<SDockTab> NewlyActivated);
@@ -78,7 +75,7 @@ protected:
 	// End of FEditorUndoClient
 
 	/** Finds specified fixture patch in all universes */
-	TSharedPtr<FDMXFixturePatchNode> FindPatchNode(TWeakObjectPtr<UDMXEntityFixturePatch> Patch) const;
+	TSharedPtr<FDMXFixturePatchNode> FindPatchNode(const TWeakObjectPtr<UDMXEntityFixturePatch> FixturePatch) const;
 
 	/** Returns first node with same fixture type */
 	TSharedPtr<FDMXFixturePatchNode> FindPatchNodeOfType(UDMXEntityFixtureType* Type, const TSharedPtr<FDMXFixturePatchNode>& IgoredNode) const;
@@ -86,9 +83,6 @@ protected:
 	/** Called when entities were added to or removed from a DMX Library */
 	void OnEntitiesAddedOrRemoved(UDMXLibrary* DMXLibrary, TArray<UDMXEntity*> Entities);
 
-	/** Called when a fixture patch changed */
-	void OnFixturePatchChanged(const UDMXEntityFixturePatch* FixturePatch);
-
 	/** Called when a fixture type changed */
 	void OnFixtureTypeChanged(const UDMXEntityFixtureType* FixtureType);
 
diff --git a/Engine/Plugins/VirtualProduction/DMX/DMXEngine/Source/DMXEditor/Private/Widgets/FixturePatch/SDMXPatchedUniverse.cpp b/Engine/Plugins/VirtualProduction/DMX/DMXEngine/Source/DMXEditor/Private/Widgets/FixturePatch/SDMXPatchedUniverse.cpp
--- a/Engine/Plugins/VirtualProduction/DMX/DMXEngine/Source/DMXEditor/Private/Widgets/FixturePatch/SDMXPatchedUniverse.cpp
+++ b/Engine/Plugins/VirtualProduction/DMX/DMXEngine/Source/DMXEditor/Private/Widgets/FixturePatch/SDMXPatchedUniverse.cpp
@@ -3,20 +3,24 @@
 #include "SDMXPatchedUniverse.h"
 
 #include "DMXEditor.h"
-#include "DMXFixturePatchNode.h"
 #include "DMXFixturePatchEditorDefinitions.h"
+#include "DMXFixturePatchNode.h"
+#include "DMXFixturePatchSharedData.h"
 #include "SDMXChannelConnector.h"
 #include "SDMXFixturePatchFragment.h"
+#include "DragDrop/DMXEntityFixturePatchDragDropOp.h"
 #include "IO/DMXInputPort.h"
 #include "IO/DMXOutputPort.h"
 #include "IO/DMXPortManager.h"
 #include "Library/DMXEntityFixturePatch.h"
 #include "Library/DMXLibrary.h"
 
+#include "Framework/Commands/GenericCommands.h"
+#include "Framework/Commands/UICommandList.h"
+#include "Framework/MultiBox/MultiBoxBuilder.h"
 #include "Styling/AppStyle.h"
 #include "SlateOptMacros.h"
 #include "Widgets/Layout/SGridPanel.h"
-#include "Widgets/Layout/SScaleBox.h"
 
 
 #define LOCTEXT_NAMESPACE "SDMXPatchedUniverse"
@@ -33,74 +37,94 @@ void SDMXPatchedUniverse::Construct(const FArguments& InArgs)
 	OnDragLeaveChannel = InArgs._OnDragLeaveChannel;
 	OnDropOntoChannel = InArgs._OnDropOntoChannel;
 
-	if (TSharedPtr<FDMXEditor> DMXEditor = DMXEditorPtr.Pin())
+	const TSharedPtr<FDMXEditor> DMXEditor = DMXEditorPtr.Pin();
+	if (!DMXEditor.IsValid())
 	{
-		ChildSlot
-		[
-			SNew (SVerticalBox)
+		return;
+	}
+	SharedData = DMXEditor->GetFixturePatchSharedData();
+	check(SharedData.IsValid());
 
-			+ SVerticalBox::Slot()
-			.HAlign(HAlign_Fill)
-			.AutoHeight()
+	ChildSlot
+	[
+		SNew (SVerticalBox)
+
+		+ SVerticalBox::Slot()
+		.HAlign(HAlign_Fill)
+		.AutoHeight()
+		[
+			SNew(SOverlay)
+			+ SOverlay::Slot()
 			[
-				SNew(SOverlay)
-				+ SOverlay::Slot()
+				SAssignNew(UniverseName, SBorder)
+				.VAlign(VAlign_Top)
+				.HAlign(HAlign_Fill)
+				.BorderBackgroundColor(FLinearColor(0.6f, 0.6f, 0.6f, 0.3f))
+				.ToolTipText(FText::Format(LOCTEXT("UniverseListCategoryTooltip", "Patches assigned to Universe {0}"), UniverseID))
 				[
-					SAssignNew(UniverseName, SBorder)
-					.VAlign(VAlign_Top)
-					.HAlign(HAlign_Fill)
-					.BorderBackgroundColor(FLinearColor(0.6f, 0.6f, 0.6f, 0.3f))
-					.ToolTipText(FText::Format(LOCTEXT("UniverseListCategoryTooltip", "Patches assigned to Universe {0}"), UniverseID))
-					[
-						SNew(STextBlock)
-						.Font(FAppStyle::GetFontStyle("PropertyWindow.NormalFont"))
-						.TextStyle(FAppStyle::Get(), "DetailsView.CategoryTextStyle")
-						.Text(this, &SDMXPatchedUniverse::GetHeaderText)
-					]
+					SNew(STextBlock)
+					.Font(FAppStyle::GetFontStyle("PropertyWindow.NormalFont"))
+					.TextStyle(FAppStyle::Get(), "DetailsView.CategoryTextStyle")
+					.Text(this, &SDMXPatchedUniverse::GetHeaderText)
 				]
+			]
 
-				+ SOverlay::Slot()
-				[
-					SNew(SBorder)
-					.BorderImage(FAppStyle::GetBrush("Graph.Node.DevelopmentBanner"))
-					.HAlign(HAlign_Fill)
-					.VAlign(VAlign_Fill)
-					.Visibility(this, &SDMXPatchedUniverse::GetPatchedUniverseReachabilityBannerVisibility)
-				]
+			+ SOverlay::Slot()
+			[
+				SNew(SBorder)
+				.BorderImage(FAppStyle::GetBrush("Graph.Node.DevelopmentBanner"))
+				.HAlign(HAlign_Fill)
+				.VAlign(VAlign_Fill)
+				.Visibility(this, &SDMXPatchedUniverse::GetPatchedUniverseReachabilityBannerVisibility)
 			]
+		]
 			
-			+ SVerticalBox::Slot()
-			.VAlign(VAlign_Top)
-			.HAlign(HAlign_Fill)		
-			[					
-				SNew(SScaleBox)	
-				.VAlign(VAlign_Top)
-				.Stretch(EStretch::ScaleToFitX)						
-				.StretchDirection(EStretchDirection::DownOnly)				
-				[
-					SNew(SBox)
-					.HAlign(HAlign_Center)
-					.MinDesiredWidth(920.0f) // Avoids a slate issue, where SScaleBox ignores border size hence scales			
-					.Padding(FMargin(4.0f, 4.0f, 4.0f, 8.0f))
-					[
-						SAssignNew(Grid, SGridPanel)												
-					]
-				]
+		+ SVerticalBox::Slot()
+		.VAlign(VAlign_Top)
+		.HAlign(HAlign_Fill)		
+		[					
+			SNew(SBox)
+			.HAlign(HAlign_Center)
+			.MinDesiredWidth(920.0f) // Avoids a slate issue, where SScaleBox ignores border size hence scales			
+			.Padding(FMargin(4.0f, 4.0f, 4.0f, 8.0f))
+			[
+				SAssignNew(Grid, SGridPanel)												
 			]
-		];
+		]
+	];
 
-		CreateChannelConnectors();
+	CreateChannelConnectors();
 
-		SetUniverseID(UniverseID);
+	SetUniverseID(UniverseID);
 
-		UpdatePatchedUniverseReachability();
+	UpdateZOrderOfNodes();
+	UpdatePatchedUniverseReachability();
 
-		// Bind to port changes
-		FDMXPortManager::Get().OnPortsChanged.AddSP(this, &SDMXPatchedUniverse::UpdatePatchedUniverseReachability);
-	}
+	RegisterCommands();
+
+	// Handle external changes
+	UDMXEntityFixturePatch::GetOnFixturePatchChanged().AddSP(this, &SDMXPatchedUniverse::OnFixturePatchChanged);
+	SharedData->OnFixturePatchSelectionChanged.AddSP(this, &SDMXPatchedUniverse::OnSelectionChanged);
+	FDMXPortManager::Get().OnPortsChanged.AddSP(this, &SDMXPatchedUniverse::UpdatePatchedUniverseReachability);
 }
 END_SLATE_FUNCTION_BUILD_OPTIMIZATION
 
+void SDMXPatchedUniverse::RequestRefresh()
+{
+	if (RequestRefreshTimerHandle.IsValid())
+	{
+		return;
+	}
+
+	RequestRefreshTimerHandle =
+		GEditor->GetTimerManager()->SetTimerForNextTick(
+			FTimerDelegate::CreateLambda([this]()
+				{
+					RefreshInternal();
+				})
+		);
+}
+
 void SDMXPatchedUniverse::SetUniverseID(int32 NewUniverseID)
 {
 	check(NewUniverseID > -1);
@@ -130,22 +154,22 @@ void SDMXPatchedUniverse::SetUniverseID(int32 NewUniverseID)
 				PatchesInUniverse.Add(Patch);
 			}
 		});
+		PatchesInUniverse.Sort([](const UDMXEntityFixturePatch& FixturePatchA, const UDMXEntityFixturePatch& FixturePatchB)
+			{
+				return FixturePatchA.GetEndingChannel() <= FixturePatchB.GetEndingChannel();
+			});
 
 		// Add fixture patches to the grid
 		for (UDMXEntityFixturePatch* FixturePatch : PatchesInUniverse)
 		{
-			if (CanAssignFixturePatch(FixturePatch))
+			TSharedPtr<FDMXFixturePatchNode> Node = FindPatchNode(FixturePatch);
+			if (!Node.IsValid())
 			{
-				TSharedPtr<FDMXFixturePatchNode> Node = FindPatchNode(FixturePatch);
-
-				if (!Node.IsValid())
-				{
-					Node = FDMXFixturePatchNode::Create(DMXEditorPtr, FixturePatch);
-				}
-				check(Node.IsValid());
-
-				Patch(Node, FixturePatch->GetStartingChannel(), false);
+				Node = FDMXFixturePatchNode::Create(DMXEditorPtr, FixturePatch);
 			}
+			check(Node.IsValid());
+
+			Patch(Node, FixturePatch->GetStartingChannel(), false);
 		}
 
 		// Update the channel connectors' Universe ID
@@ -153,6 +177,8 @@ void SDMXPatchedUniverse::SetUniverseID(int32 NewUniverseID)
 		{
 			Connector->SetUniverseID(NewUniverseID);
 		}
+
+		RequestRefresh();
 	}
 }
 
@@ -182,6 +208,9 @@ void SDMXPatchedUniverse::CreateChannelConnectors()
 			SNew(SDMXChannelConnector)
 			.ChannelID(ChannelID)
 			.Value(0)
+			.OnMouseButtonDownOnChannel(this, &SDMXPatchedUniverse::HandleOnMouseButtonDownOnChannel)
+			.OnMouseButtonUpOnChannel(this, &SDMXPatchedUniverse::HandleOnMouseButtonUpOnChannel)
+			.OnDragDetectedOnChannel(this, &SDMXPatchedUniverse::HandleOnDragDetectedOnChannel)
 			.OnDragEnterChannel(this, &SDMXPatchedUniverse::HandleDragEnterChannel)
 			.OnDragLeaveChannel(this, &SDMXPatchedUniverse::HandleDragLeaveChannel)
 			.OnDropOntoChannel(this, &SDMXPatchedUniverse::HandleDropOntoChannel)
@@ -202,11 +231,6 @@ bool SDMXPatchedUniverse::Patch(const TSharedPtr<FDMXFixturePatchNode>& Node, in
 {
 	check(Node.IsValid());
 
-	if (!CanAssignNode(Node, NewStartingChannel))
-	{
-		return false;
-	}
-
 	UDMXEntityFixturePatch* FixturePatch = Node->GetFixturePatch().Get();
 	if (!FixturePatch)
 	{
@@ -221,13 +245,8 @@ bool SDMXPatchedUniverse::Patch(const TSharedPtr<FDMXFixturePatchNode>& Node, in
 		NewStartingChannel = FixturePatch->GetStartingChannel();
 	}
 
-	// Only fully valid channels are supported
-	check(NewStartingChannel > 0);
-	check(NewStartingChannel + NewChannelSpan - 1 <= DMX_UNIVERSE_SIZE);
-
 	// Unpatch from the old universe
-	TSharedPtr<SDMXPatchedUniverse> OldUniverseWidget = Node->GetUniverseWidget();
-	if (OldUniverseWidget.IsValid())
+	if (const TSharedPtr<SDMXPatchedUniverse>& OldUniverseWidget = Node->GetUniverseWidget())
 	{
 		OldUniverseWidget->Unpatch(Node);
 	}
@@ -236,53 +255,12 @@ bool SDMXPatchedUniverse::Patch(const TSharedPtr<FDMXFixturePatchNode>& Node, in
 	check(!PatchedNodes.Contains(Node));
 	PatchedNodes.Add(Node);
 
-	// Update the node and commit to the changes
-	Node->Update(SharedThis(this), NewStartingChannel, NewChannelSpan);
-	Node->CommitPatch(bCreateTransaction);
-
-	// Assign the node's widgets to the grid
-	AddNodeToGrid(Node);
+	Node->SetAddresses(SharedThis(this), NewStartingChannel, NewChannelSpan, bCreateTransaction);
+	RequestRefresh();
 
 	return true;
 }
 
-void SDMXPatchedUniverse::Unpatch(const TSharedPtr<FDMXFixturePatchNode>& Node)
-{
-	check(Node.IsValid());
-	check(Grid.IsValid());
-
-	if (PatchedNodes.Contains(Node))
-	{
-		PatchedNodes.RemoveSingle(Node);
-		RemoveNodeFromGrid(Node);
-	}
-}
-
-void SDMXPatchedUniverse::AddNodeToGrid(const TSharedPtr<FDMXFixturePatchNode>& Node)
-{
-	check(Grid.IsValid());
-	for (const TSharedPtr<SDMXFixturePatchFragment>& Fragment : Node->GetFragmentedWidgets())
-	{
-		check(Fragment.IsValid());
-
-		Grid->AddSlot(Fragment->GetColumn(), Fragment->GetRow())
-			.ColumnSpan(Fragment->GetColumnSpan())
-			.Layer(1)
-			[
-				Fragment.ToSharedRef()
-			];
-	}
-}
-
-void SDMXPatchedUniverse::RemoveNodeFromGrid(const TSharedPtr<FDMXFixturePatchNode>& Node)
-{
-	check(Node.IsValid());
-	for (const TSharedPtr<SDMXFixturePatchFragment>& Widget : Node->GetFragmentedWidgets())
-	{
-		Grid->RemoveSlot(Widget.ToSharedRef());
-	}
-}
-
 bool SDMXPatchedUniverse::CanAssignFixturePatch(TWeakObjectPtr<UDMXEntityFixturePatch> FixturePatch) const
 {
 	if (FixturePatch.IsValid())
@@ -341,15 +319,22 @@ bool SDMXPatchedUniverse::CanAssignNode(const TSharedPtr<FDMXFixturePatchNode>&
 	return CanAssignFixturePatch(FixturePatch, StartingChannel);
 }
 
-TSharedPtr<FDMXFixturePatchNode> SDMXPatchedUniverse::FindPatchNode(const TWeakObjectPtr<UDMXEntityFixturePatch>& FixturePatch) const
+TSharedPtr<FDMXFixturePatchNode> SDMXPatchedUniverse::FindPatchNode(TWeakObjectPtr<UDMXEntityFixturePatch> FixturePatch) const
 {
-	for (const TSharedPtr<FDMXFixturePatchNode>& Node : PatchedNodes)
-	{
-		if (Node->GetFixturePatch() == FixturePatch)
+	const TSharedPtr<FDMXFixturePatchNode>* NodePtr = PatchedNodes.FindByPredicate([&FixturePatch](const TSharedPtr<FDMXFixturePatchNode>& Node)
 		{
-			return Node;
-		}
+			if (Node->GetFixturePatch() == FixturePatch)
+			{
+				return true;
+			}
+			return false;
+		});
+
+	if (NodePtr)
+	{
+		return *NodePtr;
 	}
+
 	return nullptr;
 }
 
@@ -377,17 +362,162 @@ TSharedPtr<FDMXFixturePatchNode> SDMXPatchedUniverse::FindPatchNodeOfType(UDMXEn
 	return nullptr;
 }
 
-void SDMXPatchedUniverse::HandleDragEnterChannel(int32 ChannelID, const FDragDropEvent& DragDropEvent)
+void SDMXPatchedUniverse::Unpatch(const TSharedPtr<FDMXFixturePatchNode>& Node)
+{
+	check(Node.IsValid());
+	check(Grid.IsValid());
+
+	if (PatchedNodes.Contains(Node))
+	{
+		PatchedNodes.RemoveSingle(Node);
+
+		for (const TSharedPtr<SDMXFixturePatchFragment>& Widget : Node->GetGeneratedWidgets())
+		{
+			if (Widget.IsValid())
+			{
+				Grid->RemoveSlot(Widget.ToSharedRef());
+			}
+		}
+		RefreshInternal();
+	}
+}
+
+void SDMXPatchedUniverse::OnFixturePatchChanged(const UDMXEntityFixturePatch* FixturePatch)
+{
+	// Keep track of Universe Changes
+	const TSharedPtr<FDMXFixturePatchNode>* NodePtr = PatchedNodes.FindByPredicate([FixturePatch](const TSharedPtr<FDMXFixturePatchNode>& Node)
+		{
+			return Node->GetFixturePatch() == FixturePatch;
+		});
+
+	if (NodePtr && FixturePatch && UniverseID != FixturePatch->GetUniverseID())
+	{
+		Unpatch(*NodePtr);
+	}
+	else if (!NodePtr && UniverseID == FixturePatch->GetUniverseID())
+	{
+		SharedData->SelectUniverse(UniverseID);
+		SetUniverseID(UniverseID);
+	}
+
+	RequestRefresh();
+}
+
+void SDMXPatchedUniverse::RefreshInternal()
+{
+	RequestRefreshTimerHandle.Invalidate();
+	UpdateZOrderOfNodes();
+
+	// Remove old widgets from the grid
+	for (const TSharedPtr<SDMXFixturePatchFragment>& OldWidget : FixturePatchWidgets)
+	{
+		Grid->RemoveSlot(OldWidget.ToSharedRef());
+	}
+	FixturePatchWidgets.Reset();
+
+	// Create Groups of Nodes with same Universe/Address
+	TMap<int32, TArray<TSharedPtr<FDMXFixturePatchNode>>> AddressToNodeGroupMap;
+	for (const TSharedPtr<FDMXFixturePatchNode>& Node : PatchedNodes)
+	{
+		const int32 Address = Node->GetStartingChannel();
+		AddressToNodeGroupMap.FindOrAdd(Address).Add(Node);
+	}
+
+	// Generate new widgets
+	for (const TSharedPtr<FDMXFixturePatchNode>& Node : PatchedNodes)
+	{
+		if (const TArray<TSharedPtr<FDMXFixturePatchNode>>* const NodeGroup = AddressToNodeGroupMap.Find(Node->GetStartingChannel()))
+		{
+			FixturePatchWidgets.Append(Node->GenerateWidgets(*NodeGroup));
+		}
+	}
+
+	// Add widgets to the grid
+	for (const TSharedPtr<SDMXFixturePatchFragment>& NewWidget : FixturePatchWidgets)
+	{
+		Grid->AddSlot(NewWidget->GetColumn(), NewWidget->GetRow())
+			.ColumnSpan(NewWidget->GetColumnSpan())
+			.Layer(1)
+			[
+				NewWidget.ToSharedRef()
+			];
+	}
+}
+
+FReply SDMXPatchedUniverse::HandleOnMouseButtonDownOnChannel(uint32 ChannelID, const FPointerEvent& PointerEvent)
+{
+	if (PointerEvent.GetEffectingButton() == EKeys::LeftMouseButton)
+	{
+		if (UDMXEntityFixturePatch* Patch = GetTopmostFixturePatchOnChannelID(ChannelID))
+		{
+			if (FSlateApplication::Get().GetModifierKeys().IsControlDown())
+			{
+				SharedData->AddFixturePatchToSelection(Patch);
+			}
+			else
+			{
+				SharedData->SelectFixturePatch(Patch);
+			}
+
+			// SelectFixturePatch calls SetKeyboardFocus at some point; the call to SetKeyboardFocus is not enough for Slate: 
+			// We must also return the widget using FReply.
+			FReply Reply = FReply::Handled();
+			TSharedPtr<SWidget> WidgetThatWasFocusedBy_SelectFixturePatch = FSlateApplication::Get().GetKeyboardFocusedWidget();
+			if(WidgetThatWasFocusedBy_SelectFixturePatch.IsValid())
+			{
+				Reply.SetUserFocus(WidgetThatWasFocusedBy_SelectFixturePatch.ToSharedRef());
+			}
+			if (ensureMsgf(ChannelConnectors.IsValidIndex(ChannelID - 1), TEXT("Trying to drag Fixture Patch, but the dragged channel is not valid.")))
+			{
+				return Reply.DetectDrag(ChannelConnectors[ChannelID - 1].ToSharedRef(), EKeys::LeftMouseButton);
+			}
+		}
+	}
+
+	return FReply::Unhandled();
+}
+
+FReply SDMXPatchedUniverse::HandleOnMouseButtonUpOnChannel(uint32 ChannelID, const FPointerEvent& PointerEvent)
+{
+	if (PointerEvent.GetEffectingButton() == EKeys::RightMouseButton)
+	{
+		const FVector2D& SummonLocation = PointerEvent.GetScreenSpacePosition();
+
+		TSharedRef<SWidget> MenuContent = CreateContextMenu(ChannelID);
+		FWidgetPath WidgetPath = PointerEvent.GetEventPath() ? *PointerEvent.GetEventPath() : FWidgetPath();
+		FSlateApplication::Get().PushMenu(AsShared(), WidgetPath, MenuContent, SummonLocation, FPopupTransitionEffect(FPopupTransitionEffect::ContextMenu));
+
+		return FReply::Handled();
+	}
+
+	return FReply::Unhandled();
+}
+
+FReply SDMXPatchedUniverse::HandleOnDragDetectedOnChannel(uint32 ChannelID, const FPointerEvent& PointerEvent)
+{
+	if (UDMXEntityFixturePatch* Patch = GetTopmostFixturePatchOnChannelID(ChannelID))
+	{
+		int32 ChannelOffset = ChannelID - Patch->GetStartingChannel();
+
+		TSharedRef<FDMXEntityFixturePatchDragDropOperation> DragDropOp = MakeShared<FDMXEntityFixturePatchDragDropOperation>(Patch->GetParentLibrary(), TArray<TWeakObjectPtr<UDMXEntity>>{ Patch }, ChannelOffset);
+
+		return FReply::Handled().BeginDragDrop(DragDropOp);
+	}
+
+	return FReply::Unhandled();
+}
+
+void SDMXPatchedUniverse::HandleDragEnterChannel(uint32 ChannelID, const FDragDropEvent& DragDropEvent)
 {
 	OnDragEnterChannel.ExecuteIfBound(UniverseID, ChannelID, DragDropEvent);
 }
 
-void SDMXPatchedUniverse::HandleDragLeaveChannel(int32 ChannelID, const FDragDropEvent& DragDropEvent)
+void SDMXPatchedUniverse::HandleDragLeaveChannel(uint32 ChannelID, const FDragDropEvent& DragDropEvent)
 {
 	OnDragLeaveChannel.ExecuteIfBound(UniverseID, ChannelID, DragDropEvent);
 }
 
-FReply SDMXPatchedUniverse::HandleDropOntoChannel(int32 ChannelID, const FDragDropEvent& DragDropEvent)
+FReply SDMXPatchedUniverse::HandleDropOntoChannel(uint32 ChannelID, const FDragDropEvent& DragDropEvent)
 {
 	return OnDropOntoChannel.Execute(UniverseID, ChannelID, DragDropEvent);
 }
@@ -416,6 +546,51 @@ FText SDMXPatchedUniverse::GetHeaderText() const
 	return FText::GetEmpty();
 }
 
+void SDMXPatchedUniverse::UpdateZOrderOfNodes()
+{
+	if (PatchedNodes.Num() == 0)
+	{
+		return;
+	}
+
+	TArray<TSharedPtr<FDMXFixturePatchNode>> Nodes = PatchedNodes;
+	// Sort by ZOrder descending
+	Nodes.StableSort([](const TSharedPtr<FDMXFixturePatchNode>& NodeA, const TSharedPtr<FDMXFixturePatchNode>& NodeB)
+		{
+			if (NodeA->GetStartingChannel() > NodeB->GetStartingChannel())
+			{
+				return true;
+			}
+			else if (NodeA->GetStartingChannel() == NodeB->GetStartingChannel())
+			{
+				const bool bNodeASelected = NodeA->IsSelected();
+				if (bNodeASelected)
+				{
+					return
+						!NodeB->IsSelected() ||
+						NodeA->GetZOrder() > NodeB->GetZOrder();
+				}
+				else
+				{
+					return
+						!NodeB->IsSelected() &&
+						NodeA->GetZOrder() > NodeB->GetZOrder();
+				}
+			}
+			return false;
+		});
+
+	int32 ZOrder = Nodes.Num();
+	for (const TSharedPtr<FDMXFixturePatchNode>& Node : Nodes)
+	{
+		if (Node.IsValid())
+		{
+			Node->SetZOrder(ZOrder);
+			--ZOrder;
+		}
+	}
+}
+
 EVisibility SDMXPatchedUniverse::GetPatchedUniverseReachabilityBannerVisibility() const
 {
 	if (PatchedUniverseReachability == EDMXPatchedUniverseReachability::Reachable)
@@ -457,7 +632,7 @@ void SDMXPatchedUniverse::UpdatePatchedUniverseReachability()
 		{
 			PatchedUniverseReachability = EDMXPatchedUniverseReachability::UnreachableForInputAndOutputPorts;
 		}
-		else if(!bReachableForAnyOutput)
+		else if (!bReachableForAnyOutput)
 		{
 			PatchedUniverseReachability = EDMXPatchedUniverseReachability::UnreachableForOutputPorts;
 		}
@@ -468,6 +643,84 @@ void SDMXPatchedUniverse::UpdatePatchedUniverseReachability()
 	}
 }
 
+void SDMXPatchedUniverse::OnSelectionChanged()
+{
+	const TArray<TWeakObjectPtr<UDMXEntityFixturePatch>> FixturePatches = SharedData->GetSelectedFixturePatches();
+	const bool bSelectionChanged = [&FixturePatches, this]()
+		{
+			for (const TSharedPtr<FDMXFixturePatchNode>& Node : PatchedNodes)
+			{
+				if (Node->IsSelected() ^ FixturePatches.Contains(Node->GetFixturePatch()))
+				{
+					return true;
+					break;
+				}
+			}
+			return false;
+		}();
+	if (bSelectionChanged)
+	{
+		RequestRefresh();
+	}
+}
+
+UDMXEntityFixturePatch* SDMXPatchedUniverse::GetTopmostFixturePatchOnChannelID(uint32 ChannelID) const
+{
+	TArray<TSharedPtr<FDMXFixturePatchNode>> Nodes(PatchedNodes);
+	Nodes.RemoveAll([ChannelID](const TSharedPtr<FDMXFixturePatchNode>& Node)
+			{
+				if (UDMXEntityFixturePatch* FixturePatch = Node->GetFixturePatch().Get())
+				{
+					const uint32 PatchStartingChannel = FixturePatch->GetStartingChannel();
+					const TRange<uint32> FixturePatchRange(PatchStartingChannel, PatchStartingChannel + FixturePatch->GetChannelSpan());
+
+					return !FixturePatchRange.Contains(ChannelID);
+				}
+				return true;
+			});
+
+	Nodes.Sort([](const TSharedPtr<FDMXFixturePatchNode>& NodeA, const TSharedPtr<FDMXFixturePatchNode>& NodeB)
+			{
+				return NodeA->GetZOrder() > NodeB->GetZOrder();
+			});
+
+	if (Nodes.IsEmpty())
+	{
+		return nullptr;
+	}
+	else
+	{
+		return Nodes[0]->GetFixturePatch().Get();
+	}
+}
+
+TArray<UDMXEntityFixturePatch*> SDMXPatchedUniverse::GetFixturePatchesOnChannelID(uint32 ChannelID) const
+{
+	TArray<TSharedPtr<FDMXFixturePatchNode>> Nodes(PatchedNodes);
+	Nodes.RemoveAll([ChannelID](const TSharedPtr<FDMXFixturePatchNode>& Node)
+		{
+			if (UDMXEntityFixturePatch* FixturePatch = Node->GetFixturePatch().Get())
+			{
+				const uint32 PatchStartingChannel = FixturePatch->GetStartingChannel();
+				const TRange<uint32> FixturePatchRange(PatchStartingChannel, PatchStartingChannel + FixturePatch->GetChannelSpan());
+
+				return !FixturePatchRange.Contains(ChannelID);
+			}
+			return true;
+		});
+
+	TArray<UDMXEntityFixturePatch*> Result;
+	for (const TSharedPtr<FDMXFixturePatchNode>& Node : Nodes)
+	{
+		if (UDMXEntityFixturePatch* FixturePatch = Node->GetFixturePatch().Get())
+		{
+			Result.Add(FixturePatch);
+		}
+	}
+
+	return Result;
+}
+
 UDMXLibrary* SDMXPatchedUniverse::GetDMXLibrary() const
 {
 	if (TSharedPtr<FDMXEditor> DMXEditor = DMXEditorPtr.Pin())
@@ -477,4 +730,97 @@ UDMXLibrary* SDMXPatchedUniverse::GetDMXLibrary() const
 	return nullptr;
 }
 
+void SDMXPatchedUniverse::RegisterCommands()
+{
+	if (CommandList.IsValid())
+	{
+		return;
+	}
+
+	// listen to common editor shortcuts for copy/paste etc
+	CommandList = MakeShared<FUICommandList>();
+}
+
+TSharedRef<SWidget> SDMXPatchedUniverse::CreateContextMenu(int32 ChannelID) const
+{
+	UDMXLibrary* DMXLibrary = GetDMXLibrary();
+	if (!DMXLibrary)
+	{
+		return SNullWidget::NullWidget;
+	}
+
+	TArray<UDMXEntityFixturePatch*> FixturePatchesOnChannel = GetFixturePatchesOnChannelID(ChannelID);
+	if (FixturePatchesOnChannel.Num() == 0)
+	{
+		return SNullWidget::NullWidget;
+	}
+
+	constexpr bool bCloseWindowAfterMenuSelection = true;
+	FMenuBuilder MenuBuilder(bCloseWindowAfterMenuSelection, CommandList);
+
+	// Add a context menu
+	if (FixturePatchesOnChannel.Num() > 1)
+	{
+		// Common interactions
+		MenuBuilder.BeginSection("FixturePatchInteractions");
+		{
+			// Select All
+			const FUIAction Action(
+				FExecuteAction::CreateLambda([FixturePatchesOnChannel, this]()
+					{
+						if (SharedData.IsValid())
+						{
+							TArray<TWeakObjectPtr<UDMXEntityFixturePatch>> WeakFixturePatchesOnChannel;
+							for (UDMXEntityFixturePatch* FixturePatch : FixturePatchesOnChannel)
+							{
+								WeakFixturePatchesOnChannel.Add(FixturePatch);
+							}
+							SharedData->SelectFixturePatches(WeakFixturePatchesOnChannel);
+						}
+					}));
+
+			const FText SelectAllText = LOCTEXT("SelectAllFixturePatches", "Select all");
+			const TSharedRef<SWidget> Widget =
+				SNew(STextBlock)
+				.Text(SelectAllText);
+
+			MenuBuilder.AddMenuEntry(Action, Widget);
+		}
+		MenuBuilder.EndSection();
+
+		// Single fixture patch selection
+		MenuBuilder.BeginSection("BringFixturePatchForward", LOCTEXT("BringFixturePatchForwardLabel", "Bring Forward"));
+		{
+			const TArray<UDMXEntityFixturePatch*> DMXFixturePatchesInLibrary = DMXLibrary->GetEntitiesTypeCast<UDMXEntityFixturePatch>();
+			FixturePatchesOnChannel.Sort([&DMXFixturePatchesInLibrary](const UDMXEntityFixturePatch& FixturePatchA, const UDMXEntityFixturePatch& FixturePatchB)
+				{
+					return DMXFixturePatchesInLibrary.IndexOfByKey(&FixturePatchA) < DMXFixturePatchesInLibrary.IndexOfByKey(&FixturePatchB);
+				});
+
+			for (TWeakObjectPtr<UDMXEntityFixturePatch> WeakFixturePatch : FixturePatchesOnChannel)
+			{
+				const FUIAction Action(
+					FExecuteAction::CreateLambda([this, WeakFixturePatch]()
+						{
+							if (SharedData.IsValid() && WeakFixturePatch.IsValid())
+							{
+								SharedData->SelectFixturePatch(WeakFixturePatch.Get());
+							}
+						})
+				);
+
+				const FText FixturePatchName = FText::FromString(WeakFixturePatch->Name);
+				const TSharedRef<SWidget> Widget =
+					SNew(STextBlock)
+					.Text(FixturePatchName);
+
+				MenuBuilder.AddMenuEntry(Action, Widget);
+			}
+		}
+		MenuBuilder.EndSection();
+	}
+
+	return MenuBuilder.MakeWidget();
+}
+
 #undef LOCTEXT_NAMESPACE
diff --git a/Engine/Plugins/VirtualProduction/DMX/DMXEngine/Source/DMXEditor/Private/Widgets/FixturePatch/SDMXPatchedUniverse.h b/Engine/Plugins/VirtualProduction/DMX/DMXEngine/Source/DMXEditor/Private/Widgets/FixturePatch/SDMXPatchedUniverse.h
--- a/Engine/Plugins/VirtualProduction/DMX/DMXEngine/Source/DMXEditor/Private/Widgets/FixturePatch/SDMXPatchedUniverse.h
+++ b/Engine/Plugins/VirtualProduction/DMX/DMXEngine/Source/DMXEditor/Private/Widgets/FixturePatch/SDMXPatchedUniverse.h
@@ -7,6 +7,7 @@
 
 #include "CoreMinimal.h"
 #include "EditorUndoClient.h"
+#include "Engine/EngineTypes.h"
 #include "UObject/WeakObjectPtrTemplates.h"
 #include "Widgets/DeclarativeSyntaxSupport.h"
 #include "Widgets/SCompoundWidget.h"
@@ -15,9 +16,13 @@ class FDMXEditor;
 class FDMXFixturePatchNode;
 class FDMXFixturePatchSharedData;
 class SDMXChannelConnector;
+class SDMXFixturePatchFragment;
 class UDMXLibrary;
 class UDMXEntityFixturePatch;
 
+class FDragDropEvent;
+struct FTimerHandle;
+class FUICommandList;
 class SBorder;
 class SGridPanel;
 
@@ -68,6 +73,9 @@ public:
 
 	/** Constructs this widget */
 	void Construct(const FArguments& InArgs);
+
+	/** Requests to refresh the Widget on the next tick */
+	void RequestRefresh();
 	
 	/** 
 	 * Shows specified Universe ID in the widget. 
@@ -83,23 +91,9 @@ public:
 	 */
 	bool Patch(const TSharedPtr<FDMXFixturePatchNode>& Node, int32 NewStartingChannel, bool bCreateTransaction);
 
-protected:
-	/** Removes the node. Should be called when the node is Patched in another instance */
-	void Unpatch(const TSharedPtr<FDMXFixturePatchNode>& Node);
-
-	/** Adds the node to the grid */
-	void AddNodeToGrid(const TSharedPtr<FDMXFixturePatchNode>& Node);
-
-	/** Removes the node from the grid */
-	void RemoveNodeFromGrid(const TSharedPtr<FDMXFixturePatchNode>& Node);
-
-public:
 	/** If set to true, shows a universe name above the patcher universe */
 	void SetShowUniverseName(bool bShow);
 
-	/** Returns the grid of the universe */
-	TSharedPtr<SGridPanel> GetGrid() const { return Grid; }
-
 	/** Returns wether the patch can be patched to its current channels */
 	bool CanAssignFixturePatch(TWeakObjectPtr<UDMXEntityFixturePatch> TestedPatch) const;
 
@@ -110,9 +104,9 @@ public:
 	bool CanAssignNode(const TSharedPtr<FDMXFixturePatchNode>& TestedNode, int32 StartingChannel) const;
 
 	/** Returns if the node is patched in the unvierse */
-	TSharedPtr<FDMXFixturePatchNode> FindPatchNode(const TWeakObjectPtr<UDMXEntityFixturePatch>& FixturePatch) const;
+	TSharedPtr<FDMXFixturePatchNode> FindPatchNode(TWeakObjectPtr<UDMXEntityFixturePatch> FixturePatch) const;
 
-	/** Returns first node with same fixture type as specified node */	
+	/** Returns first node with same fixture type as specified node */
 	TSharedPtr<FDMXFixturePatchNode> FindPatchNodeOfType(UDMXEntityFixtureType* Type, const TSharedPtr<FDMXFixturePatchNode>& IgoredNode) const;
 
 	/** Returns the ID of the universe */
@@ -121,35 +115,58 @@ public:
 	/** Gets all nodes patched to this universe */
 	const TArray<TSharedPtr<FDMXFixturePatchNode>>& GetPatchedNodes() const { return PatchedNodes; }
 
-protected:
+private:
+	/** Removes the node. Should be called when the node is Patched in another instance */
+	void Unpatch(const TSharedPtr<FDMXFixturePatchNode>& Node);
+
+	/** Called when a Fixture Patch changed */
+	void OnFixturePatchChanged(const UDMXEntityFixturePatch* FixturePatch);
+
+	/** Refreshes the widget directly */
+	void RefreshInternal();
+
 	/** Creates a a new grid of channels */
 	void CreateChannelConnectors();
 
 	/** Returns the name of the universe displayed */
 	FText GetHeaderText() const;
 	
-protected:
-	/** Called when drag enters a channel */
-	void HandleDragEnterChannel(int32 ChannelID, const FDragDropEvent& DragDropEvent);
+	/** Handles when a mouse button was pressed on a Channel */
+	FReply HandleOnMouseButtonDownOnChannel(uint32 ChannelID, const FPointerEvent& PointerEvent);
 
-	/** Called when drag leaves a channel */
-	void HandleDragLeaveChannel(int32 ChannelID, const FDragDropEvent& DragDropEvent);
+	/** Handles when a mouse button was released on a Channel */
+	FReply HandleOnMouseButtonUpOnChannel(uint32 ChannelID, const FPointerEvent& PointerEvent);
 
-	/** Called when drag dropped onto a channel */
-	FReply HandleDropOntoChannel(int32 ChannelID, const FDragDropEvent& DragDropEvent);
+	/** Handles when a Channel was dragged */
+	FReply HandleOnDragDetectedOnChannel(uint32 ChannelID, const FPointerEvent& PointerEvent);
 
-	// Drag drop events
-	FOnDragOverChannel OnDragEnterChannel;
-	FOnDragOverChannel OnDragLeaveChannel;
-	FOnDropOntoChannel OnDropOntoChannel;
+	/** Handles when drag enters a channel */
+	void HandleDragEnterChannel(uint32 ChannelID, const FDragDropEvent& DragDropEvent);
+
+	/** Handles when drag leaves a channel */
+	void HandleDragLeaveChannel(uint32 ChannelID, const FDragDropEvent& DragDropEvent);
+
+	/** Handles when drag dropped onto a channel */
+	FReply HandleDropOntoChannel(uint32 ChannelID, const FDragDropEvent& DragDropEvent);
+
+	/** Updates the ZOrder of all Nodes */
+	void UpdateZOrderOfNodes();
 
-protected:
 	/** Returns wether the out of controllers' ranges banner should be visible */
 	EVisibility GetPatchedUniverseReachabilityBannerVisibility() const;
 
 	/** Updates bOutOfControllersRanges member */
 	void UpdatePatchedUniverseReachability();
 
+	/** Called when selection changed */
+	void OnSelectionChanged();
+
+	/** Returns the Fixture Patch that is topmost under ChannelID */
+	UDMXEntityFixturePatch* GetTopmostFixturePatchOnChannelID(uint32 ChannelID) const;
+
+	/** Returns all Fixture Patches on a Channel ID */
+	TArray<UDMXEntityFixturePatch*> GetFixturePatchesOnChannelID(uint32 ChannelID) const;
+
 	/** Returns the DMXLibrary or nullptr if not available */
 	UDMXLibrary* GetDMXLibrary() const;
 
@@ -171,9 +188,37 @@ protected:
 	/** The Channel connectors in this universe */
 	TArray<TSharedPtr<SDMXChannelConnector>> ChannelConnectors;
 
+	/** Delegate executed when drag enters a Channel */
+	FOnDragOverChannel OnDragEnterChannel;
+
+	/** Delegate executed when drag leaves a Channel */
+	FOnDragOverChannel OnDragLeaveChannel;
+
+	/** Delegate executed when a Drag Drop event was dropped onto a Channel */
+	FOnDropOntoChannel OnDropOntoChannel;
+
+	/** The Fixture Patch Widgets that are currently being displalyed */
+	TArray<TSharedPtr<SDMXFixturePatchFragment>> FixturePatchWidgets;
+
+	/** Timer handle for the Request Refresh Timer */
+	FTimerHandle RequestRefreshTimerHandle;
+
 	/** Shared data for fixture patch editors */
 	TSharedPtr<FDMXFixturePatchSharedData> SharedData;
 
 	/** The owning editor */
 	TWeakPtr<FDMXEditor> DMXEditorPtr;
+
+	private:
+	///////////////////////////////////////////////////
+	// Context menu Commands related
+
+	/** Registers commands for this widget */
+	void RegisterCommands();
+
+	/** Creates the right click context menu */
+	TSharedRef<SWidget> CreateContextMenu(int32 ChannelID) const;
+
+	/** Command list for this widget */
+	TSharedPtr<FUICommandList> CommandList;
 };
