DMX - Fix CIS  Localization warnings in SDMXMVRFixtureList.cpp

diff --git a/Engine/Plugins/VirtualProduction/DMX/DMXEngine/Source/DMXEditor/Private/Widgets/FixturePatch/SDMXMVRFixtureList.cpp b/Engine/Plugins/VirtualProduction/DMX/DMXEngine/Source/DMXEditor/Private/Widgets/FixturePatch/SDMXMVRFixtureList.cpp
--- a/Engine/Plugins/VirtualProduction/DMX/DMXEngine/Source/DMXEditor/Private/Widgets/FixturePatch/SDMXMVRFixtureList.cpp
+++ b/Engine/Plugins/VirtualProduction/DMX/DMXEngine/Source/DMXEditor/Private/Widgets/FixturePatch/SDMXMVRFixtureList.cpp
@@ -164,17 +164,17 @@ private:
 			if (Item->GetAddress() < 1 &&
 				EndingAddress > DMX_MAX_ADDRESS)
 			{
-				const FText ConflictText = FText::Format(LOCTEXT("ChannelExcessConflict", "Exceeds available DMX Address range. Staring Address is {0} but min Address is 1. Ending Address is {1} but max Address is 512."), Item->GetAddress(), EndingAddress);
+				const FText ConflictText = FText::Format(LOCTEXT("ChannelExceedsMinAndMaxChannelConflict", "Exceeds available DMX Address range. Staring Address is {0} but min Address is 1. Ending Address is {1} but max Address is 512."), Item->GetAddress(), EndingAddress);
 				ItemToConflictMap.Add(Item, ConflictText);
 			}
 			else if (Item->GetAddress() < 1)
 			{
-				const FText ConflictText = FText::Format(LOCTEXT("ChannelExcessConflict", "Exceeds available DMX Address range. Staring Address is {0} but min Address is 1."), Item->GetAddress());
+				const FText ConflictText = FText::Format(LOCTEXT("ChannelExceedsMinChannelNumberConflict", "Exceeds available DMX Address range. Staring Address is {0} but min Address is 1."), Item->GetAddress());
 				ItemToConflictMap.Add(Item, ConflictText);
 			}
 			else if (EndingAddress > DMX_MAX_ADDRESS)
 			{
-				const FText ConflictText = FText::Format(LOCTEXT("ChannelExcessConflict", "Exceeds available DMX Address range. Ending Address is {0} but max Address is 512."), EndingAddress);
+				const FText ConflictText = FText::Format(LOCTEXT("ChannelExeedsMaxChannelNumberConflict", "Exceeds available DMX Address range. Ending Address is {0} but max Address is 512."), EndingAddress);
 				ItemToConflictMap.Add(Item, ConflictText);
 			}			
 		}
