DMX - DMXProtocolEditor Module now tests for Property Editor Module to be loaded before accessing it during shutdown, DMXEditorModule now correctly unregisters all its customizations

diff --git a/Engine/Plugins/VirtualProduction/DMX/DMXEngine/Source/DMXEditor/Private/DMXEditorModule.cpp b/Engine/Plugins/VirtualProduction/DMX/DMXEngine/Source/DMXEditor/Private/DMXEditorModule.cpp
--- a/Engine/Plugins/VirtualProduction/DMX/DMXEngine/Source/DMXEditor/Private/DMXEditorModule.cpp
+++ b/Engine/Plugins/VirtualProduction/DMX/DMXEngine/Source/DMXEditor/Private/DMXEditorModule.cpp
@@ -250,23 +250,19 @@ void FDMXEditorModule::RegisterAssetTypeActions()
 
 void FDMXEditorModule::RegisterClassCustomizations()
 {
-	FPropertyEditorModule& PropertyModule = FModuleManager::GetModuleChecked<FPropertyEditorModule>("PropertyEditor");
-
 	// Details customization for the UDMXEntityFixtureType class
-	PropertyModule.RegisterCustomClassLayout(UDMXEntityFixtureType::StaticClass()->GetFName(),
+	RegisterCustomClassLayout(UDMXEntityFixtureType::StaticClass()->GetFName(),
 		FOnGetDetailCustomizationInstance::CreateStatic(&FDMXEntityFixtureTypeDetails::MakeInstance)
 	);
 
 	// Details customization for the ADMXMVRSceneActor class
-	PropertyModule.RegisterCustomClassLayout(ADMXMVRSceneActor::StaticClass()->GetFName(),
+	RegisterCustomClassLayout(ADMXMVRSceneActor::StaticClass()->GetFName(),
 		FOnGetDetailCustomizationInstance::CreateStatic(&FDMXMVRSceneActorDetails::MakeInstance)
 	);
 }
 
 void FDMXEditorModule::RegisterPropertyTypeCustomizations()
 {
-	FPropertyEditorModule& PropertyModule = FModuleManager::GetModuleChecked<FPropertyEditorModule>("PropertyEditor");
-
 	// Property type customization for the EDMXPixelMappingDistribution enum
 	RegisterCustomPropertyTypeLayout("EDMXPixelMappingDistribution", 
 		FOnGetPropertyTypeCustomizationInstance::CreateStatic(&FDMXPixelMappingDistributionCustomization::MakeInstance)
@@ -455,10 +451,10 @@ void FDMXEditorModule::UnregisterAssetTypeActions()
 {
 	if (FModuleManager::Get().IsModuleLoaded("AssetTools"))
 	{
-		IAssetTools& AssetTools = FModuleManager::GetModuleChecked<FAssetToolsModule>("AssetTools").Get();
+		IAssetTools& AssetToolsModule = FModuleManager::GetModuleChecked<FAssetToolsModule>("AssetTools").Get();
 		for (TSharedPtr<IAssetTypeActions>& AssetIt : RegisteredAssetTypeActions)
 		{
-			AssetTools.UnregisterAssetTypeActions(AssetIt.ToSharedRef());
+			AssetToolsModule.UnregisterAssetTypeActions(AssetIt.ToSharedRef());
 		}
 	}
 
@@ -480,13 +476,12 @@ void FDMXEditorModule::RegisterCustomClassLayout(FName ClassName, FOnGetDetailCu
 
 void FDMXEditorModule::UnregisterCustomClassLayouts()
 {
-	if (FModuleManager::Get().IsModuleLoaded("PropertyEditor"))
+	FPropertyEditorModule* PropertyModule = FModuleManager::GetModulePtr<FPropertyEditorModule>("PropertyEditor");
+	if (PropertyModule)
 	{
-		FPropertyEditorModule& PropertyModule = FModuleManager::GetModuleChecked<FPropertyEditorModule>("PropertyEditor");
-
 		for (const FName& ClassName : RegisteredClassNames)
 		{
-			PropertyModule.UnregisterCustomClassLayout(ClassName);
+			PropertyModule->UnregisterCustomClassLayout(ClassName);
 		}
 	}
 
@@ -508,13 +503,12 @@ void FDMXEditorModule::RegisterCustomPropertyTypeLayout(FName PropertyTypeName,
 
 void FDMXEditorModule::UnregisterCustomPropertyTypeLayouts()
 {
-	if (FModuleManager::Get().IsModuleLoaded("PropertyEditor"))
+	FPropertyEditorModule* PropertyModule = FModuleManager::GetModulePtr<FPropertyEditorModule>("PropertyEditor");
+	if (PropertyModule)
 	{
-		FPropertyEditorModule& PropertyModule = FModuleManager::GetModuleChecked<FPropertyEditorModule>("PropertyEditor");
-
 		for (const FName& PropertyTypeName : RegisteredPropertyTypes)
 		{
-			PropertyModule.UnregisterCustomPropertyTypeLayout(PropertyTypeName);
+			PropertyModule->UnregisterCustomPropertyTypeLayout(PropertyTypeName);
 		}
 	}
 
@@ -529,13 +523,12 @@ void  FDMXEditorModule::RegisterCustomSequencerTrackType(const FOnCreateTrackEdi
 
 void FDMXEditorModule::UnregisterCustomSequencerTrackTypes()
 {
-	if (FModuleManager::Get().IsModuleLoaded("Sequencer"))
+	ISequencerModule* SequencerModule = FModuleManager::GetModulePtr<ISequencerModule>("Sequencer");
+	if (SequencerModule)
 	{
-		ISequencerModule& SequencerModule = FModuleManager::Get().LoadModuleChecked<ISequencerModule>("Sequencer");
-
 		for (const FDelegateHandle& TrackCreateHandle : RegisteredSequencerTrackHandles)
 		{
-			SequencerModule.UnRegisterTrackEditor(TrackCreateHandle);
+			SequencerModule->UnRegisterTrackEditor(TrackCreateHandle);
 		}
 	}
 
diff --git a/Engine/Plugins/VirtualProduction/DMX/DMXProtocol/Source/DMXProtocolEditor/Private/DMXProtocolEditorModule.cpp b/Engine/Plugins/VirtualProduction/DMX/DMXProtocol/Source/DMXProtocolEditor/Private/DMXProtocolEditorModule.cpp
--- a/Engine/Plugins/VirtualProduction/DMX/DMXProtocol/Source/DMXProtocolEditor/Private/DMXProtocolEditorModule.cpp
+++ b/Engine/Plugins/VirtualProduction/DMX/DMXProtocol/Source/DMXProtocolEditor/Private/DMXProtocolEditorModule.cpp
@@ -53,13 +53,17 @@ void FDMXProtocolEditorModule::RegisterDetailsCustomizations()
 
 void FDMXProtocolEditorModule::UnregisterDetailsCustomizations()
 {
-	FPropertyEditorModule& PropertyModule = FModuleManager::GetModuleChecked<FPropertyEditorModule>("PropertyEditor");
+	FPropertyEditorModule* PropertyModule = FModuleManager::GetModulePtr<FPropertyEditorModule>("PropertyEditor");
+	if (PropertyModule)
+	{
+		PropertyModule->UnregisterCustomPropertyTypeLayout(FDMXInputPortConfig::StaticStruct()->GetFName());
+		PropertyModule->UnregisterCustomPropertyTypeLayout(FDMXOutputPortConfig::StaticStruct()->GetFName());
 
-	PropertyModule.UnregisterCustomPropertyTypeLayout(FDMXInputPortConfig::StaticStruct()->GetFName());
-	PropertyModule.UnregisterCustomPropertyTypeLayout(FDMXOutputPortConfig::StaticStruct()->GetFName());
+		PropertyModule->UnregisterCustomPropertyTypeLayout(FDMXOutputPortDestinationAddress::StaticStruct()->GetFName());
 
-	PropertyModule.UnregisterCustomPropertyTypeLayout(FDMXInputPortReference::StaticStruct()->GetFName());
-	PropertyModule.UnregisterCustomPropertyTypeLayout(FDMXOutputPortReference::StaticStruct()->GetFName());
+		PropertyModule->UnregisterCustomPropertyTypeLayout(FDMXInputPortReference::StaticStruct()->GetFName());
+		PropertyModule->UnregisterCustomPropertyTypeLayout(FDMXOutputPortReference::StaticStruct()->GetFName());
+	}
 }
 
 IMPLEMENT_MODULE(FDMXProtocolEditorModule, DMXProtocolEditor);
