DMX - Rework auto-assign feature to be a feature of the Fixture Patch view instead of a property of DMXEntityFixturePatch

diff --git a/Engine/Plugins/VirtualProduction/DMX/DMXEngine/Source/DMXEditor/Private/Customizations/DMXEntityFixturePatchDetails.cpp b/Engine/Plugins/VirtualProduction/DMX/DMXEngine/Source/DMXEditor/Private/Customizations/DMXEntityFixturePatchDetails.cpp
--- a/Engine/Plugins/VirtualProduction/DMX/DMXEngine/Source/DMXEditor/Private/Customizations/DMXEntityFixturePatchDetails.cpp
+++ b/Engine/Plugins/VirtualProduction/DMX/DMXEngine/Source/DMXEditor/Private/Customizations/DMXEntityFixturePatchDetails.cpp
@@ -28,12 +28,10 @@ void FDMXEntityFixturePatchDetails::CustomizeDetails(IDetailLayoutBuilder& Detai
 {
 	PropertyUtilities = DetailBuilder.GetPropertyUtilities();
 
-	AutoAssignAddressHandle = DetailBuilder.GetProperty(UDMXEntityFixturePatch::GetAutoAssignAddressPropertyNameChecked());
 	ParentFixtureTypeHandle = DetailBuilder.GetProperty(UDMXEntityFixturePatch::GetParentFixtureTypeTemplatePropertyNameChecked());
 	ActiveModeHandle = DetailBuilder.GetProperty(UDMXEntityFixturePatch::GetActiveModePropertyNameChecked());
 
 	// Handle Property changes
-	AutoAssignAddressHandle->SetOnPropertyValueChanged(FSimpleDelegate::CreateSP(this, &FDMXEntityFixturePatchDetails::OnAutoAssignAddressChanged));
 	UniverseIDHandle = DetailBuilder.GetProperty(UDMXEntityFixturePatch::GetUniverseIDPropertyNameChecked());
 	UniverseIDHandle->SetOnPropertyValueChanged(FSimpleDelegate::CreateSP(this, &FDMXEntityFixturePatchDetails::OnUniverseIDChanged));
 
@@ -166,28 +164,6 @@ void FDMXEntityFixturePatchDetails::OnActiveModeChanged(const TSharedPtr<uint32>
 	}
 }
 
-void FDMXEntityFixturePatchDetails::OnAutoAssignAddressChanged()
-{
-	bool bAutoAssignAddress;
-	if (ensure(AutoAssignAddressHandle->GetValue(bAutoAssignAddress) == FPropertyAccess::Success))
-	{
-		if (bAutoAssignAddress)
-		{
-			TArray<UObject*> OuterObjects;
-			AutoAssignAddressHandle->GetOuterObjects(OuterObjects);
-
-			TArray<UDMXEntityFixturePatch*> FixturePatches;
-			for (UObject* Object : OuterObjects)
-			{
-				UDMXEntityFixturePatch* Patch = CastChecked<UDMXEntityFixturePatch>(Object);
-				FixturePatches.Add(Patch);
-			}
-
-			FDMXEditorUtils::AutoAssignedAddresses(FixturePatches);
-		}
-	}
-}
-
 void FDMXEntityFixturePatchDetails::GenerateActiveModesSource()
 {
 	ActiveModesSource.Reset();
@@ -216,7 +192,6 @@ TWeakObjectPtr<UDMXEntityFixtureType> FDMXEntityFixturePatchDetails::GetParentFi
 	return nullptr;
 }
 
-
 bool FDMXEntityFixturePatchDetails::IsParentFixtureTypeMultipleValues() const
 {
 	UObject* Object;
diff --git a/Engine/Plugins/VirtualProduction/DMX/DMXEngine/Source/DMXEditor/Private/Customizations/DMXEntityFixturePatchDetails.h b/Engine/Plugins/VirtualProduction/DMX/DMXEngine/Source/DMXEditor/Private/Customizations/DMXEntityFixturePatchDetails.h
--- a/Engine/Plugins/VirtualProduction/DMX/DMXEngine/Source/DMXEditor/Private/Customizations/DMXEntityFixturePatchDetails.h
+++ b/Engine/Plugins/VirtualProduction/DMX/DMXEngine/Source/DMXEditor/Private/Customizations/DMXEntityFixturePatchDetails.h
@@ -49,9 +49,6 @@ private:
 	/** Called when the ActiveMode property changed */
 	void OnActiveModeChanged(const TSharedPtr<uint32> InSelectedMode, ESelectInfo::Type SelectInfo);
 
-	/** Called when the bAutoAssignAddress property changed */
-	void OnAutoAssignAddressChanged();
-
 	/** Generates the Array of available mode in the ActiveModeSource member array */
 	void GenerateActiveModesSource();
 
@@ -79,9 +76,6 @@ private:
 	/** Handle for the ActiveMode property */
 	TSharedPtr<IPropertyHandle> ActiveModeHandle;
 
-	/** Handle for the bAutoAssignAddress property */
-	TSharedPtr<IPropertyHandle> AutoAssignAddressHandle;
-
 	/** ComboBox that displays the ActiveMode property */
 	TSharedPtr<SComboBox<TSharedPtr<uint32>>> ActiveModeComboBox;
 
diff --git a/Engine/Plugins/VirtualProduction/DMX/DMXEngine/Source/DMXEditor/Private/DMXEditorUtils.cpp b/Engine/Plugins/VirtualProduction/DMX/DMXEngine/Source/DMXEditor/Private/DMXEditorUtils.cpp
--- a/Engine/Plugins/VirtualProduction/DMX/DMXEngine/Source/DMXEditor/Private/DMXEditorUtils.cpp
+++ b/Engine/Plugins/VirtualProduction/DMX/DMXEngine/Source/DMXEditor/Private/DMXEditorUtils.cpp
@@ -566,13 +566,14 @@ FText FDMXEditorUtils::GetEntityTypeNameText(TSubclassOf<UDMXEntity> EntityClass
 
 bool FDMXEditorUtils::TryAutoAssignToUniverses(UDMXEntityFixturePatch* Patch, const TSet<int32>& AllowedUniverses)
 {
+	PRAGMA_DISABLE_DEPRECATION_WARNINGS
 	check(Patch->IsAutoAssignAddress());
-	
+
 	Patch->Modify();
 	const int32 UniverseToRestore = Patch->GetUniverseID();
 	const int32 AutoAddressToRestore = Patch->GetAutoStartingAddress();
-	
-	for(auto UniverseIt = AllowedUniverses.CreateConstIterator(); UniverseIt; ++UniverseIt)
+
+	for (auto UniverseIt = AllowedUniverses.CreateConstIterator(); UniverseIt; ++UniverseIt)
 	{
 		// Don't auto assign to a universe smaller than the initial one
 		if (Patch->GetUniverseID() > *UniverseIt)
@@ -582,9 +583,9 @@ bool FDMXEditorUtils::TryAutoAssignToUniverses(UDMXEntityFixturePatch* Patch, co
 
 		Patch->SetUniverseID(*UniverseIt);
 		const FUnassignedPatchesArray UnassignedPatches = AutoAssignedAddresses({ Patch }, 1, false);
-		
+
 		const bool bWasPatchAssignedToUniverse = UnassignedPatches.Num() == 0;
-		if(bWasPatchAssignedToUniverse)
+		if (bWasPatchAssignedToUniverse)
 		{
 			return true;
 		}
@@ -592,11 +593,14 @@ bool FDMXEditorUtils::TryAutoAssignToUniverses(UDMXEntityFixturePatch* Patch, co
 
 	Patch->SetUniverseID(UniverseToRestore);
 	Patch->SetAutoStartingAddress(AutoAddressToRestore);
+
+	PRAGMA_ENABLE_DEPRECATION_WARNINGS
 	return false;
 }
 
 void FDMXEditorUtils::AutoAssignedAddresses(UDMXEntityFixtureType* ChangedParentFixtureType)
 {
+	PRAGMA_DISABLE_DEPRECATION_WARNINGS
 	if (ChangedParentFixtureType)
 	{
 		UDMXLibrary* Library = ChangedParentFixtureType->GetParentLibrary();
@@ -613,6 +617,7 @@ void FDMXEditorUtils::AutoAssignedAddresses(UDMXEntityFixtureType* ChangedParent
 
 		AutoAssignedAddresses(FixturePatches);
 	}
+	PRAGMA_DISABLE_DEPRECATION_WARNINGS
 }
 
 FDMXEditorUtils::FUnassignedPatchesArray FDMXEditorUtils::AutoAssignedAddresses(
@@ -620,33 +625,34 @@ FDMXEditorUtils::FUnassignedPatchesArray FDMXEditorUtils::AutoAssignedAddresses(
 	int32 MinimumAddress,
 	bool bCanChangePatchUniverses)
 {
+	PRAGMA_DISABLE_DEPRECATION_WARNINGS
 	struct Local
 	{
 		static void AssignAddressesInUniverse(
 			TArray<UDMXEntityFixturePatch*>& SortedPatchesWithSetAddress,
-			TArray<UDMXEntityFixturePatch*>& PatchesToAssign, 
+			TArray<UDMXEntityFixturePatch*>& PatchesToAssign,
 			int32 MinimumAddress,
 			const int32 UniverseToAssignTo,
 			TArray<UDMXEntityFixturePatch*>& UnassignedPatches)
 		{
-			if(PatchesToAssign.Num() == 0)
+			if (PatchesToAssign.Num() == 0)
 			{
 				return;
 			}
-			
+
 			int32 IndexOfFirstPatchWithMinAddress =
 				SortedPatchesWithSetAddress.IndexOfByPredicate([MinimumAddress](UDMXEntityFixturePatch* Other)
-				{
-					const int32 ChannelEnd = Other->GetStartingChannel() + Other->GetChannelSpan();
-					return ChannelEnd >= MinimumAddress;
-				});
-			
-			for(UDMXEntityFixturePatch* ToAssign : PatchesToAssign)
+					{
+						const int32 ChannelEnd = Other->GetStartingChannel() + Other->GetChannelSpan();
+						return ChannelEnd >= MinimumAddress;
+					});
+
+			for (UDMXEntityFixturePatch* ToAssign : PatchesToAssign)
 			{
 				const bool bIsUniverseEmpty = IndexOfFirstPatchWithMinAddress == INDEX_NONE;
 				if (bIsUniverseEmpty)
 				{
-					if(HasEnoughSpaceToUniverseEnd(ToAssign, MinimumAddress))
+					if (HasEnoughSpaceToUniverseEnd(ToAssign, MinimumAddress))
 					{
 						IndexOfFirstPatchWithMinAddress = 0;
 						AssignPatchTo(ToAssign, MinimumAddress, UniverseToAssignTo);
@@ -660,10 +666,10 @@ FDMXEditorUtils::FUnassignedPatchesArray FDMXEditorUtils::AutoAssignedAddresses(
 				}
 
 				const bool bFoundGap = FillIntoFirstGap(ToAssign, UniverseToAssignTo, MinimumAddress, IndexOfFirstPatchWithMinAddress, SortedPatchesWithSetAddress);
-				if(!bFoundGap)
+				if (!bFoundGap)
 				{
 					UDMXEntityFixturePatch* LastPatchInUniverse = SortedPatchesWithSetAddress[SortedPatchesWithSetAddress.Num() - 1];
-					if(HasEnoughSpaceToUniverseEnd(ToAssign, GetEndAddressOf(LastPatchInUniverse) + 1))
+					if (HasEnoughSpaceToUniverseEnd(ToAssign, GetEndAddressOf(LastPatchInUniverse) + 1))
 					{
 						AssignPatchTo(ToAssign, GetEndAddressOf(LastPatchInUniverse) + 1, UniverseToAssignTo);
 						SortedPatchesWithSetAddress.Add(ToAssign);
@@ -675,13 +681,13 @@ FDMXEditorUtils::FUnassignedPatchesArray FDMXEditorUtils::AutoAssignedAddresses(
 				}
 			}
 		}
-		
+
 		static bool HasEnoughSpaceToUniverseEnd(UDMXEntityFixturePatch* Patch, int32 AtAddress)
 		{
 			// + 1 is needed otherwise we're off by one, e.g. when AtAddress=DMX_UNIVERSE_SIZE and ChannelSpan=1
 			return (DMX_UNIVERSE_SIZE + 1) - AtAddress >= Patch->GetChannelSpan();
 		}
-		
+
 		static void AssignPatchTo(UDMXEntityFixturePatch* Patch, int32 ToAddress, int32 UniverseToAssignTo)
 		{
 			Patch->Modify();
@@ -693,7 +699,7 @@ FDMXEditorUtils::FUnassignedPatchesArray FDMXEditorUtils::AutoAssignedAddresses(
 			UDMXEntityFixturePatch* ToAssign,
 			int32 UniverseToAssignTo,
 			int32& MinimumAddress,
-			int32& IndexOfFirstPatchWithMinAddress, 
+			int32& IndexOfFirstPatchWithMinAddress,
 			TArray<UDMXEntityFixturePatch*>& SortedPatchesWithSetAddress)
 		{
 			const int32 NeededSpan = ToAssign->GetChannelSpan();
@@ -705,7 +711,7 @@ FDMXEditorUtils::FUnassignedPatchesArray FDMXEditorUtils::AutoAssignedAddresses(
 			UDMXEntityFixturePatch* FirstPatchWithMinAddress = SortedPatchesWithSetAddress[IndexOfFirstPatchWithMinAddress];
 			const bool bDoesPatchFitBeforeFirstPatchWithMinAddress = FirstPatchWithMinAddress->GetStartingChannel() >= MinimumAddress
 				&& FirstPatchWithMinAddress->GetStartingChannel() - MinimumAddress >= NeededSpan;
-			if(bDoesPatchFitBeforeFirstPatchWithMinAddress)
+			if (bDoesPatchFitBeforeFirstPatchWithMinAddress)
 			{
 				AssignPatchTo(ToAssign, MinimumAddress, UniverseToAssignTo);
 				SortedPatchesWithSetAddress.Insert(ToAssign, 0);
@@ -715,13 +721,13 @@ FDMXEditorUtils::FUnassignedPatchesArray FDMXEditorUtils::AutoAssignedAddresses(
 
 				return true;
 			}
-			
+
 			const bool bNoGapsToLookAt = IndexOfFirstPatchWithMinAddress == SortedPatchesWithSetAddress.Num() - 1;
-			if(bNoGapsToLookAt)
+			if (bNoGapsToLookAt)
 			{
 				return false;
 			}
-			
+
 			int32 PrevAvailableAddress = GetEndAddressOf(SortedPatchesWithSetAddress[IndexOfFirstPatchWithMinAddress]) + 1;
 			for (int32 iPatchInUniverse = IndexOfFirstPatchWithMinAddress + 1; iPatchInUniverse < SortedPatchesWithSetAddress.Num(); ++iPatchInUniverse)
 			{
@@ -763,11 +769,11 @@ FDMXEditorUtils::FUnassignedPatchesArray FDMXEditorUtils::AutoAssignedAddresses(
 	PatchesToAutoAssign.RemoveAll([](UDMXEntityFixturePatch* Patch) {
 		return !Patch->IsAutoAssignAddress();
 		});
-	
+
 	TArray<UDMXEntityFixturePatch*> AllFixturePatches = Library->GetEntitiesTypeCast<UDMXEntityFixturePatch>();
 	TMap<int32, TArray<UDMXEntityFixturePatch*>> UniverseToAllPatches = FDMXRuntimeUtils::MapToUniverses(AllFixturePatches);
 	TArray<UDMXEntityFixturePatch*> PatchesToAssignInNextUniverse;
-	for(auto UniverseIterator = UniverseToAllPatches.CreateIterator(); UniverseIterator; ++UniverseIterator)
+	for (auto UniverseIterator = UniverseToAllPatches.CreateIterator(); UniverseIterator; ++UniverseIterator)
 	{
 		const int32 CurrentUniverse = UniverseIterator->Key;
 		TArray<UDMXEntityFixturePatch*>& SortedPatchesWithSetAddress = [&UniverseIterator, &PatchesToAutoAssign, CurrentUniverse]() -> TArray<UDMXEntityFixturePatch*>&
@@ -797,7 +803,7 @@ FDMXEditorUtils::FUnassignedPatchesArray FDMXEditorUtils::AutoAssignedAddresses(
 			return Result;
 		}();
 
-		if(bCanChangePatchUniverses)
+		if (bCanChangePatchUniverses)
 		{
 			// The patches in PatchesToAssignInNextUniverse are left over from last universe: hence they should ignore the MinimumAddress. 
 			TArray<UDMXEntityFixturePatch*> ToAssignNextUniverse;
@@ -809,9 +815,9 @@ FDMXEditorUtils::FUnassignedPatchesArray FDMXEditorUtils::AutoAssignedAddresses(
 	}
 
 	const bool bNeedNewUniverse = PatchesToAssignInNextUniverse.Num() > 0;
-	if(bNeedNewUniverse)
+	if (bNeedNewUniverse)
 	{
-		if(bCanChangePatchUniverses)
+		if (bCanChangePatchUniverses)
 		{
 			const int32 HighestUniverse = [&UniverseToAllPatches]()
 			{
@@ -822,20 +828,145 @@ FDMXEditorUtils::FUnassignedPatchesArray FDMXEditorUtils::AutoAssignedAddresses(
 				}
 				return HighestSoFar;
 			}();
-			for(UDMXEntityFixturePatch* UnassignedPatch : PatchesToAssignInNextUniverse)
+			for (UDMXEntityFixturePatch* UnassignedPatch : PatchesToAssignInNextUniverse)
 			{
 				UnassignedPatch->Modify();
 				UnassignedPatch->SetUniverseID(HighestUniverse + 1);
 			}
 			return AutoAssignedAddresses(PatchesToAssignInNextUniverse, 1, bCanChangePatchUniverses);
 		}
-		
+
 		return PatchesToAssignInNextUniverse;
 	}
-	
+	PRAGMA_ENABLE_DEPRECATION_WARNINGS
+
 	return {};
 }
 
+void FDMXEditorUtils::AutoAssignedChannels(bool bAllowDecrementUniverse, bool bAllowDecrementChannels, TArray<UDMXEntityFixturePatch*> FixturePatches)
+{
+	if (FixturePatches.IsEmpty())
+	{
+		return;
+	}
+
+	// Ensure patches are of a single DMX Library
+	const UDMXLibrary* DMXLibrary = FixturePatches[0]->GetParentLibrary();
+	for (UDMXEntityFixturePatch* FixturePatch : FixturePatches)
+	{
+		if (!ensureAlwaysMsgf(FixturePatch->GetParentLibrary() == DMXLibrary, TEXT("Trying to auto assign fixture patches from different DMX Libraries at once. This is not supported.")))
+		{
+			return;
+		}
+	}
+
+	TArray<UDMXEntityFixturePatch*> AllFixturePatchesInLibrary = DMXLibrary->GetEntitiesTypeCast<UDMXEntityFixturePatch>();
+	
+	// Sort by Universe and channel
+	FixturePatches.Sort([](const UDMXEntityFixturePatch& FixturePatchA, const UDMXEntityFixturePatch& FixturePatchB)
+		{
+			const bool bUniverseIsSmaller = FixturePatchA.GetUniverseID() < FixturePatchB.GetUniverseID();
+			const bool bUniverseIsEqual = FixturePatchA.GetUniverseID() == FixturePatchB.GetUniverseID();
+			const bool bChannelIsSmallerOrEqual = FixturePatchA.GetStartingChannel() <= FixturePatchB.GetStartingChannel();
+			return bUniverseIsSmaller || (bUniverseIsEqual && bChannelIsSmallerOrEqual);
+		});
+
+	// Using absolute channels (Universe * DMX_UIVERSE_SIZE + Channel)
+	auto FindAbsoluteChannelLambda([bAllowDecrementUniverse, bAllowDecrementChannels, &FixturePatches, &AllFixturePatchesInLibrary](UDMXEntityFixturePatch* FixturePatch, const TArray<UDMXEntityFixturePatch*>& OtherPendingAutoAssignFixturePatches) -> int32
+		{
+			AllFixturePatchesInLibrary.StableSort([](const UDMXEntityFixturePatch& FixturePatchA, const UDMXEntityFixturePatch& FixturePatchB)
+				{
+					const bool bUniverseIsSmaller = FixturePatchA.GetUniverseID() < FixturePatchB.GetUniverseID();
+					const bool bUniverseIsEqual = FixturePatchA.GetUniverseID() == FixturePatchB.GetUniverseID();
+					const bool bChannelIsSmallerOrEqual = FixturePatchA.GetStartingChannel() <= FixturePatchB.GetStartingChannel();
+					return bUniverseIsSmaller || (bUniverseIsEqual && bChannelIsSmallerOrEqual);
+				});
+
+			const int32 OldUniverse = FixturePatch->GetUniverseID();
+			const int32 OldAbsoluteStartingChannel = FixturePatch->GetUniverseID() * DMX_UNIVERSE_SIZE + FixturePatch->GetStartingChannel();
+			const int32 RequiredChannelSpan = FixturePatch->GetChannelSpan();
+			const int32 MinAbsoluteChannel = bAllowDecrementUniverse ? 1 : FixturePatch->GetUniverseID() * DMX_UNIVERSE_SIZE + 1;
+
+			int32 AbsoluteStartOfGap = MinAbsoluteChannel;
+			for (int32 IndexOfFixturePatch = 0; IndexOfFixturePatch < AllFixturePatchesInLibrary.Num(); IndexOfFixturePatch++)
+			{				
+				UDMXEntityFixturePatch* Other = AllFixturePatchesInLibrary[IndexOfFixturePatch];
+				const int32 OtherAbsoluteEndingChannel = Other->GetUniverseID() * DMX_UNIVERSE_SIZE + Other->GetEndingChannel();
+				const int32 GapChannelSpan = Other->GetUniverseID() * DMX_UNIVERSE_SIZE + Other->GetStartingChannel() - AbsoluteStartOfGap;
+
+				// Ignore others that are pending auto-assign
+				if (OtherPendingAutoAssignFixturePatches.Contains(Other))
+				{
+					continue;
+				}
+
+				// Peek ahead when looking at self
+				if (Other == FixturePatch)
+				{
+					UDMXEntityFixturePatch* Next = AllFixturePatchesInLibrary.IsValidIndex(IndexOfFixturePatch + 1) ? AllFixturePatchesInLibrary[IndexOfFixturePatch + 1] : nullptr;
+					if (!Next)
+					{
+						break;
+					}
+
+					const int32 NextAbsoluteStartingChannel = Next->GetUniverseID() * DMX_UNIVERSE_SIZE + Next->GetStartingChannel();
+					const int32 GapToNext = NextAbsoluteStartingChannel - AbsoluteStartOfGap;
+					if (GapToNext >= RequiredChannelSpan)
+					{
+						break;
+					}
+				}
+
+				if (!bAllowDecrementUniverse && OtherAbsoluteEndingChannel + 1 < MinAbsoluteChannel)
+				{
+					AbsoluteStartOfGap = MinAbsoluteChannel;
+					continue;
+				}
+
+				if (!bAllowDecrementChannels && AbsoluteStartOfGap < OldAbsoluteStartingChannel)
+				{
+					AbsoluteStartOfGap = OtherAbsoluteEndingChannel + 1;
+					continue;
+				}
+
+				const bool bGapSizeIsSufficient = GapChannelSpan >= RequiredChannelSpan;
+				const bool bGapSizeDoesNotExceedUniverse = (AbsoluteStartOfGap + GapChannelSpan - 1) % DMX_UNIVERSE_SIZE <= DMX_UNIVERSE_SIZE;
+				if (bGapSizeIsSufficient && bGapSizeDoesNotExceedUniverse)
+				{
+					break;
+				}
+
+				AbsoluteStartOfGap = OtherAbsoluteEndingChannel + 1;
+			}
+
+			int32 Universe = AbsoluteStartOfGap / DMX_UNIVERSE_SIZE;
+			int32 Channel = AbsoluteStartOfGap % DMX_UNIVERSE_SIZE;
+			if (Channel + FixturePatch->GetChannelSpan() - 1 > DMX_UNIVERSE_SIZE)
+			{
+				return (Universe + 1) * DMX_UNIVERSE_SIZE + 1;
+			}
+			else
+			{
+				return AbsoluteStartOfGap;
+			}
+		});
+
+	const FScopedTransaction AutoAssignTransaction(LOCTEXT("AutoAssignTransaction", "Auto Assign Fixture Patch"));
+	TArray<UDMXEntityFixturePatch*> OtherPendingAutoAssignFixturePatches(FixturePatches);
+	for (UDMXEntityFixturePatch* FixturePatch : FixturePatches)
+	{
+		OtherPendingAutoAssignFixturePatches.Remove(FixturePatch);
+		const int32 AutoAssignAbsoluteChannel = FindAbsoluteChannelLambda(FixturePatch, OtherPendingAutoAssignFixturePatches);
+		if (AutoAssignAbsoluteChannel > 0)
+		{
+			FixturePatch->PreEditChange(UDMXEntityFixturePatch::StaticClass()->FindPropertyByName(UDMXEntityFixturePatch::GetStartingChannelPropertyNameChecked()));
+			FixturePatch->SetUniverseID(AutoAssignAbsoluteChannel / DMX_UNIVERSE_SIZE);
+			FixturePatch->SetStartingChannel(AutoAssignAbsoluteChannel % DMX_UNIVERSE_SIZE);
+			FixturePatch->PostEditChange();
+		}
+	}
+}
+
 void FDMXEditorUtils::UpdatePatchColors(UDMXLibrary* Library)
 {
 	check(Library);
diff --git a/Engine/Plugins/VirtualProduction/DMX/DMXEngine/Source/DMXEditor/Private/DMXFixtureTypeSharedData.cpp b/Engine/Plugins/VirtualProduction/DMX/DMXEngine/Source/DMXEditor/Private/DMXFixtureTypeSharedData.cpp
--- a/Engine/Plugins/VirtualProduction/DMX/DMXEngine/Source/DMXEditor/Private/DMXFixtureTypeSharedData.cpp
+++ b/Engine/Plugins/VirtualProduction/DMX/DMXEngine/Source/DMXEditor/Private/DMXFixtureTypeSharedData.cpp
@@ -712,8 +712,6 @@ void FDMXFixtureTypeSharedData::AddFunctionToSelectedMode()
 					FixtureType->UpdateChannelSpan(ModeIndex);
 				}
 
-				FDMXEditorUtils::AutoAssignedAddresses(FixtureType);
-
 				FixtureType->PostEditChange();
 			}
 		}
diff --git a/Engine/Plugins/VirtualProduction/DMX/DMXEngine/Source/DMXEditor/Private/Tests/DMXEditorUtilsTests.cpp b/Engine/Plugins/VirtualProduction/DMX/DMXEngine/Source/DMXEditor/Private/Tests/DMXEditorUtilsTests.cpp
deleted file mode 100644
diff --git a/Engine/Plugins/VirtualProduction/DMX/DMXEngine/Source/DMXEditor/Private/Widgets/FixturePatch/DMXFixturePatchNode.cpp b/Engine/Plugins/VirtualProduction/DMX/DMXEngine/Source/DMXEditor/Private/Widgets/FixturePatch/DMXFixturePatchNode.cpp
--- a/Engine/Plugins/VirtualProduction/DMX/DMXEngine/Source/DMXEditor/Private/Widgets/FixturePatch/DMXFixturePatchNode.cpp
+++ b/Engine/Plugins/VirtualProduction/DMX/DMXEngine/Source/DMXEditor/Private/Widgets/FixturePatch/DMXFixturePatchNode.cpp
@@ -351,7 +351,7 @@ void FDMXFixturePatchNode::SetAddresses(const TSharedRef<SDMXPatchedUniverse>& O
 		// As we allow the asset to update even when not transacted,
 		// We have to set the last known values to enable proper undo
 		FixturePatch->SetUniverseID(LastTransactedUniverseID);
-		FixturePatch->SetManualStartingAddress(LastTransactedChannelID);
+		FixturePatch->SetStartingChannel(LastTransactedChannelID);
 
 		LastTransactedUniverseID = UniverseWidget->GetUniverseID();
 		LastTransactedChannelID = StartingChannel;
@@ -360,14 +360,14 @@ void FDMXFixturePatchNode::SetAddresses(const TSharedRef<SDMXPatchedUniverse>& O
 		FixturePatch->PreEditChange(nullptr);
 
 		FixturePatch->SetUniverseID(UniverseWidget->GetUniverseID());
-		FixturePatch->SetManualStartingAddress(StartingChannel);
+		FixturePatch->SetStartingChannel(StartingChannel);
 
 		FixturePatch->PostEditChange();
 	}
 	else
 	{
 		FixturePatch->SetUniverseID(UniverseWidget->GetUniverseID());
-		FixturePatch->SetManualStartingAddress(StartingChannel);
+		FixturePatch->SetStartingChannel(StartingChannel);
 	}
 }
 
@@ -424,9 +424,29 @@ bool FDMXFixturePatchNode::IsSelected() const
 
 int32 FDMXFixturePatchNode::GetUniverseID() const
 {
-	if (UniverseWidget.IsValid())
+	if (FixturePatch.IsValid())
 	{
-		return UniverseWidget->GetUniverseID();
+		return FixturePatch->GetUniverseID();
+	}
+
+	return -1;
+}
+
+int32 FDMXFixturePatchNode::GetStartingChannel() const
+{
+	if (FixturePatch.IsValid())
+	{
+		return FixturePatch->GetStartingChannel();
+	}
+
+	return -1;
+}
+
+int32 FDMXFixturePatchNode::GetChannelSpan() const
+{
+	if (FixturePatch.IsValid())
+	{
+		return FixturePatch->GetChannelSpan();
 	}
 
 	return -1;
diff --git a/Engine/Plugins/VirtualProduction/DMX/DMXEngine/Source/DMXEditor/Private/Widgets/FixturePatch/DMXFixturePatchNode.h b/Engine/Plugins/VirtualProduction/DMX/DMXEngine/Source/DMXEditor/Private/Widgets/FixturePatch/DMXFixturePatchNode.h
--- a/Engine/Plugins/VirtualProduction/DMX/DMXEngine/Source/DMXEditor/Private/Widgets/FixturePatch/DMXFixturePatchNode.h
+++ b/Engine/Plugins/VirtualProduction/DMX/DMXEngine/Source/DMXEditor/Private/Widgets/FixturePatch/DMXFixturePatchNode.h
@@ -52,10 +52,10 @@ public:
 	int32 GetUniverseID() const;
 
 	/** Returns the Starting Channel the Node currently resides on */
-	int32 GetStartingChannel() const { return StartingChannel; }
+	int32 GetStartingChannel() const;
 
 	/** Returns the Channel Span the Node currently occupies */
-	int32 GetChannelSpan() const { return ChannelSpan; }
+	int32 GetChannelSpan() const;
 
 	/** Sets the ZOrder of the Node */
 	void SetZOrder(int32 NewZOrder) { ZOrder = NewZOrder; }
diff --git a/Engine/Plugins/VirtualProduction/DMX/DMXEngine/Source/DMXEditor/Private/Widgets/FixturePatch/SDMXFixturePatchTree.cpp b/Engine/Plugins/VirtualProduction/DMX/DMXEngine/Source/DMXEditor/Private/Widgets/FixturePatch/SDMXFixturePatchTree.cpp
deleted file mode 100644
diff --git a/Engine/Plugins/VirtualProduction/DMX/DMXEngine/Source/DMXEditor/Private/Widgets/FixturePatch/SDMXFixturePatchTree.h b/Engine/Plugins/VirtualProduction/DMX/DMXEngine/Source/DMXEditor/Private/Widgets/FixturePatch/SDMXFixturePatchTree.h
deleted file mode 100644
diff --git a/Engine/Plugins/VirtualProduction/DMX/DMXEngine/Source/DMXEditor/Private/Widgets/FixturePatch/SDMXFixturePatchTreeFixturePatchRow.cpp b/Engine/Plugins/VirtualProduction/DMX/DMXEngine/Source/DMXEditor/Private/Widgets/FixturePatch/SDMXFixturePatchTreeFixturePatchRow.cpp
deleted file mode 100644
diff --git a/Engine/Plugins/VirtualProduction/DMX/DMXEngine/Source/DMXEditor/Private/Widgets/FixturePatch/SDMXFixturePatchTreeFixturePatchRow.h b/Engine/Plugins/VirtualProduction/DMX/DMXEngine/Source/DMXEditor/Private/Widgets/FixturePatch/SDMXFixturePatchTreeFixturePatchRow.h
deleted file mode 100644
diff --git a/Engine/Plugins/VirtualProduction/DMX/DMXEngine/Source/DMXEditor/Private/Widgets/FixturePatch/SDMXFixturePatchTreeUniverseRow.cpp b/Engine/Plugins/VirtualProduction/DMX/DMXEngine/Source/DMXEditor/Private/Widgets/FixturePatch/SDMXFixturePatchTreeUniverseRow.cpp
deleted file mode 100644
diff --git a/Engine/Plugins/VirtualProduction/DMX/DMXEngine/Source/DMXEditor/Private/Widgets/FixturePatch/SDMXFixturePatchTreeUniverseRow.h b/Engine/Plugins/VirtualProduction/DMX/DMXEngine/Source/DMXEditor/Private/Widgets/FixturePatch/SDMXFixturePatchTreeUniverseRow.h
deleted file mode 100644
diff --git a/Engine/Plugins/VirtualProduction/DMX/DMXEngine/Source/DMXEditor/Private/Widgets/FixturePatch/SDMXFixturePatcher.cpp b/Engine/Plugins/VirtualProduction/DMX/DMXEngine/Source/DMXEditor/Private/Widgets/FixturePatch/SDMXFixturePatcher.cpp
--- a/Engine/Plugins/VirtualProduction/DMX/DMXEngine/Source/DMXEditor/Private/Widgets/FixturePatch/SDMXFixturePatcher.cpp
+++ b/Engine/Plugins/VirtualProduction/DMX/DMXEngine/Source/DMXEditor/Private/Widgets/FixturePatch/SDMXFixturePatcher.cpp
@@ -371,12 +371,6 @@ TSharedPtr<FDMXFixturePatchNode> SDMXFixturePatcher::GetDraggedNode(const TArray
 				DraggedNode = FDMXFixturePatchNode::Create(DMXEditorPtr, FixturePatch);
 			}
 
-			// Remove auto assign to let drag drop set it
-			if (FixturePatch->IsAutoAssignAddress())
-			{
-				DisableAutoAssignAdress(FixturePatch);
-			}
-
 			return DraggedNode;
 		}
 	}
@@ -772,24 +766,6 @@ int32 SDMXFixturePatcher::ClampStartingChannel(int32 StartingChannel, int32 Chan
 	return StartingChannel;
 }
 
-void SDMXFixturePatcher::DisableAutoAssignAdress(TWeakObjectPtr<UDMXEntityFixturePatch> FixturePatch)
-{
-	if (FixturePatch.IsValid())
-	{
-		const FScopedTransaction Transaction = FScopedTransaction(
-			FText::Format(LOCTEXT("AutoAssignAdressChanged", "Disabled Auto Assign Adress for {0}"), 
-				FText::FromString(FixturePatch->GetDisplayName()))
-		);
-
-		FixturePatch->Modify();
-		FixturePatch->PreEditChange(UDMXEntityFixturePatch::StaticClass()->FindPropertyByName(UDMXEntityFixturePatch::GetAutoAssignAddressPropertyNameChecked()));
-
-		FixturePatch->SetAutoAssignAddressUnsafe(false);
-
-		FixturePatch->PostEditChange();
-	}
-}
-
 UDMXLibrary* SDMXFixturePatcher::GetDMXLibrary() const
 {
 	if (TSharedPtr<FDMXEditor> DMXEditor = DMXEditorPtr.Pin())
diff --git a/Engine/Plugins/VirtualProduction/DMX/DMXEngine/Source/DMXEditor/Private/Widgets/FixturePatch/SDMXFixturePatcher.h b/Engine/Plugins/VirtualProduction/DMX/DMXEngine/Source/DMXEditor/Private/Widgets/FixturePatch/SDMXFixturePatcher.h
--- a/Engine/Plugins/VirtualProduction/DMX/DMXEngine/Source/DMXEditor/Private/Widgets/FixturePatch/SDMXFixturePatcher.h
+++ b/Engine/Plugins/VirtualProduction/DMX/DMXEngine/Source/DMXEditor/Private/Widgets/FixturePatch/SDMXFixturePatcher.h
@@ -138,9 +138,6 @@ protected:
 	/** Clamps the starting channel to remain within a valid channel range */
 	int32 ClampStartingChannel(int32 StartingChannel, int32 ChannelSpan) const;
 
-	/** Helper disabling bAutoAssignAdress in fixture patch transacted */
-	void DisableAutoAssignAdress(TWeakObjectPtr<UDMXEntityFixturePatch> FixturePatch);
-
 	/** Returns the DMXLibrary or nullptr if not available */
 	UDMXLibrary* GetDMXLibrary() const;
 
diff --git a/Engine/Plugins/VirtualProduction/DMX/DMXEngine/Source/DMXEditor/Private/Widgets/FixturePatch/SDMXMVRFixtureList.cpp b/Engine/Plugins/VirtualProduction/DMX/DMXEngine/Source/DMXEditor/Private/Widgets/FixturePatch/SDMXMVRFixtureList.cpp
--- a/Engine/Plugins/VirtualProduction/DMX/DMXEngine/Source/DMXEditor/Private/Widgets/FixturePatch/SDMXMVRFixtureList.cpp
+++ b/Engine/Plugins/VirtualProduction/DMX/DMXEngine/Source/DMXEditor/Private/Widgets/FixturePatch/SDMXMVRFixtureList.cpp
@@ -655,6 +655,34 @@ void SDMXMVRFixtureList::AdoptSelectionFromFixturePatchSharedData()
 	}
 }
 
+void SDMXMVRFixtureList::AutoAssignFixturePatches()
+{
+	if (FixturePatchSharedData.IsValid())
+	{
+		TArray<UDMXEntityFixturePatch*> FixturePatchesToAutoAssign;
+		const TArray<TWeakObjectPtr<UDMXEntityFixturePatch>> SelectedFixturePatches = FixturePatchSharedData->GetSelectedFixturePatches();
+		for (TWeakObjectPtr<UDMXEntityFixturePatch> FixturePatch : SelectedFixturePatches)
+		{
+			if (FixturePatch.IsValid())
+			{
+				FixturePatchesToAutoAssign.Add(FixturePatch.Get());
+			}
+		}
+
+		if (FixturePatchesToAutoAssign.IsEmpty())
+		{
+			return;
+		}
+
+		constexpr bool bAllowDecrementUniverse = false;
+		constexpr bool bAllowDecrementChannels = true;
+		FDMXEditorUtils::AutoAssignedChannels(bAllowDecrementUniverse, bAllowDecrementChannels, FixturePatchesToAutoAssign);
+		FixturePatchSharedData->SelectUniverse(FixturePatchesToAutoAssign[0]->GetUniverseID());
+
+		RequestListRefresh();
+	}
+}
+
 TSharedRef<SHeaderRow> SDMXMVRFixtureList::GenerateHeaderRow()
 {
 	const float StatusColumnWidth = FMath::Max(FAppStyle::GetBrush("Icons.Warning")->GetImageSize().X + 6.f, FAppStyle::GetBrush("Icons.Error")->GetImageSize().X + 6.f);
@@ -866,7 +894,23 @@ TSharedPtr<SWidget> SDMXMVRFixtureList::OnContextMenuOpening()
 
 	if (ListView->GetNumItemsSelected() > 0)
 	{
-		MenuBuilder.BeginSection("BasicOperations");
+		// Auto Assign Section
+		MenuBuilder.BeginSection("AutoAssignSection", LOCTEXT("AutoAssignSection", "Auto-Assign"));
+		{
+			// Auto Assign Entry
+			const FUIAction Action(FExecuteAction::CreateSP(this, &SDMXMVRFixtureList::AutoAssignFixturePatches));
+
+			const FText AutoAssignText = LOCTEXT("AutoAssignContextMenuEntry", "Auto-Assign Selection");
+			const TSharedRef<SWidget> Widget =
+				SNew(STextBlock)
+				.Text(AutoAssignText);
+
+			MenuBuilder.AddMenuEntry(Action, Widget);
+			MenuBuilder.EndSection();
+		}
+
+		// Basic Operations Section
+		MenuBuilder.BeginSection("BasicOperationsSection", LOCTEXT("BasicOperationsSection", "Basic Operations"));
 		{
 			MenuBuilder.AddMenuEntry(FGenericCommands::Get().Cut);
 			MenuBuilder.AddMenuEntry(FGenericCommands::Get().Copy);
diff --git a/Engine/Plugins/VirtualProduction/DMX/DMXEngine/Source/DMXEditor/Private/Widgets/FixturePatch/SDMXMVRFixtureList.h b/Engine/Plugins/VirtualProduction/DMX/DMXEngine/Source/DMXEditor/Private/Widgets/FixturePatch/SDMXMVRFixtureList.h
--- a/Engine/Plugins/VirtualProduction/DMX/DMXEngine/Source/DMXEditor/Private/Widgets/FixturePatch/SDMXMVRFixtureList.h
+++ b/Engine/Plugins/VirtualProduction/DMX/DMXEngine/Source/DMXEditor/Private/Widgets/FixturePatch/SDMXMVRFixtureList.h
@@ -105,6 +105,9 @@ private:
 	/** Selects what is selected in Fixture Patch Shared Data */
 	void AdoptSelectionFromFixturePatchSharedData();
 
+	/** Auto assigns selected Fixture Patches */
+	void AutoAssignFixturePatches();
+
 	/** Generates the Header Row of the List */
 	TSharedRef<SHeaderRow> GenerateHeaderRow();
 
diff --git a/Engine/Plugins/VirtualProduction/DMX/DMXEngine/Source/DMXEditor/Private/Widgets/FixturePatch/SDMXPatchedUniverse.cpp b/Engine/Plugins/VirtualProduction/DMX/DMXEngine/Source/DMXEditor/Private/Widgets/FixturePatch/SDMXPatchedUniverse.cpp
--- a/Engine/Plugins/VirtualProduction/DMX/DMXEngine/Source/DMXEditor/Private/Widgets/FixturePatch/SDMXPatchedUniverse.cpp
+++ b/Engine/Plugins/VirtualProduction/DMX/DMXEngine/Source/DMXEditor/Private/Widgets/FixturePatch/SDMXPatchedUniverse.cpp
@@ -3,6 +3,7 @@
 #include "SDMXPatchedUniverse.h"
 
 #include "DMXEditor.h"
+#include "DMXEditorUtils.h"
 #include "DMXFixturePatchEditorDefinitions.h"
 #include "DMXFixturePatchNode.h"
 #include "DMXFixturePatchSharedData.h"
@@ -202,11 +203,11 @@ void SDMXPatchedUniverse::CreateChannelConnectors()
 	{
 		int32 Column = ChannelIndex % FDMXChannelGridSpecs::NumColumns;
 		int32 Row = ChannelIndex / FDMXChannelGridSpecs::NumColumns;
-		int32 ChannelID = ChannelIndex + FDMXChannelGridSpecs::ChannelIDOffset;
+		int32 Channel = ChannelIndex + FDMXChannelGridSpecs::ChannelIDOffset;
 
 		TSharedRef<SDMXChannelConnector> ChannelPatchWidget =
 			SNew(SDMXChannelConnector)
-			.ChannelID(ChannelID)
+			.ChannelID(Channel)
 			.Value(0)
 			.OnMouseButtonDownOnChannel(this, &SDMXPatchedUniverse::HandleOnMouseButtonDownOnChannel)
 			.OnMouseButtonUpOnChannel(this, &SDMXPatchedUniverse::HandleOnMouseButtonUpOnChannel)
@@ -239,12 +240,6 @@ bool SDMXPatchedUniverse::Patch(const TSharedPtr<FDMXFixturePatchNode>& Node, in
 
 	const int32 NewChannelSpan = FixturePatch->GetChannelSpan();
 
-	// Auto assign patches that have bAutoAssignAddress set
-	if (FixturePatch->IsAutoAssignAddress())
-	{
-		NewStartingChannel = FixturePatch->GetStartingChannel();
-	}
-
 	// Unpatch from the old universe
 	if (const TSharedPtr<SDMXPatchedUniverse>& OldUniverseWidget = Node->GetUniverseWidget())
 	{
@@ -444,11 +439,35 @@ void SDMXPatchedUniverse::RefreshInternal()
 	}
 }
 
-FReply SDMXPatchedUniverse::HandleOnMouseButtonDownOnChannel(uint32 ChannelID, const FPointerEvent& PointerEvent)
+FReply SDMXPatchedUniverse::HandleOnMouseButtonDownOnChannel(uint32 Channel, const FPointerEvent& PointerEvent)
 {
+	if (PointerEvent.GetEffectingButton() == EKeys::RightMouseButton)
+	{
+		const TArray<TWeakObjectPtr<UDMXEntityFixturePatch>> OldSelection = SharedData->GetSelectedFixturePatches();
+		const TArray<UDMXEntityFixturePatch*> FixturePatchesOnChannel = GetFixturePatchesOnChannel(Channel);
+		const bool bOldSelectionContainsPatchesOnChannel = FixturePatchesOnChannel.ContainsByPredicate([&OldSelection](const UDMXEntityFixturePatch* FixturePatch)
+			{
+				return OldSelection.Contains(FixturePatch);
+			});
+		if (!bOldSelectionContainsPatchesOnChannel)
+		{
+			if (UDMXEntityFixturePatch* Patch = GetTopmostFixturePatchOnChannel(Channel))
+			{
+				if (FSlateApplication::Get().GetModifierKeys().IsControlDown())
+				{
+					SharedData->AddFixturePatchToSelection(Patch);
+				}
+				else
+				{
+					SharedData->SelectFixturePatch(Patch);
+				}
+			}
+		}
+	}
+
 	if (PointerEvent.GetEffectingButton() == EKeys::LeftMouseButton)
 	{
-		if (UDMXEntityFixturePatch* Patch = GetTopmostFixturePatchOnChannelID(ChannelID))
+		if (UDMXEntityFixturePatch* Patch = GetTopmostFixturePatchOnChannel(Channel))
 		{
 			if (FSlateApplication::Get().GetModifierKeys().IsControlDown())
 			{
@@ -467,9 +486,9 @@ FReply SDMXPatchedUniverse::HandleOnMouseButtonDownOnChannel(uint32 ChannelID, c
 			{
 				Reply.SetUserFocus(WidgetThatWasFocusedBy_SelectFixturePatch.ToSharedRef());
 			}
-			if (ensureMsgf(ChannelConnectors.IsValidIndex(ChannelID - 1), TEXT("Trying to drag Fixture Patch, but the dragged channel is not valid.")))
+			if (ensureMsgf(ChannelConnectors.IsValidIndex(Channel - 1), TEXT("Trying to drag Fixture Patch, but the dragged channel is not valid.")))
 			{
-				return Reply.DetectDrag(ChannelConnectors[ChannelID - 1].ToSharedRef(), EKeys::LeftMouseButton);
+				return Reply.DetectDrag(ChannelConnectors[Channel - 1].ToSharedRef(), EKeys::LeftMouseButton);
 			}
 		}
 	}
@@ -477,13 +496,13 @@ FReply SDMXPatchedUniverse::HandleOnMouseButtonDownOnChannel(uint32 ChannelID, c
 	return FReply::Unhandled();
 }
 
-FReply SDMXPatchedUniverse::HandleOnMouseButtonUpOnChannel(uint32 ChannelID, const FPointerEvent& PointerEvent)
+FReply SDMXPatchedUniverse::HandleOnMouseButtonUpOnChannel(uint32 Channel, const FPointerEvent& PointerEvent)
 {
 	if (PointerEvent.GetEffectingButton() == EKeys::RightMouseButton)
 	{
 		const FVector2D& SummonLocation = PointerEvent.GetScreenSpacePosition();
 
-		TSharedRef<SWidget> MenuContent = CreateContextMenu(ChannelID);
+		TSharedRef<SWidget> MenuContent = CreateContextMenu(Channel);
 		FWidgetPath WidgetPath = PointerEvent.GetEventPath() ? *PointerEvent.GetEventPath() : FWidgetPath();
 		FSlateApplication::Get().PushMenu(AsShared(), WidgetPath, MenuContent, SummonLocation, FPopupTransitionEffect(FPopupTransitionEffect::ContextMenu));
 
@@ -493,11 +512,11 @@ FReply SDMXPatchedUniverse::HandleOnMouseButtonUpOnChannel(uint32 ChannelID, con
 	return FReply::Unhandled();
 }
 
-FReply SDMXPatchedUniverse::HandleOnDragDetectedOnChannel(uint32 ChannelID, const FPointerEvent& PointerEvent)
+FReply SDMXPatchedUniverse::HandleOnDragDetectedOnChannel(uint32 Channel, const FPointerEvent& PointerEvent)
 {
-	if (UDMXEntityFixturePatch* Patch = GetTopmostFixturePatchOnChannelID(ChannelID))
+	if (UDMXEntityFixturePatch* Patch = GetTopmostFixturePatchOnChannel(Channel))
 	{
-		int32 ChannelOffset = ChannelID - Patch->GetStartingChannel();
+		int32 ChannelOffset = Channel - Patch->GetStartingChannel();
 
 		TSharedRef<FDMXEntityFixturePatchDragDropOperation> DragDropOp = MakeShared<FDMXEntityFixturePatchDragDropOperation>(Patch->GetParentLibrary(), TArray<TWeakObjectPtr<UDMXEntity>>{ Patch }, ChannelOffset);
 
@@ -507,19 +526,19 @@ FReply SDMXPatchedUniverse::HandleOnDragDetectedOnChannel(uint32 ChannelID, cons
 	return FReply::Unhandled();
 }
 
-void SDMXPatchedUniverse::HandleDragEnterChannel(uint32 ChannelID, const FDragDropEvent& DragDropEvent)
+void SDMXPatchedUniverse::HandleDragEnterChannel(uint32 Channel, const FDragDropEvent& DragDropEvent)
 {
-	OnDragEnterChannel.ExecuteIfBound(UniverseID, ChannelID, DragDropEvent);
+	OnDragEnterChannel.ExecuteIfBound(UniverseID, Channel, DragDropEvent);
 }
 
-void SDMXPatchedUniverse::HandleDragLeaveChannel(uint32 ChannelID, const FDragDropEvent& DragDropEvent)
+void SDMXPatchedUniverse::HandleDragLeaveChannel(uint32 Channel, const FDragDropEvent& DragDropEvent)
 {
-	OnDragLeaveChannel.ExecuteIfBound(UniverseID, ChannelID, DragDropEvent);
+	OnDragLeaveChannel.ExecuteIfBound(UniverseID, Channel, DragDropEvent);
 }
 
-FReply SDMXPatchedUniverse::HandleDropOntoChannel(uint32 ChannelID, const FDragDropEvent& DragDropEvent)
+FReply SDMXPatchedUniverse::HandleDropOntoChannel(uint32 Channel, const FDragDropEvent& DragDropEvent)
 {
-	return OnDropOntoChannel.Execute(UniverseID, ChannelID, DragDropEvent);
+	return OnDropOntoChannel.Execute(UniverseID, Channel, DragDropEvent);
 }
 
 FText SDMXPatchedUniverse::GetHeaderText() const
@@ -664,17 +683,45 @@ void SDMXPatchedUniverse::OnSelectionChanged()
 	}
 }
 
-UDMXEntityFixturePatch* SDMXPatchedUniverse::GetTopmostFixturePatchOnChannelID(uint32 ChannelID) const
+void SDMXPatchedUniverse::AutoAssignFixturePatches()
+{
+	if (SharedData.IsValid())
+	{
+		TArray<UDMXEntityFixturePatch*> FixturePatchesToAutoAssign;
+		const TArray<TWeakObjectPtr<UDMXEntityFixturePatch>> SelectedFixturePatches = SharedData->GetSelectedFixturePatches();
+		for (TWeakObjectPtr<UDMXEntityFixturePatch> FixturePatch : SelectedFixturePatches)
+		{
+			if (FixturePatch.IsValid())
+			{
+				FixturePatchesToAutoAssign.Add(FixturePatch.Get());
+			}
+		}
+
+		if (FixturePatchesToAutoAssign.IsEmpty())
+		{
+			return;
+		}
+
+		constexpr bool bAllowDecrementUniverse = false;
+		constexpr bool bAllowDecrementChannels = true;
+		FDMXEditorUtils::AutoAssignedChannels(bAllowDecrementUniverse, bAllowDecrementChannels, FixturePatchesToAutoAssign);
+		SharedData->SelectUniverse(FixturePatchesToAutoAssign[0]->GetUniverseID());
+
+		RequestRefresh();
+	}
+}
+
+UDMXEntityFixturePatch* SDMXPatchedUniverse::GetTopmostFixturePatchOnChannel(uint32 Channel) const
 {
 	TArray<TSharedPtr<FDMXFixturePatchNode>> Nodes(PatchedNodes);
-	Nodes.RemoveAll([ChannelID](const TSharedPtr<FDMXFixturePatchNode>& Node)
+	Nodes.RemoveAll([Channel](const TSharedPtr<FDMXFixturePatchNode>& Node)
 			{
 				if (UDMXEntityFixturePatch* FixturePatch = Node->GetFixturePatch().Get())
 				{
 					const uint32 PatchStartingChannel = FixturePatch->GetStartingChannel();
 					const TRange<uint32> FixturePatchRange(PatchStartingChannel, PatchStartingChannel + FixturePatch->GetChannelSpan());
 
-					return !FixturePatchRange.Contains(ChannelID);
+					return !FixturePatchRange.Contains(Channel);
 				}
 				return true;
 			});
@@ -694,17 +741,17 @@ UDMXEntityFixturePatch* SDMXPatchedUniverse::GetTopmostFixturePatchOnChannelID(u
 	}
 }
 
-TArray<UDMXEntityFixturePatch*> SDMXPatchedUniverse::GetFixturePatchesOnChannelID(uint32 ChannelID) const
+TArray<UDMXEntityFixturePatch*> SDMXPatchedUniverse::GetFixturePatchesOnChannel(uint32 Channel) const
 {
 	TArray<TSharedPtr<FDMXFixturePatchNode>> Nodes(PatchedNodes);
-	Nodes.RemoveAll([ChannelID](const TSharedPtr<FDMXFixturePatchNode>& Node)
+	Nodes.RemoveAll([Channel](const TSharedPtr<FDMXFixturePatchNode>& Node)
 		{
 			if (UDMXEntityFixturePatch* FixturePatch = Node->GetFixturePatch().Get())
 			{
 				const uint32 PatchStartingChannel = FixturePatch->GetStartingChannel();
 				const TRange<uint32> FixturePatchRange(PatchStartingChannel, PatchStartingChannel + FixturePatch->GetChannelSpan());
 
-				return !FixturePatchRange.Contains(ChannelID);
+				return !FixturePatchRange.Contains(Channel);
 			}
 			return true;
 		});
@@ -741,7 +788,7 @@ void SDMXPatchedUniverse::RegisterCommands()
 	CommandList = MakeShared<FUICommandList>();
 }
 
-TSharedRef<SWidget> SDMXPatchedUniverse::CreateContextMenu(int32 ChannelID) const
+TSharedRef<SWidget> SDMXPatchedUniverse::CreateContextMenu(int32 Channel)
 {
 	UDMXLibrary* DMXLibrary = GetDMXLibrary();
 	if (!DMXLibrary)
@@ -749,8 +796,8 @@ TSharedRef<SWidget> SDMXPatchedUniverse::CreateContextMenu(int32 ChannelID) cons
 		return SNullWidget::NullWidget;
 	}
 
-	TArray<UDMXEntityFixturePatch*> FixturePatchesOnChannel = GetFixturePatchesOnChannelID(ChannelID);
-	if (FixturePatchesOnChannel.Num() == 0)
+	TArray<UDMXEntityFixturePatch*> FixturePatchesOnChannel = GetFixturePatchesOnChannel(Channel);
+	if (FixturePatchesOnChannel.IsEmpty())
 	{
 		return SNullWidget::NullWidget;
 	}
@@ -758,65 +805,81 @@ TSharedRef<SWidget> SDMXPatchedUniverse::CreateContextMenu(int32 ChannelID) cons
 	constexpr bool bCloseWindowAfterMenuSelection = true;
 	FMenuBuilder MenuBuilder(bCloseWindowAfterMenuSelection, CommandList);
 
-	// Add a context menu
-	if (FixturePatchesOnChannel.Num() > 1)
+	// Auto Assign Section
+	MenuBuilder.BeginSection("AutoAssignSection", LOCTEXT("AutoAssignSection", "Auto-Assign"));
 	{
-		// Common interactions
-		MenuBuilder.BeginSection("FixturePatchInteractions");
-		{
-			// Select All
-			const FUIAction Action(
-				FExecuteAction::CreateLambda([FixturePatchesOnChannel, this]()
-					{
-						if (SharedData.IsValid())
-						{
-							TArray<TWeakObjectPtr<UDMXEntityFixturePatch>> WeakFixturePatchesOnChannel;
-							for (UDMXEntityFixturePatch* FixturePatch : FixturePatchesOnChannel)
-							{
-								WeakFixturePatchesOnChannel.Add(FixturePatch);
-							}
-							SharedData->SelectFixturePatches(WeakFixturePatchesOnChannel);
-						}
-					}));
+		const FUIAction Action(FExecuteAction::CreateSP(this, &SDMXPatchedUniverse::AutoAssignFixturePatches));
 
-			const FText SelectAllText = LOCTEXT("SelectAllFixturePatches", "Select all");
-			const TSharedRef<SWidget> Widget =
-				SNew(STextBlock)
-				.Text(SelectAllText);
+		const FText AutoAssignText = LOCTEXT("AutoAssignContextMenuEntry", "Auto-Assign Selection");
+		const TSharedRef<SWidget> Widget =
+			SNew(STextBlock)
+			.Text(AutoAssignText);
 
-			MenuBuilder.AddMenuEntry(Action, Widget);
-		}
+		MenuBuilder.AddMenuEntry(Action, Widget);
 		MenuBuilder.EndSection();
+	}
 
-		// Single fixture patch selection
-		MenuBuilder.BeginSection("BringFixturePatchForward", LOCTEXT("BringFixturePatchForwardLabel", "Bring Forward"));
+	if (FixturePatchesOnChannel.Num() > 1)
+	{
+		// Select section
+		MenuBuilder.BeginSection("SelectSection", LOCTEXT("SelectSection", "Select"));
 		{
-			const TArray<UDMXEntityFixturePatch*> DMXFixturePatchesInLibrary = DMXLibrary->GetEntitiesTypeCast<UDMXEntityFixturePatch>();
-			FixturePatchesOnChannel.Sort([&DMXFixturePatchesInLibrary](const UDMXEntityFixturePatch& FixturePatchA, const UDMXEntityFixturePatch& FixturePatchB)
-				{
-					return DMXFixturePatchesInLibrary.IndexOfByKey(&FixturePatchA) < DMXFixturePatchesInLibrary.IndexOfByKey(&FixturePatchB);
-				});
-
-			for (TWeakObjectPtr<UDMXEntityFixturePatch> WeakFixturePatch : FixturePatchesOnChannel)
+			// Select all
 			{
 				const FUIAction Action(
-					FExecuteAction::CreateLambda([this, WeakFixturePatch]()
+					FExecuteAction::CreateLambda([FixturePatchesOnChannel, this]()
 						{
-							if (SharedData.IsValid() && WeakFixturePatch.IsValid())
+							if (SharedData.IsValid())
 							{
-								SharedData->SelectFixturePatch(WeakFixturePatch.Get());
+								TArray<TWeakObjectPtr<UDMXEntityFixturePatch>> WeakFixturePatchesOnChannel;
+								for (UDMXEntityFixturePatch* FixturePatch : FixturePatchesOnChannel)
+								{
+									WeakFixturePatchesOnChannel.Add(FixturePatch);
+								}
+								SharedData->SelectFixturePatches(WeakFixturePatchesOnChannel);
 							}
-						})
-				);
+						}));
 
-				const FText FixturePatchName = FText::FromString(WeakFixturePatch->Name);
+				const FText SelectAllText = LOCTEXT("SelectAllFixturePatchesContextMenuEntry", "Select All");
 				const TSharedRef<SWidget> Widget =
 					SNew(STextBlock)
-					.Text(FixturePatchName);
+					.Text(SelectAllText);
 
 				MenuBuilder.AddMenuEntry(Action, Widget);
 			}
+
+			MenuBuilder.AddSeparator();
+
+			// Select specific patch
+			{
+				const TArray<UDMXEntityFixturePatch*> DMXFixturePatchesInLibrary = DMXLibrary->GetEntitiesTypeCast<UDMXEntityFixturePatch>();
+				FixturePatchesOnChannel.Sort([&DMXFixturePatchesInLibrary](const UDMXEntityFixturePatch& FixturePatchA, const UDMXEntityFixturePatch& FixturePatchB)
+					{
+						return DMXFixturePatchesInLibrary.IndexOfByKey(&FixturePatchA) < DMXFixturePatchesInLibrary.IndexOfByKey(&FixturePatchB);
+					});
+
+				for (TWeakObjectPtr<UDMXEntityFixturePatch> WeakFixturePatch : FixturePatchesOnChannel)
+				{
+					const FUIAction Action(
+						FExecuteAction::CreateLambda([this, WeakFixturePatch]()
+							{
+								if (SharedData.IsValid() && WeakFixturePatch.IsValid())
+								{
+									SharedData->SelectFixturePatch(WeakFixturePatch.Get());
+								}
+							})
+					);
+
+					const FText FixturePatchName = FText::FromString(WeakFixturePatch->Name);
+					const TSharedRef<SWidget> Widget =
+						SNew(STextBlock)
+						.Text(FixturePatchName);
+
+					MenuBuilder.AddMenuEntry(Action, Widget);
+				}
+			}
 		}
+
 		MenuBuilder.EndSection();
 	}
 
diff --git a/Engine/Plugins/VirtualProduction/DMX/DMXEngine/Source/DMXEditor/Private/Widgets/FixturePatch/SDMXPatchedUniverse.h b/Engine/Plugins/VirtualProduction/DMX/DMXEngine/Source/DMXEditor/Private/Widgets/FixturePatch/SDMXPatchedUniverse.h
--- a/Engine/Plugins/VirtualProduction/DMX/DMXEngine/Source/DMXEditor/Private/Widgets/FixturePatch/SDMXPatchedUniverse.h
+++ b/Engine/Plugins/VirtualProduction/DMX/DMXEngine/Source/DMXEditor/Private/Widgets/FixturePatch/SDMXPatchedUniverse.h
@@ -40,9 +40,9 @@ enum class EDMXPatchedUniverseReachability
 class SDMXPatchedUniverse
 	: public SCompoundWidget
 {
-	DECLARE_DELEGATE_ThreeParams(FOnDragOverChannel, int32 /** UniverseID */, int32 /** ChannelID */, const FDragDropEvent&);
+	DECLARE_DELEGATE_ThreeParams(FOnDragOverChannel, int32 /** UniverseID */, int32 /** Channel */, const FDragDropEvent&);
 
-	DECLARE_DELEGATE_RetVal_ThreeParams(FReply, FOnDropOntoChannel, int32 /** UniverseID */, int32 /** ChannelID */, const FDragDropEvent&);
+	DECLARE_DELEGATE_RetVal_ThreeParams(FReply, FOnDropOntoChannel, int32 /** UniverseID */, int32 /** Channel */, const FDragDropEvent&);
 
 public:
 	SDMXPatchedUniverse()
@@ -132,22 +132,22 @@ private:
 	FText GetHeaderText() const;
 	
 	/** Handles when a mouse button was pressed on a Channel */
-	FReply HandleOnMouseButtonDownOnChannel(uint32 ChannelID, const FPointerEvent& PointerEvent);
+	FReply HandleOnMouseButtonDownOnChannel(uint32 Channel, const FPointerEvent& PointerEvent);
 
 	/** Handles when a mouse button was released on a Channel */
-	FReply HandleOnMouseButtonUpOnChannel(uint32 ChannelID, const FPointerEvent& PointerEvent);
+	FReply HandleOnMouseButtonUpOnChannel(uint32 Channel, const FPointerEvent& PointerEvent);
 
 	/** Handles when a Channel was dragged */
-	FReply HandleOnDragDetectedOnChannel(uint32 ChannelID, const FPointerEvent& PointerEvent);
+	FReply HandleOnDragDetectedOnChannel(uint32 Channel, const FPointerEvent& PointerEvent);
 
 	/** Handles when drag enters a channel */
-	void HandleDragEnterChannel(uint32 ChannelID, const FDragDropEvent& DragDropEvent);
+	void HandleDragEnterChannel(uint32 Channel, const FDragDropEvent& DragDropEvent);
 
 	/** Handles when drag leaves a channel */
-	void HandleDragLeaveChannel(uint32 ChannelID, const FDragDropEvent& DragDropEvent);
+	void HandleDragLeaveChannel(uint32 Channel, const FDragDropEvent& DragDropEvent);
 
 	/** Handles when drag dropped onto a channel */
-	FReply HandleDropOntoChannel(uint32 ChannelID, const FDragDropEvent& DragDropEvent);
+	FReply HandleDropOntoChannel(uint32 Channel, const FDragDropEvent& DragDropEvent);
 
 	/** Updates the ZOrder of all Nodes */
 	void UpdateZOrderOfNodes();
@@ -161,11 +161,14 @@ private:
 	/** Called when selection changed */
 	void OnSelectionChanged();
 
-	/** Returns the Fixture Patch that is topmost under ChannelID */
-	UDMXEntityFixturePatch* GetTopmostFixturePatchOnChannelID(uint32 ChannelID) const;
+	/** Auto assigns selected Fixture Patches */
+	void AutoAssignFixturePatches();
+
+	/** Returns the Fixture Patch that is topmost under Channel */
+	UDMXEntityFixturePatch* GetTopmostFixturePatchOnChannel(uint32 Channel) const;
 
 	/** Returns all Fixture Patches on a Channel ID */
-	TArray<UDMXEntityFixturePatch*> GetFixturePatchesOnChannelID(uint32 ChannelID) const;
+	TArray<UDMXEntityFixturePatch*> GetFixturePatchesOnChannel(uint32 Channel) const;
 
 	/** Returns the DMXLibrary or nullptr if not available */
 	UDMXLibrary* GetDMXLibrary() const;
@@ -217,7 +220,7 @@ private:
 	void RegisterCommands();
 
 	/** Creates the right click context menu */
-	TSharedRef<SWidget> CreateContextMenu(int32 ChannelID) const;
+	TSharedRef<SWidget> CreateContextMenu(int32 Channel);
 
 	/** Command list for this widget */
 	TSharedPtr<FUICommandList> CommandList;
diff --git a/Engine/Plugins/VirtualProduction/DMX/DMXEngine/Source/DMXEditor/Public/DMXEditorUtils.h b/Engine/Plugins/VirtualProduction/DMX/DMXEngine/Source/DMXEditor/Public/DMXEditorUtils.h
--- a/Engine/Plugins/VirtualProduction/DMX/DMXEngine/Source/DMXEditor/Public/DMXEditorUtils.h
+++ b/Engine/Plugins/VirtualProduction/DMX/DMXEngine/Source/DMXEditor/Public/DMXEditorUtils.h
@@ -80,26 +80,18 @@ public:
 	/**  Returns the Entity class type name (e.g: Fixture Type for UDMXEntityFixtureType) in singular or plural */
 	static FText GetEntityTypeNameText(TSubclassOf<UDMXEntity> EntityClass, bool bPlural = false);
 
-
-// Auto Assign:
-	
-	/**
-	 * Updates starting address for the patch, if it has bAutoAssignAddress set and only if it can be assigned
-	 * to a universe in AllowedUniverses.
-	 *
-	 * If no universe can be assigned, then this patch is left unmodified.
-	 *
-	 * @return Whether Patch was auto assigned.
-	 */
+	/** DEPRECATED 5.1 */
+	UE_DEPRECATED(5.1, "Simplyifing AutoAssign, just use FDMXEditorUtils::AutoAssignedAddresses instead.")
 	static bool TryAutoAssignToUniverses(UDMXEntityFixturePatch* Patch, const TSet<int32>& AllowedUniverses);
-	
-	/**
-	 * Updates Addresses for Fixture Patches that use specified Parent fixture type and have bAutoAssignAddress set. 
-	 *
-	 * @param ChangedParentFixtureType	The parent fixture type of the patches that want their channels to be auto assigned
-	 */
+
+	/** DEPRECATED 5.1 */
+	UE_DEPRECATED(5.1, "The bAutoAssign property got deprecated in UDMXEntityFixturePatch, so there's no meaning to ever auto assign on Fixture Type changes. So this function is removed without replacement.")
 	static void AutoAssignedAddresses(UDMXEntityFixtureType* ChangedParentFixtureType);
 
+	/** DEPRECATED 5.1 */
+	UE_DEPRECATED(5.1, "Simplyifing AutoAssign, just use FDMXEditorUtils::AutoAssignedChannels instead.")
+	static FUnassignedPatchesArray AutoAssignedAddresses(const TArray<UDMXEntityFixturePatch*>& ChangedFixturePatches, int32 MinimumAddress = 1, bool bCanChangePatchUniverses = true);
+
 	/**
 	 * Updates starting addresses for fixture patches that have bAutoAssignAddress set, ignores others.
 	 * Note, patches all have to reside in the same library.
@@ -109,13 +101,11 @@ public:
 	 * If bCanChangePatchUniverses = true, this function will assign patches that do not fit into the existing universes to
 	 * the next universe. MinimumAddress is ignored for the new universe, i.e. we start placing remaining patches at position 1.
 	 *
-	 * @param ChangedFixturePatches		The patches that want their channels to be auto assigned
-	 * @param MinimumAddress			All patches must be placed after this address
-	 * @param bCanChangePatchUniverses	Whether we are allowed to move a patch to another universe if the assigned universe has no space.
-	 *
-	 * @result Patches that were not assigned.
+	 * @param bAllowDecrementUniverse	If true, the patches can be assigend to a previous universe
+	 * @param bCanChangePatchUniverses	If true, the patches can lower their absoulte channel
+	 * @param FixturePatches			The patches that want their channels to be auto assigned
 	 */
-	static FUnassignedPatchesArray AutoAssignedAddresses(const TArray<UDMXEntityFixturePatch*>& ChangedFixturePatches, int32 MinimumAddress = 1, bool bCanChangePatchUniverses = true);
+	static void AutoAssignedChannels(bool bAllowDecrementUniverse, bool bAllowDecrementChannels, TArray<UDMXEntityFixturePatch*> FixturePatches);
 	
 	/**
 	 * Creates a unique color for all patches that use the default color FLinearColor(1.0f, 0.0f, 1.0f)
diff --git a/Engine/Plugins/VirtualProduction/DMX/DMXEngine/Source/DMXRuntime/Private/DMXRuntimeMainStreamObjectVersion.h b/Engine/Plugins/VirtualProduction/DMX/DMXEngine/Source/DMXRuntime/Private/DMXRuntimeMainStreamObjectVersion.h
--- a/Engine/Plugins/VirtualProduction/DMX/DMXEngine/Source/DMXRuntime/Private/DMXRuntimeMainStreamObjectVersion.h
+++ b/Engine/Plugins/VirtualProduction/DMX/DMXEngine/Source/DMXRuntime/Private/DMXRuntimeMainStreamObjectVersion.h
@@ -27,6 +27,9 @@ struct FDMXRuntimeMainStreamObjectVersion
 		// Update Fixture Patch to use an MVR UUID
 		DMXFixturePatchHasMVRUUID,
 
+		// Remove auto assign features from fixture patch (instead it's auto assigned via FDMXEditorUtils where required)
+		DMXFixturePatchNoLongerImplementsAutoAssign,
+
 		VersionPlusOne,
 		LatestVersion = VersionPlusOne - 1
 	};
diff --git a/Engine/Plugins/VirtualProduction/DMX/DMXEngine/Source/DMXRuntime/Private/Library/DMXEntityFixturePatch.cpp b/Engine/Plugins/VirtualProduction/DMX/DMXEngine/Source/DMXRuntime/Private/Library/DMXEntityFixturePatch.cpp
--- a/Engine/Plugins/VirtualProduction/DMX/DMXEngine/Source/DMXRuntime/Private/Library/DMXEntityFixturePatch.cpp
+++ b/Engine/Plugins/VirtualProduction/DMX/DMXEngine/Source/DMXRuntime/Private/Library/DMXEntityFixturePatch.cpp
@@ -32,9 +32,6 @@ FDMXOnFixturePatchChangedDelegate UDMXEntityFixturePatch::OnFixturePatchChangedD
 
 UDMXEntityFixturePatch::UDMXEntityFixturePatch()
 	: UniverseID(1)
-	, bAutoAssignAddress(true)
-	, ManualStartingAddress(1)
-	, AutoStartingAddress(1)
 	, ActiveMode(0)
 #if WITH_EDITORONLY_DATA
 	, EditorColor(FLinearColor(1.0f, 0.0f, 1.0f))
@@ -166,6 +163,15 @@ void UDMXEntityFixturePatch::Serialize(FArchive& Ar)
 		{
 			MVRFixtureUUID = FGuid::NewGuid();
 		}
+
+#if WITH_EDITOR
+		if (Ar.CustomVer(FDMXRuntimeMainStreamObjectVersion::GUID) < FDMXRuntimeMainStreamObjectVersion::DMXFixturePatchNoLongerImplementsAutoAssign)
+		{
+			PRAGMA_DISABLE_DEPRECATION_WARNINGS
+			StartingChannel = ManualStartingAddress_DEPRECATED;
+			PRAGMA_ENABLE_DEPRECATION_WARNINGS
+		}
+#endif
 	}
 }
 
@@ -499,47 +505,39 @@ void UDMXEntityFixturePatch::SetUniverseID(int32 NewUniverseID)
 	RebuildCache();
 }
 
+#if WITH_EDITOR
 void UDMXEntityFixturePatch::SetAutoStartingAddress(int32 NewAutoStartingAddress)
 {
-	AutoStartingAddress = NewAutoStartingAddress;
-	ManualStartingAddress = NewAutoStartingAddress;
+	// DEPRECATED 5.1
+	AutoStartingAddress_DEPRECATED = NewAutoStartingAddress;
+	ManualStartingAddress_DEPRECATED = NewAutoStartingAddress;
 
 	RebuildCache();
 }
+#endif // WITH_EDITOR
 
+#if WITH_EDITOR
 void UDMXEntityFixturePatch::SetManualStartingAddress(int32 NewManualStartingAddress)
 {
-	ManualStartingAddress = NewManualStartingAddress;
+	// DEPRECATED 5.1
+	ManualStartingAddress_DEPRECATED = NewManualStartingAddress;
 
 	RebuildCache();
 }
+#endif // WITH_EDITOR
 
 void UDMXEntityFixturePatch::SetStartingChannel(int32 NewStartingChannel)
 {
-	if (NewStartingChannel == AutoStartingAddress && ManualStartingAddress == NewStartingChannel)
+	if (NewStartingChannel == StartingChannel)
 	{
 		return;
 	}
 
-	bAutoAssignAddress = false;
-	AutoStartingAddress = NewStartingChannel;
-	ManualStartingAddress = NewStartingChannel;
+	StartingChannel = NewStartingChannel;
 
 	RebuildCache();
 }
 
-int32 UDMXEntityFixturePatch::GetStartingChannel() const
-{
-	if (bAutoAssignAddress)
-	{
-		return AutoStartingAddress;
-	}
-	else
-	{
-		return ManualStartingAddress;
-	}
-}
-
 int32 UDMXEntityFixturePatch::GetChannelSpan() const
 {
 	return Cache.GetChannelSpan();
@@ -566,12 +564,14 @@ bool UDMXEntityFixturePatch::SetActiveModeIndex(int32 NewActiveModeIndex)
 	return false;
 }
 
+#if WITH_EDITOR
 int32 UDMXEntityFixturePatch::GetRemoteUniverse() const
 {	
 	/** DEPRECATED 4.27 */
 	UE_LOG(LogDMXRuntime, Error, TEXT("No clear remote Universe can be deduced in DMXEntityFixturePatch::GetRemoteUniverse. Returning 0."));
 	return 0;
 }
+#endif // WITH_EDITOR
 
 TArray<FDMXAttributeName> UDMXEntityFixturePatch::GetAllAttributesInActiveMode() const
 {
diff --git a/Engine/Plugins/VirtualProduction/DMX/DMXEngine/Source/DMXRuntime/Public/Library/DMXEntityFixturePatch.h b/Engine/Plugins/VirtualProduction/DMX/DMXEngine/Source/DMXRuntime/Public/Library/DMXEntityFixturePatch.h
--- a/Engine/Plugins/VirtualProduction/DMX/DMXEngine/Source/DMXRuntime/Public/Library/DMXEntityFixturePatch.h
+++ b/Engine/Plugins/VirtualProduction/DMX/DMXEngine/Source/DMXRuntime/Public/Library/DMXEntityFixturePatch.h
@@ -165,37 +165,39 @@ public:
 	/** Sets the Universe ID of the patch*/
 	void SetUniverseID(int32 NewUniverseID);
 
-	/**  Sets the auto starting address */
+#if WITH_EDITOR
+	UE_DEPRECATED(5.1, "bAutoAssignAddress and related members are deprecated. Please use SetStartingChannel instead of SetAutoStartingAddress.")
 	void SetAutoStartingAddress(int32 NewAutoStartingAddress);
 
-	/** Returns the auto starting channel. Used only when bAutoAssignAdress is true. In most cases, you'll want GetStartingChannel instead. */
-	FORCEINLINE int32 GetAutoStartingAddress() const { return AutoStartingAddress; }
+	UE_DEPRECATED(5.1, "bAutoAssignAddress and related members are deprecated. Please use GetStartingChannel instead of GetAutoStartingAddress.")
+	FORCEINLINE int32 GetAutoStartingAddress() const { return AutoStartingAddress_DEPRECATED; }
 
-	/**  Sets the manual starting address */
+	UE_DEPRECATED(5.1, "bAutoAssignAddress and related members are deprecated. Please use SetStartingChannel instead of SetManualStartingAddress.")
 	void SetManualStartingAddress(int32 NewManualStartingAddress);
 
-	/** Returns the manual starting channel. Used only when bAutoAssignAdress is false.  In most cases, you'll want GetStartingChannel instead. */
-	FORCEINLINE int32 GetManualStartingAddress() const { return ManualStartingAddress; }
+	UE_DEPRECATED(5.1, "bAutoAssignAddress and related members are deprecated. Please use GetStartingChannel instead of GetManualStartingAddress.")
+	FORCEINLINE int32 GetManualStartingAddress() const { return ManualStartingAddress_DEPRECATED; }
+#endif // WITH_EDITOR
 
 	/** 
 	 * Sets the starting channel of the Fixture Patch.
 	 * 
 	 * If Auto Assign Address was set to true, turns off Auto Assign Address.
 	 */
-	UFUNCTION(BlueprintCallable, Category = "DMX|Fixture Patch")
 	void SetStartingChannel(int32 NewStartingChannel);
 
 	/** Return the starting channel */
-	UFUNCTION(BlueprintPure, Category = "DMX|Fixture Patch")
-	int32 GetStartingChannel() const;
+	FORCEINLINE int32 GetStartingChannel() const { return StartingChannel; }
 
 #if WITH_EDITOR
 	/** Sets bAutoAssignAddress for the patch. Does not update relevant properties. Use UDMXEditor Module's UDMXEditorUtils::AutoAssignAddresses methods instead. */
-	void SetAutoAssignAddressUnsafe(bool bShouldAutoAssignAddress) { bAutoAssignAddress = bShouldAutoAssignAddress; }
-#endif // WITH_EDITOR
+	UE_DEPRECATED(5.1, "bAutoAssignAddress and related members are deprecated. Auto assign is now a method in FDMXEditorUtils and applied on demand.")
+	void SetAutoAssignAddressUnsafe(bool bShouldAutoAssignAddress) { bAutoAssignAddress_DEPRECATED = bShouldAutoAssignAddress; }
 
 	/** Returns true if the patch is set to auto assign address */
-	FORCEINLINE bool IsAutoAssignAddress() const { return bAutoAssignAddress; }
+	UE_DEPRECATED(5.1, "bAutoAssignAddress and related members are deprecated. Auto assign is now a method in FDMXEditorUtils and applied on demand.")
+	FORCEINLINE bool IsAutoAssignAddress() const { return bAutoAssignAddress_DEPRECATED; }
+#endif // WITH_EDITOR
 
 	/** Returns the number of channels this Patch occupies with the Fixture functions from its Active Mode or 0 if the patch has no valid Active Mode */
 	UFUNCTION(BlueprintPure, Category = "DMX|Fixture Patch")
@@ -220,12 +222,19 @@ public:
 #if WITH_EDITOR
 	/** Property name getters. When accessing the this way, use with care. The patch will need to update its cache after using property setters, use Pre/PostEditChanged events to achieve that. */
 	static FName GetUniverseIDPropertyNameChecked() { return GET_MEMBER_NAME_CHECKED(UDMXEntityFixturePatch, UniverseID); }
-	static FName GetAutoAssignAddressPropertyNameChecked() { return GET_MEMBER_NAME_CHECKED(UDMXEntityFixturePatch, bAutoAssignAddress); }
-	static FName GetManualStartingAddressPropertyNameChecked() { return GET_MEMBER_NAME_CHECKED(UDMXEntityFixturePatch, ManualStartingAddress); }
-	static FName GetAutoStartingAddressPropertyNameChecked() { return GET_MEMBER_NAME_CHECKED(UDMXEntityFixturePatch, AutoStartingAddress); }
 	static FName GetParentFixtureTypeTemplatePropertyNameChecked() { return GET_MEMBER_NAME_CHECKED(UDMXEntityFixturePatch, ParentFixtureTypeTemplate); }
 	static FName GetActiveModePropertyNameChecked() { return GET_MEMBER_NAME_CHECKED(UDMXEntityFixturePatch, ActiveMode); }
 	static FName GetMVRFixtureUUIDPropertyNameChecked() { return GET_MEMBER_NAME_CHECKED(UDMXEntityFixturePatch, MVRFixtureUUID); }
+	static FName GetStartingChannelPropertyNameChecked() { return GET_MEMBER_NAME_CHECKED(UDMXEntityFixturePatch, StartingChannel); }
+
+	UE_DEPRECATED(5.1, "bAutoAssignAddress and related members are deprecated. Auto assign is now a method in FDMXEditorUtils and applied on demand.")
+	static FName GetManualStartingAddressPropertyNameChecked() { return GET_MEMBER_NAME_CHECKED(UDMXEntityFixturePatch, ManualStartingAddress_DEPRECATED); }
+
+	UE_DEPRECATED(5.1, "bAutoAssignAddress and related members are deprecated. Auto assign is now a method in FDMXEditorUtils and applied on demand.")
+	static FName GetAutoStartingAddressPropertyNameChecked() { return GET_MEMBER_NAME_CHECKED(UDMXEntityFixturePatch, AutoStartingAddress_DEPRECATED); }
+
+	UE_DEPRECATED(5.1, "bAutoAssignAddress and related members are deprecated. Auto assign is now a method in FDMXEditorUtils and applied on demand.")
+	static FName GetAutoAssignAddressPropertyNameChecked() { return GET_MEMBER_NAME_CHECKED(UDMXEntityFixturePatch, bAutoAssignAddress_DEPRECATED); }
 #endif // WITH_EDITOR
 
 protected:
@@ -233,17 +242,23 @@ protected:
 	UPROPERTY(EditAnywhere, BlueprintReadOnly, Category = "Fixture Patch", meta = (ClampMin = 0, DisplayName = "Universe"))
 	int32 UniverseID;
 
-	/** Auto-assign address from drag/drop list order and available channels */
-	UPROPERTY(EditAnywhere, BlueprintReadOnly, Category = "Fixture Patch", meta = (DisplayName = "Auto-Assign Address"))
-	bool bAutoAssignAddress;
+#if WITH_EDITORONLY_DATA
+	/** DEPRECATED 5.1 */
+	UPROPERTY(Meta = (DeprecatedProperty, DeprecationMessage = "bAutoAssignAddress and related members are deprecated. Auto assign is now only a method in FDMXEditorUtils and should be applied on demand."))
+	bool bAutoAssignAddress_DEPRECATED = true;
 
-	/** Starting channel for when auto-assign address is false */
-	UPROPERTY(EditAnywhere, BlueprintReadOnly, Category = "Fixture Patch", meta = (EditCondition = "!bAutoAssignAddress", DisplayName = "Manual Starting Address", UIMin = "1", UIMax = "512", ClampMin = "1", ClampMax = "512"))
-	int32 ManualStartingAddress;
+	/** DEPRECATED 5.1 */
+	UPROPERTY(Meta = (DeprecatedProperty, DeprecationMessage = "bAutoAssignAddress and related members are deprecated. Auto assign is now only a method in FDMXEditorUtils and should be applied on demand."))
+	int32 ManualStartingAddress_DEPRECATED = 1;
 
-	/** Starting channel from auto-assignment. Used when AutoAssignAddress is true */
-	UPROPERTY(NonTransactional)
-	int32 AutoStartingAddress;
+	/** DEPRECATED 5.1 */
+	UPROPERTY(Meta = (DeprecatedProperty, DeprecationMessage = "bAutoAssignAddress and related members are deprecated. Auto assign is now only a method in FDMXEditorUtils and should be applied on demand."))
+	int32 AutoStartingAddress_DEPRECATED = 1;
+#endif // WITH_EDITORONLY_DATA
+
+	/** Starting Channel of the Patch */
+	UPROPERTY(EditAnywhere, BlueprintReadOnly, Category = "Fixture Patch", Meta = (UIMin = "1", UIMax = "512", ClampMin = "1", ClampMax = "512"))
+	int32 StartingChannel = 0;
 
 	/** Property to point to the template parent fixture for details panel purposes */
 	UPROPERTY(VisibleDefaultsOnly, BlueprintReadWrite, BlueprintSetter = SetFixtureType, Category = "Fixture Patch", meta = (DisplayName = "Fixture Type"))
@@ -279,10 +294,12 @@ public:
 #endif
 
 public:
+#if WITH_EDITOR
 	/** DEPRECATED 4.27 */
 	UE_DEPRECATED(4.27, "Controllers are replaced with DMX Ports.")
 	UFUNCTION(BlueprintPure, Category = "DMX|Fixture Patch", meta = (DeprecatedFunction, DeprecationMessage = "Deprecated 4.27. No clear remote universe can be deduced from controllers (before 4.27) or ports (from 4.27 on)."))
 	int32 GetRemoteUniverse() const;
+#endif // WITH_EDITOR
 
 	/**
 	 * Returns an array of valid attributes for the currently active mode.
