Fix mismatched deprecation marcros

- CL 20982598 introduced a mismatch pair of PRAGMA_DISABLE_DEPRECATION_WARNINGS
  - convert the second to PRAGMA_ENABLE_DEPRECATION_WARNINGS to resolve

#rnx

diff --git a/Engine/Plugins/VirtualProduction/DMX/DMXEngine/Source/DMXEditor/Private/DMXEditorUtils.cpp b/Engine/Plugins/VirtualProduction/DMX/DMXEngine/Source/DMXEditor/Private/DMXEditorUtils.cpp
--- a/Engine/Plugins/VirtualProduction/DMX/DMXEngine/Source/DMXEditor/Private/DMXEditorUtils.cpp
+++ b/Engine/Plugins/VirtualProduction/DMX/DMXEngine/Source/DMXEditor/Private/DMXEditorUtils.cpp
@@ -617,7 +617,7 @@ void FDMXEditorUtils::AutoAssignedAddresses(UDMXEntityFixtureType* ChangedParent
 
 		AutoAssignedAddresses(FixturePatches);
 	}
-	PRAGMA_DISABLE_DEPRECATION_WARNINGS
+	PRAGMA_ENABLE_DEPRECATION_WARNINGS
 }
 
 FDMXEditorUtils::FUnassignedPatchesArray FDMXEditorUtils::AutoAssignedAddresses(
