Fix CIS Static Analysis Win64 (PVS-Studio) - DMXEditorUtils.cpp warning: Expression is always true.

diff --git a/Engine/Plugins/VirtualProduction/DMX/DMXEngine/Source/DMXEditor/Private/DMXEditorUtils.cpp b/Engine/Plugins/VirtualProduction/DMX/DMXEngine/Source/DMXEditor/Private/DMXEditorUtils.cpp
--- a/Engine/Plugins/VirtualProduction/DMX/DMXEngine/Source/DMXEditor/Private/DMXEditorUtils.cpp
+++ b/Engine/Plugins/VirtualProduction/DMX/DMXEngine/Source/DMXEditor/Private/DMXEditorUtils.cpp
@@ -930,7 +930,8 @@ void FDMXEditorUtils::AutoAssignedChannels(bool bAllowDecrementUniverse, bool bA
 				}
 
 				const bool bGapSizeIsSufficient = GapChannelSpan >= RequiredChannelSpan;
-				const bool bGapSizeDoesNotExceedUniverse = (AbsoluteStartOfGap + GapChannelSpan - 1) % DMX_UNIVERSE_SIZE <= DMX_UNIVERSE_SIZE;
+				const int32 RelativeStartingChannel = (AbsoluteStartOfGap + GapChannelSpan - 1) % DMX_UNIVERSE_SIZE;
+				const bool bGapSizeDoesNotExceedUniverse = RelativeStartingChannel <= DMX_UNIVERSE_SIZE;
 				if (bGapSizeIsSufficient && bGapSizeDoesNotExceedUniverse)
 				{
 					break;
