DMX - Fix badly named variable, relative ending channel is now correctly named RelativeEndingChannel

diff --git a/Engine/Plugins/VirtualProduction/DMX/DMXEngine/Source/DMXEditor/Private/DMXEditorUtils.cpp b/Engine/Plugins/VirtualProduction/DMX/DMXEngine/Source/DMXEditor/Private/DMXEditorUtils.cpp
--- a/Engine/Plugins/VirtualProduction/DMX/DMXEngine/Source/DMXEditor/Private/DMXEditorUtils.cpp
+++ b/Engine/Plugins/VirtualProduction/DMX/DMXEngine/Source/DMXEditor/Private/DMXEditorUtils.cpp
@@ -930,8 +930,8 @@ void FDMXEditorUtils::AutoAssignedChannels(bool bAllowDecrementUniverse, bool bA
 				}
 
 				const bool bGapSizeIsSufficient = GapChannelSpan >= RequiredChannelSpan;
-				const int32 RelativeStartingChannel = (AbsoluteStartOfGap + GapChannelSpan - 1) % DMX_UNIVERSE_SIZE;
-				const bool bGapSizeDoesNotExceedUniverse = RelativeStartingChannel <= DMX_UNIVERSE_SIZE;
+				const int32 RelativeEndingChannel = (AbsoluteStartOfGap + GapChannelSpan - 1) % DMX_UNIVERSE_SIZE;
+				const bool bGapSizeDoesNotExceedUniverse = RelativeEndingChannel <= DMX_UNIVERSE_SIZE;
 				if (bGapSizeIsSufficient && bGapSizeDoesNotExceedUniverse)
 				{
 					break;
