Fix CIS Static Analysis Win64 (PVS-Studio) - DMXEditorUtils.cpp warning: Expression is always true.

diff --git a/Engine/Plugins/VirtualProduction/DMX/DMXEngine/Source/DMXEditor/Private/DMXEditorUtils.cpp b/Engine/Plugins/VirtualProduction/DMX/DMXEngine/Source/DMXEditor/Private/DMXEditorUtils.cpp
--- a/Engine/Plugins/VirtualProduction/DMX/DMXEngine/Source/DMXEditor/Private/DMXEditorUtils.cpp
+++ b/Engine/Plugins/VirtualProduction/DMX/DMXEngine/Source/DMXEditor/Private/DMXEditorUtils.cpp
@@ -929,10 +929,7 @@ void FDMXEditorUtils::AutoAssignedChannels(bool bAllowDecrementUniverse, bool bA
 					continue;
 				}
 
-				const bool bGapSizeIsSufficient = GapChannelSpan >= RequiredChannelSpan;
-				const int32 RelativeEndingChannel = (AbsoluteStartOfGap + GapChannelSpan - 1) % DMX_UNIVERSE_SIZE;
-				const bool bGapSizeDoesNotExceedUniverse = RelativeEndingChannel <= DMX_UNIVERSE_SIZE;
-				if (bGapSizeIsSufficient && bGapSizeDoesNotExceedUniverse)
+				if (GapChannelSpan >= RequiredChannelSpan)
 				{
 					break;
 				}
