DMX - Add MVR export functionality to the DMX Library

diff --git a/Engine/Plugins/VirtualProduction/DMX/DMXEngine/Source/DMXEditor/DMXEditor.Build.cs b/Engine/Plugins/VirtualProduction/DMX/DMXEngine/Source/DMXEditor/DMXEditor.Build.cs
--- a/Engine/Plugins/VirtualProduction/DMX/DMXEngine/Source/DMXEditor/DMXEditor.Build.cs
+++ b/Engine/Plugins/VirtualProduction/DMX/DMXEngine/Source/DMXEditor/DMXEditor.Build.cs
@@ -27,6 +27,7 @@ public class DMXEditor : ModuleRules
 				"AssetTools",
 				"ContentBrowser",
 				"CoreUObject",
+				"DesktopPlatform",
 				"EditorFramework",
 				"EditorStyle",
 				"EditorWidgets",
@@ -45,11 +46,11 @@ public class DMXEditor : ModuleRules
 				"SlateCore",
 				"TakesCore",
 				"TakeRecorder",
-				"ToolWidgets",
 				"TakeTrackRecorders",
+				"ToolWidgets",
 				"UnrealEd",
 				"XmlParser",
 			}
 		);
 	}
-}
\ No newline at end of file
+}
diff --git a/Engine/Plugins/VirtualProduction/DMX/DMXEngine/Source/DMXEditor/Private/Commands/DMXEditorCommands.cpp b/Engine/Plugins/VirtualProduction/DMX/DMXEngine/Source/DMXEditor/Private/Commands/DMXEditorCommands.cpp
--- a/Engine/Plugins/VirtualProduction/DMX/DMXEngine/Source/DMXEditor/Private/Commands/DMXEditorCommands.cpp
+++ b/Engine/Plugins/VirtualProduction/DMX/DMXEngine/Source/DMXEditor/Private/Commands/DMXEditorCommands.cpp
@@ -12,6 +12,8 @@ void FDMXEditorCommands::RegisterCommands()
 {
 	UI_COMMAND(GoToDocumentation, "View Documentation", "View documentation about DMX editor", EUserInterfaceActionType::Button, FInputChord());
 
+	UI_COMMAND(ExportDMXLibrary, "Export", "Exports the DMX Library", EUserInterfaceActionType::Button, FInputChord());
+
 	UI_COMMAND(AddNewEntityFixtureType, "New Fixture Type", "Creates a new Fixture Type in this library", EUserInterfaceActionType::Button, FInputChord());
 	UI_COMMAND(AddNewEntityFixturePatch, "Add Fixture", "Creates a new Fixture Patch in this library", EUserInterfaceActionType::Button, FInputChord());
 
diff --git a/Engine/Plugins/VirtualProduction/DMX/DMXEngine/Source/DMXEditor/Private/DMXEditor.cpp b/Engine/Plugins/VirtualProduction/DMX/DMXEngine/Source/DMXEditor/Private/DMXEditor.cpp
--- a/Engine/Plugins/VirtualProduction/DMX/DMXEngine/Source/DMXEditor/Private/DMXEditor.cpp
+++ b/Engine/Plugins/VirtualProduction/DMX/DMXEngine/Source/DMXEditor/Private/DMXEditor.cpp
@@ -3,12 +3,14 @@
 #include "DMXEditor.h"
 #include "DMXEditorLog.h"
 #include "DMXEditorModule.h"
+#include "DMXEditorSettings.h"
 #include "DMXEditorTabNames.h"
 #include "DMXEditorUtils.h"
 #include "DMXFixtureTypeSharedData.h"
 #include "DMXFixturePatchSharedData.h"
 #include "DMXRuntimeLog.h"
 #include "DMXRuntimeUtils.h"
+#include "Exporters/DMXMVRExporter.h"
 #include "Library/DMXLibrary.h"
 #include "Library/DMXEntityFixtureType.h"
 #include "Library/DMXEntityFixturePatch.h"
@@ -26,8 +28,13 @@
 #include "Widgets/LibrarySettings/SDMXLibraryEditorTab.h"
 #include "Widgets/OutputConsole/SDMXOutputConsole.h"
 
-#include "Modules/ModuleManager.h"
+#include "DesktopPlatformModule.h"
+#include "IDesktopPlatform.h"
 #include "ScopedTransaction.h"
+#include "Framework/Notifications/NotificationManager.h"
+#include "Modules/ModuleManager.h"
+#include "Styling/AppStyle.h"
+#include "Widgets/Notifications/SNotificationList.h"
 
 
 #define LOCTEXT_NAMESPACE "FDMXEditor"
@@ -113,11 +120,70 @@ void FDMXEditor::RegisterApplicationModes(UDMXLibrary* DMXLibrary, bool bShouldO
 	FWorkflowCentricApplication::SetCurrentMode(FDMXEditorApplicationMode::DefaultsMode);
 }
 
-UDMXLibrary * FDMXEditor::GetDMXLibrary() const
+UDMXLibrary* FDMXEditor::GetDMXLibrary() const
 {
-	return GetEditableDMXLibrary();
+	return Cast<UDMXLibrary>(GetEditingObject());
 }
 
+void FDMXEditor::ExportDMXLibrary() const
+{
+	UDMXLibrary* DMXLibrary = GetDMXLibrary();
+	if (!DMXLibrary)
+	{
+		return;
+	}
+
+	IDesktopPlatform* DesktopPlatform = FDesktopPlatformModule::Get();
+	if (DesktopPlatform)
+	{
+		UDMXEditorSettings* DMXEditorSettings = GetMutableDefault<UDMXEditorSettings>();
+		check(DMXEditorSettings);
+
+		const FString DefaultPath = [DMXEditorSettings]()
+		{
+			if (FPaths::DirectoryExists(DMXEditorSettings->LastMVRExportPath))
+			{
+				return DMXEditorSettings->LastMVRExportPath;
+			}
+			return FPaths::ProjectSavedDir();
+		}();
+
+		TArray<FString> SaveFilenames;
+		const bool bSaveFile = DesktopPlatform->SaveFileDialog(
+			FSlateApplication::Get().FindBestParentWindowHandleForDialogs(nullptr),
+			LOCTEXT("ExportMVR", "Export MVR").ToString(),
+			DefaultPath,
+			DMXLibrary->GetName() + TEXT(".mvr"),
+			TEXT("My Virtual Rig (*.mvr)|*.mvr"),
+			EFileDialogFlags::None,
+			SaveFilenames);
+
+		if (!bSaveFile || SaveFilenames.IsEmpty())
+		{
+			return;
+		}
+
+		const bool bExportSuccess = FDMXMVRExporter::Export(DMXLibrary, SaveFilenames[0]);
+		if (bExportSuccess)
+		{
+			const FString SavePath = FPaths::GetPath(SaveFilenames[0]);
+			DMXEditorSettings->LastMVRExportPath = SavePath;
+
+			FNotificationInfo NotificationInfo(FText::Format(LOCTEXT("ExportDMXLibraryAsMVRSuccessNotification", "Successfully exported MVR to {0}."), FText::FromString(SaveFilenames[0])));
+			NotificationInfo.ExpireDuration = 5.f;
+
+			FSlateNotificationManager::Get().AddNotification(NotificationInfo);
+		}
+		else
+		{
+			FNotificationInfo NotificationInfo(LOCTEXT("ExportDMXLibraryAsMVRFailureNotification", "Failed to export MVR"));
+			NotificationInfo.ExpireDuration = 10.f;
+			NotificationInfo.Image = FAppStyle::GetBrush("Icons.Warning");
+
+			FSlateNotificationManager::Get().AddNotification(NotificationInfo);
+		}
+	}
+}
 
 void FDMXEditor::RegisterToolbarTab(const TSharedRef<class FTabManager>& InTabManager)
 {
@@ -135,6 +201,10 @@ void FDMXEditor::CreateDefaultCommands()
 {
 	FDMXEditorCommands::Register();
 
+	ToolkitCommands->MapAction(
+		FDMXEditorCommands::Get().ExportDMXLibrary,
+		FExecuteAction::CreateSP(this, &FDMXEditor::ExportDMXLibrary)
+	);
 	ToolkitCommands->MapAction(
 		FDMXEditorCommands::Get().AddNewEntityFixtureType,
 		FExecuteAction::CreateLambda([this]() { OnAddNewEntity(UDMXEntityFixtureType::StaticClass()); })
@@ -319,14 +389,9 @@ TArray<UDMXEntity*> FDMXEditor::GetSelectedEntitiesFromTypeTab(TSubclassOf<UDMXE
 	return TArray<UDMXEntity*>();
 }
 
-UDMXLibrary* FDMXEditor::GetEditableDMXLibrary() const
-{
-	return Cast<UDMXLibrary>(GetEditingObject());
-}
-
 TSharedRef<SDMXLibraryEditorTab> FDMXEditor::CreateDMXLibraryEditorTab()
 {
-	UDMXLibrary* DMXLibrary = GetEditableDMXLibrary();
+	UDMXLibrary* DMXLibrary = GetDMXLibrary();
 	check(DMXLibrary);
 
 	return SNew(SDMXLibraryEditorTab)
diff --git a/Engine/Plugins/VirtualProduction/DMX/DMXEngine/Source/DMXEditor/Private/DMXEditorStyle.cpp b/Engine/Plugins/VirtualProduction/DMX/DMXEngine/Source/DMXEditor/Private/DMXEditorStyle.cpp
--- a/Engine/Plugins/VirtualProduction/DMX/DMXEngine/Source/DMXEditor/Private/DMXEditorStyle.cpp
+++ b/Engine/Plugins/VirtualProduction/DMX/DMXEngine/Source/DMXEditor/Private/DMXEditorStyle.cpp
@@ -12,6 +12,8 @@
 #include "Styling/StyleColors.h"
 
 
+#define EDITOR_IMAGE_BRUSH_SVG( RelativePath, ... ) FSlateVectorImageBrush(FPaths::EngineContentDir() / "Editor/Slate" / RelativePath + TEXT(".svg"), __VA_ARGS__ )
+
 FDMXEditorStyle::FDMXEditorStyle()
 	: FSlateStyleSet("DMXEditorStyle")
 {
@@ -75,9 +77,10 @@ FDMXEditorStyle::FDMXEditorStyle()
 	{ 
 		Set("Icons.DMXLibrary", new IMAGE_BRUSH_SVG("DMXLibrary_20", Icon20x20));
 
-		{
-			Set("Icons.DMXLibrarySettings", new CORE_IMAGE_BRUSH_SVG("Starship/Common/settings", Icon16x16));
-		}
+		Set("Icons.DMXLibrarySettings", new CORE_IMAGE_BRUSH_SVG("Starship/Common/settings", Icon16x16));
+
+		Set("Icons.DMXLibraryToolbar.Export", new EDITOR_IMAGE_BRUSH_SVG("Starship/Common/Export", Icon20x20));
+
 		Set("Icons.FixtureType", new IMAGE_BRUSH_SVG("FixtureType_16", Icon16x16));
 		Set("Icons.FixturePatch", new IMAGE_BRUSH_SVG("FixturePatch_16", Icon16x16));
 
@@ -139,6 +142,7 @@ FDMXEditorStyle::FDMXEditorStyle()
 			.SetArrowsImage(FSlateNoResource()));
 	}
 
+	// Fixture List
 	{
 		Set("MVRFixtureList.Row", FTableRowStyle()
 			.SetEvenRowBackgroundBrush(IMAGE_BRUSH("Common/Selection", Icon8x8, FLinearColor(1.f, 1.f, 1.f, .2f)))
@@ -198,3 +202,5 @@ void FDMXEditorStyle::ReloadTextures()
 		FSlateApplication::Get().GetRenderer()->ReloadTextureResources();
 	}
 }
+
+#undef EDITOR_IMAGE_BRUSH_SVG
diff --git a/Engine/Plugins/VirtualProduction/DMX/DMXEngine/Source/DMXEditor/Private/DMXZipper.cpp b/Engine/Plugins/VirtualProduction/DMX/DMXEngine/Source/DMXEditor/Private/DMXZipper.cpp
new file mode 100644
--- /dev/null
+++ b/Engine/Plugins/VirtualProduction/DMX/DMXEngine/Source/DMXEditor/Private/DMXZipper.cpp
@@ -0,0 +1,381 @@
+// Fill out your copyright notice in the Description page of Project Settings.
+
+#include "DMXZipper.h"
+
+#include "DMXEditorLog.h"
+
+#include "Misc/Compression.h"
+#include "Misc/FileHelper.h"
+#include "Misc/Paths.h"
+
+
+FDMXZipper::FDMXZipper()
+{
+	SizeOfCentralDirectory = 0;
+	OffsetOfCentralDirectory = 0;
+	TotalNumberOfCentralDirectoryRecords = 0;
+}
+
+bool FDMXZipper::LoadFromFile(const FString& Filename)
+{
+	// Reset the reader
+	Reader.Empty();
+
+	if (!FFileHelper::LoadFileToArray(Reader, *Filename))
+	{
+		return false;
+	}
+
+	return Parse();
+}
+
+bool FDMXZipper::LoadFromData(const TArray64<uint8>& Data)
+{
+	// Reset the reader
+	Reader.Empty();
+
+	Reader.Append(Data);
+
+	return Parse();
+}
+
+bool FDMXZipper::SaveToFile(const FString& Filename)
+{
+	FArrayWriter Writer;
+
+	OffsetOfCentralDirectory = 0;
+	SizeOfCentralDirectory = 0;
+	TotalNumberOfCentralDirectoryRecords = 0;
+
+	for (const TPair<FString, TPair<TArray64<uint8>, bool>>& Pair : NewFilesToShouldCompressMap)
+	{
+		const bool bFileAdded = AddFileInternal(Writer, Pair.Key, Pair.Value.Key, Pair.Value.Value);
+		if (!bFileAdded)
+		{
+			UE_LOG(LogDMXEditor, Error, TEXT("Cannot save Zip. File '%s' added no longer exists on disk!"), *Pair.Key);
+			return false;
+		}
+	}
+
+	for (const TPair<FString, uint32>& Pair : OffsetsMap)
+	{
+		// skip overwritten files
+		if (NewFilesToShouldCompressMap.Contains(Pair.Key))
+		{
+			continue;
+		}
+		TArray64<uint8> Data;
+		if (!GetFileContent(Pair.Key, Data))
+		{
+			UE_LOG(LogDMXEditor, Error, TEXT("Cannot save Zip.Failed to read File '%s' from Zip when saving it."), *Pair.Key);
+			return false;
+		}
+
+		if (!AddFileInternal(Writer, Pair.Key, Data, true))
+		{
+			return false;
+		}
+	}
+
+	AddEndOfCentralDirectory(Writer);
+
+	return FFileHelper::SaveArrayToFile(Writer, *Filename);
+}
+
+TArray<FString> FDMXZipper::GetFiles() const
+{
+	TArray<FString> Files;
+	OffsetsMap.GetKeys(Files);
+	return Files;
+}
+
+void FDMXZipper::AddFile(const FString& RelativeFilePathAndName, const TArray64<uint8>& Data, const bool bCompress)
+{
+	NewFilesToShouldCompressMap.Add(RelativeFilePathAndName, { Data, bCompress });
+}
+
+bool FDMXZipper::GetFileContent(const FString& Filename, TArray64<uint8>& OutData)
+{
+	uint32* Offset = OffsetsMap.Find(Filename);
+	if (!Offset)
+	{
+		return false;
+	}
+
+	constexpr uint64 LocalEntryMinSize = 30;
+
+	if (*Offset + LocalEntryMinSize > Reader.Num())
+	{
+		return false;
+	}
+
+	uint16 Compression = 0;
+	uint32 CompressedSize;
+	uint32 UncompressedSize = 0;
+	uint16 FilenameLen = 0;
+	uint16 ExtraFieldLen = 0;
+
+	// Seek to Compression
+	Reader.Seek(*Offset + 8);
+	Reader << Compression;
+	// Seek to CompressedSize
+	Reader.Seek(*Offset + 18);
+	Reader << CompressedSize;
+	Reader << UncompressedSize;
+	Reader << FilenameLen;
+	Reader << ExtraFieldLen;
+
+	if (*Offset + LocalEntryMinSize + FilenameLen + ExtraFieldLen + CompressedSize > Reader.Num())
+	{
+		return false;
+	}
+
+	if (Compression == 8)
+	{
+		OutData.AddUninitialized(UncompressedSize);
+		if (!FCompression::UncompressMemory(NAME_Zlib, OutData.GetData(), UncompressedSize, Reader.GetData() + *Offset + LocalEntryMinSize + FilenameLen + ExtraFieldLen, CompressedSize, COMPRESS_NoFlags, -15))
+		{
+			return false;
+		}
+	}
+	else if (Compression == 0 && CompressedSize == UncompressedSize)
+	{
+		OutData.Append(Reader.GetData() + *Offset + LocalEntryMinSize + FilenameLen + ExtraFieldLen, UncompressedSize);
+	}
+	else
+	{
+		return false;
+	}
+
+	return true;
+}
+
+FDMXZipper::FDMXScopedUnzipToTempFile::FDMXScopedUnzipToTempFile(const TSharedRef<FDMXZipper>& DMXZipper, const FString& FilenameInZip)
+{
+	if (!ensureAlwaysMsgf(DMXZipper->GetFiles().Contains(FilenameInZip), TEXT("Tried to unzip '%s', but zip does not contain the file."), *FilenameInZip))
+	{
+		return;
+	}
+	TArray64<uint8> FileData;
+	if (!ensureAlwaysMsgf(DMXZipper->GetFileContent(FilenameInZip, FileData), TEXT("Tried to unzip '%s', but zip does not contain the file."), *FilenameInZip))
+	{
+		return;
+	}
+
+	const FString Directory = FPaths::EngineSavedDir() / TEXT("DMX_Temp");
+	TempFilePathAndName = FPaths::ConvertRelativePathToFull(Directory / FilenameInZip);
+
+	FFileHelper::SaveArrayToFile(FileData, *TempFilePathAndName);
+}
+
+FDMXZipper::FDMXScopedUnzipToTempFile::~FDMXScopedUnzipToTempFile()
+{
+	// Delete the extracted file
+	IFileManager& FileManager = IFileManager::Get();
+	constexpr bool bRequireExists = true;
+	constexpr bool bEvenIfReadOnly = false;
+	constexpr bool bQuiet = true;
+	ensureAlwaysMsgf(FileManager.Delete(*TempFilePathAndName, bRequireExists, bEvenIfReadOnly, bQuiet), TEXT("Unexpected failed to delete temp file when using FDMXScopedUnzipToTempFile"));
+}
+
+bool FDMXZipper::AddFileInternal(FArrayWriter& Writer, const FString& Path, const TArray64<uint8>& Data, const bool bCompress)
+{
+	constexpr uint16 DeflateMode = 8;
+
+	uint32 LocalFileHeaderSignature = 0x04034b50;
+	uint16 VersionNeededToExtract = 0;
+	uint16 GeneralPurposeFlags = 0;
+	uint16 CompressionMethod = bCompress ? DeflateMode : 0;
+	uint16 FileLastModificationTime = 0;
+	uint16 FileLastModificationDate = 0;
+	uint32 Crc = FCrc::MemCrc32(Data.GetData(), Data.Num(), 0);
+	uint32 CompressedSize = Data.Num();
+	uint32 UncompressedSize = Data.Num();
+	uint16 FilenameLength = FCStringUtf8::Strlen(reinterpret_cast<UTF8CHAR*>(TCHAR_TO_UTF8(*Path)));
+	uint16 ExtraFieldLength = 0;
+
+	TArray64<uint8> CompressedData;
+
+	if (bCompress)
+	{
+		int32 CompressionBounds = FCompression::CompressMemoryBound(NAME_Zlib, Data.Num());
+		CompressedData.AddUninitialized(CompressionBounds);
+		if (!FCompression::CompressMemory(NAME_Zlib, CompressedData.GetData(), CompressionBounds, Data.GetData(), Data.Num(), COMPRESS_NoFlags, -15))
+		{
+			return false;
+		}
+		CompressedSize = CompressionBounds;
+	}
+
+	uint32 LocalFileEntrySize = 30 + FilenameLength + (bCompress ? CompressedSize : Data.Num());
+
+	Writer.InsertZeroed(OffsetOfCentralDirectory, LocalFileEntrySize);
+	Writer.Seek(OffsetOfCentralDirectory);
+
+	Writer << LocalFileHeaderSignature;
+	Writer << VersionNeededToExtract;
+	Writer << GeneralPurposeFlags;
+	Writer << CompressionMethod;
+	Writer << FileLastModificationTime;
+	Writer << FileLastModificationDate;
+	Writer << Crc;
+	Writer << CompressedSize;
+	Writer << UncompressedSize;
+	Writer << FilenameLength;
+	Writer << ExtraFieldLength;
+
+	Writer.Serialize(TCHAR_TO_UTF8(*Path), FilenameLength);
+
+	if (bCompress)
+	{
+		Writer.Serialize(CompressedData.GetData(), CompressedSize);
+	}
+	else
+	{
+		Writer.Serialize(const_cast<uint8*>(Data.GetData()), Data.Num());
+	}
+
+	Writer.InsertZeroed(OffsetOfCentralDirectory + LocalFileEntrySize + SizeOfCentralDirectory, 46 + FilenameLength);
+	Writer.Seek(OffsetOfCentralDirectory + LocalFileEntrySize + SizeOfCentralDirectory);
+
+	uint32 CentralDirectoryFileHeaderSignature = 0x02014b50;
+	uint16 VersionMadeBy = 0;
+	uint16 FileCommentLength = 0;
+	uint16 DiskNumberWhereFileStarts = 0;
+	uint16 InternalFileAttributes = 0;
+	uint32 ExternalFileAttributes = 0;
+
+	Writer << CentralDirectoryFileHeaderSignature;
+	Writer << VersionMadeBy;
+	Writer << VersionNeededToExtract;
+	Writer << GeneralPurposeFlags;
+	Writer << CompressionMethod;
+	Writer << FileLastModificationTime;
+	Writer << FileLastModificationDate;
+	Writer << Crc;
+	Writer << CompressedSize;
+	Writer << UncompressedSize;
+	Writer << FilenameLength;
+	Writer << ExtraFieldLength;
+	Writer << FileCommentLength;
+	Writer << DiskNumberWhereFileStarts;
+	Writer << InternalFileAttributes;
+	Writer << ExternalFileAttributes;
+	Writer << OffsetOfCentralDirectory;
+	Writer.Serialize(TCHAR_TO_UTF8(*Path), FilenameLength);
+
+	OffsetOfCentralDirectory += LocalFileEntrySize;
+	SizeOfCentralDirectory += 46 + FilenameLength;
+	TotalNumberOfCentralDirectoryRecords++;
+
+	return true;
+}
+
+void FDMXZipper::AddEndOfCentralDirectory(FArrayWriter& Writer)
+{
+	uint32 EndOfCentralDirectoryRecordSignature = 0x06054b50;
+	uint16 NumberOfThisDisk = 0;
+	uint16 DiskWhereCentralDirectoryStarts = 0;
+	uint16 CommentLength = 0;
+
+	Writer << EndOfCentralDirectoryRecordSignature;
+	Writer << NumberOfThisDisk;
+	Writer << DiskWhereCentralDirectoryStarts;
+	Writer << TotalNumberOfCentralDirectoryRecords; // on disk
+	Writer << TotalNumberOfCentralDirectoryRecords;
+	Writer << SizeOfCentralDirectory;
+	Writer << OffsetOfCentralDirectory;
+	Writer << CommentLength;
+}
+
+bool FDMXZipper::Parse()
+{
+	// Reset parser status
+	Reader.Seek(0);
+	OffsetsMap.Empty();
+	NewFilesToShouldCompressMap.Empty();
+
+	// Retrieve the trailer magic
+	TArray<uint8> Magic;
+	bool bIndexFound = false;
+	int64 Index = 0;
+	for (Index = Reader.Num() - 1; Index >= 0; Index--)
+	{
+		Magic.Insert(Reader[Index], 0);
+		if (Magic.Num() == 4)
+		{
+			if (Magic[0] == 0x50 && Magic[1] == 0x4b && Magic[2] == 0x05 && Magic[3] == 0x06)
+			{
+				bIndexFound = true;
+				break;
+			}
+			Magic.Pop();
+		}
+	}
+
+	if (!bIndexFound)
+	{
+		return false;
+	}
+
+	uint16 DiskEntries = 0;
+	uint16 TotalEntries = 0;
+	uint32 ReaderSizeOfCentralDirectory = 0;
+	uint32 ReaderOffsetOfCentralDirectory = 0;
+	uint16 CommentLen = 0;
+
+	constexpr uint64 TrailerMinSize = 22;
+	constexpr uint64 CentralDirectoryMinSize = 46;
+
+	if (Index + TrailerMinSize > Reader.Num())
+	{
+		return false;
+	}
+
+	// Skip signature and disk data
+	Reader.Seek(Index + 8);
+	Reader << DiskEntries;
+	Reader << TotalEntries;
+	Reader << ReaderSizeOfCentralDirectory;
+	Reader << ReaderOffsetOfCentralDirectory;
+	Reader << CommentLen;
+
+	uint16 ReaderTotalNumberOfCentralDirectoryRecords = FMath::Min(DiskEntries, TotalEntries);
+
+	for (uint16 DirectoryIndex = 0; DirectoryIndex < ReaderTotalNumberOfCentralDirectoryRecords; DirectoryIndex++)
+	{
+		if (ReaderTotalNumberOfCentralDirectoryRecords + CentralDirectoryMinSize > Reader.Num())
+		{
+			return false;
+		}
+
+		uint16 FilenameLen = 0;
+		uint16 ExtraFieldLen = 0;
+		uint16 EntryCommentLen = 0;
+		uint32 EntryOffset = 0;
+
+		// Seek to FilenameLen
+		Reader.Seek(ReaderOffsetOfCentralDirectory + 28);
+		Reader << FilenameLen;
+		Reader << ExtraFieldLen;
+		Reader << EntryCommentLen;
+		// Seek to EntryOffset
+		Reader.Seek(ReaderOffsetOfCentralDirectory + 42);
+		Reader << EntryOffset;
+
+		if (ReaderOffsetOfCentralDirectory + CentralDirectoryMinSize + FilenameLen + ExtraFieldLen + EntryCommentLen > Reader.Num())
+		{
+			return false;
+		}
+
+		TArray64<uint8> FilenameBytes;
+		FilenameBytes.Append(Reader.GetData() + ReaderOffsetOfCentralDirectory + CentralDirectoryMinSize, FilenameLen);
+		FilenameBytes.Add(0);
+
+		OffsetsMap.Add(UTF8_TO_TCHAR(FilenameBytes.GetData()), EntryOffset);
+
+		ReaderOffsetOfCentralDirectory += CentralDirectoryMinSize + FilenameLen + ExtraFieldLen + EntryCommentLen;
+	}
+
+	return true;
+}
diff --git a/Engine/Plugins/VirtualProduction/DMX/DMXEngine/Source/DMXEditor/Private/DMXZipper.h b/Engine/Plugins/VirtualProduction/DMX/DMXEngine/Source/DMXEditor/Private/DMXZipper.h
new file mode 100644
--- /dev/null
+++ b/Engine/Plugins/VirtualProduction/DMX/DMXEngine/Source/DMXEditor/Private/DMXZipper.h
@@ -0,0 +1,107 @@
+// Fill out your copyright notice in the Description page of Project Settings.
+
+#pragma once
+
+#include "CoreMinimal.h"
+#include "Serialization/ArrayReader.h"
+#include "Serialization/ArrayWriter.h"
+
+
+/**
+ * Zip reader/writer for DMX specific Zip Files such as MVR and GDTF.
+ * 
+ * Example of use:
+ * bool DoStuffWithZip(const FString& ZipFileName)
+ * {
+ *		FDMXZipper Zipper;
+ *		Zipper.LoadFromFile(ZipFileName);
+ *
+ *		Zipper.AddFile("hello/world", { 'A', 'B', 'C' });
+ *		Zipper.AddFile("test001", { '0', '1', '2', '3', 'X' }, false);
+ *		Zipper.AddFile("a/b/c/d/e/test002", { 80, 90, 100, 110, 111, 112, 113 });
+ *
+ *		TArray64<uint8> BigFile;
+ *		if (FFileHelper::LoadFileToArray(BigFile, TEXT("D:/SM_MatPreviewMesh_01.stl")))
+ *		{
+ *			Zipper.AddFile("Meshes/Mesh.stl", BigFile);
+ *		}
+ *
+ *		if (!Zipper.SaveToFile(ZipFileName))
+ *		{
+ *			return false;
+ *		}
+ *
+ *		return true;
+ * }
+ */
+class FDMXZipper
+	: public TSharedFromThis<FDMXZipper>
+{
+public:
+	/** Constructor */
+	FDMXZipper();
+
+	/** Loads the specified zip file */
+	UE_NODISCARD bool LoadFromFile(const FString& Filename);
+
+	/** Loads the data as zip file */
+	UE_NODISCARD bool LoadFromData(const TArray64<uint8>& Data);
+
+	/** Saves the zip to specified Filename */
+	UE_NODISCARD bool SaveToFile(const FString& Filename);
+
+	/** Returns the File Names in the Zip. Note, may contain relative paths */
+	TArray<FString> GetFiles() const;
+
+	/** 
+	 * Adds a file to the Zip at its Relative File Path and Name. 
+	 * Note, files added this way need to exist on disk until SaveToFile is called.
+	 * 
+	 * Examples:
+	 * AddFile("hello/world", { 'A', 'B', 'C' });
+     * AddFile("test001", { '0', '1', '2', '3', 'X' }, false);
+     * AddFile("a/b/c/d/e/test002", { 80, 90, 100, 110, 111, 112, 113 });
+	 */
+	void AddFile(const FString& RelativeFilePathAndName, const TArray64<uint8>& Data, const bool bCompress = false);
+
+	/** Gets the Data of a File in the Zip */
+	UE_NODISCARD bool GetFileContent(const FString& Filename, TArray64<uint8>& OutData);
+
+	/** Helper to unzip a file within the zip as temp file. Deletes the file when running out of scope. */
+	struct FDMXScopedUnzipToTempFile
+	{
+		FDMXScopedUnzipToTempFile(const TSharedRef<FDMXZipper>& DMXZipper, const FString& FilenameInZip);
+		~FDMXScopedUnzipToTempFile();
+
+		/** Absolute File Path and Name of the Temp File */
+		FString TempFilePathAndName;
+	};
+
+protected:
+	/** Writes a zipped File to the array writer */
+	UE_NODISCARD bool AddFileInternal(FArrayWriter& Writer, const FString& Path, const TArray64<uint8>& Data, const bool bCompress = true);
+
+	/** Adds an End of Central Directory entry to the zip */
+	void AddEndOfCentralDirectory(FArrayWriter& Writer);
+
+	/** Parses the Zip File */
+	UE_NODISCARD bool Parse();
+
+	/** The Size of the Central Directory */
+	uint32 SizeOfCentralDirectory = 0;
+
+	/** The Offset of the Central Directory */
+	uint32 OffsetOfCentralDirectory = 0;
+
+	/** The total number of Central Directory records */
+	uint16 TotalNumberOfCentralDirectoryRecords = 0;
+
+	/** Array reader handling the data */
+	FArrayReader Reader;
+
+	/** Map of File Names and their Offset */
+	TMap<FString, uint32> OffsetsMap;
+
+	/** Map of Files added to the Zip */
+	TMap<FString, TPair<TArray64<uint8>, bool>> NewFilesToShouldCompressMap;
+};
diff --git a/Engine/Plugins/VirtualProduction/DMX/DMXEngine/Source/DMXEditor/Private/Exporters/DMXMVRExporter.cpp b/Engine/Plugins/VirtualProduction/DMX/DMXEngine/Source/DMXEditor/Private/Exporters/DMXMVRExporter.cpp
new file mode 100644
--- /dev/null
+++ b/Engine/Plugins/VirtualProduction/DMX/DMXEngine/Source/DMXEditor/Private/Exporters/DMXMVRExporter.cpp
@@ -0,0 +1,198 @@
+// Copyright Epic Games, Inc. All Rights Reserved.
+
+#include "Exporters/DMXMVRExporter.h"
+
+#include "DMXEditorLog.h"
+#include "DMXMVRXmlMergeUtility.h"
+#include "DMXZipper.h"
+#include "Library/DMXEntityFixturePatch.h"
+#include "Library/DMXEntityFixtureType.h"
+#include "Library/DMXGDTFAssetImportData.h"
+#include "Library/DMXImportGDTF.h"
+#include "Library/DMXLibrary.h"
+#include "MVR/DMXMVRAssetImportData.h"
+#include "MVR/DMXMVRGeneralSceneDescription.h"
+
+#include "XmlFile.h"
+#include "Misc/FileHelper.h"
+#include "Misc/Paths.h"
+
+
+bool FDMXMVRExporter::Export(UDMXLibrary* DMXLibrary, const FString& FilePathAndName)
+{
+	FDMXMVRExporter Instance;
+	return Instance.ExportInternal(DMXLibrary, FilePathAndName);
+}
+
+bool FDMXMVRExporter::ExportInternal(UDMXLibrary* DMXLibrary, const FString& FilePathAndName)
+{
+	if (!ensureAlwaysMsgf(DMXLibrary, TEXT("Trying to export DMX Library '%s' as MVR file, but the DMX Library is invalid."), *DMXLibrary->GetName()))
+	{
+		return false;
+	}
+
+	DMXLibrary->UpdateGeneralSceneDescription();
+	const UDMXMVRGeneralSceneDescription* GeneralSceneDescription = DMXLibrary->GetLazyGeneralSceneDescription();
+	if (!ensureAlwaysMsgf(GeneralSceneDescription, TEXT("Trying to export DMX Library '%s' as MVR file, but its General Scene Description is invalid."), *DMXLibrary->GetName()))
+	{
+		return false;
+	}
+
+	const TSharedRef<FDMXZipper> Zip = MakeShared<FDMXZipper>();
+	if (!ZipGeneralSceneDescription(Zip, GeneralSceneDescription))
+	{
+		return false;
+	}
+
+	ZipGDTFs(Zip, DMXLibrary);
+	ZipThirdPartyData(Zip, GeneralSceneDescription);
+
+	return Zip->SaveToFile(FilePathAndName);
+}
+
+bool FDMXMVRExporter::ZipGeneralSceneDescription(const TSharedRef<FDMXZipper>& Zip, const UDMXMVRGeneralSceneDescription* GeneralSceneDescription)
+{
+	TSharedPtr<FXmlFile> XmlFile = GeneralSceneDescription->CreateXmlFile();
+	
+	// Try to merge the source General Scene Description Xml
+	const TSharedPtr<FXmlFile> SourceXmlFile = CreateSourceGeneralSceneDescriptionXmlFile(GeneralSceneDescription);
+	if (SourceXmlFile.IsValid())
+	{
+		// Merge the General Scene Description with source data's xml (source) to retain 3rd party data
+		XmlFile = FDMXXmlMergeUtility::Merge(GeneralSceneDescription, SourceXmlFile.ToSharedRef());
+	}
+	
+	// Create a temp GeneralSceneDescription.xml file
+	const FString TempPath = FPaths::ConvertRelativePathToFull(FPaths::EngineSavedDir() / TEXT("DMX_Temp"));
+	constexpr TCHAR GeneralSceneDescriptionFileName[] = TEXT("GeneralSceneDescription.xml");
+	const FString TempGeneralSceneDescriptionFilePathAndName = TempPath / GeneralSceneDescriptionFileName;
+	if (!XmlFile->Save(TempGeneralSceneDescriptionFilePathAndName))
+	{
+		UE_LOG(LogDMXEditor, Error, TEXT("Failed to save General Scene Description. See previous errors for details."));
+		return false;
+	}
+
+	TArray64<uint8> GeneralSceneDescriptionData;
+	FFileHelper::LoadFileToArray(GeneralSceneDescriptionData, *TempGeneralSceneDescriptionFilePathAndName);
+	Zip->AddFile(GeneralSceneDescriptionFileName, GeneralSceneDescriptionData);
+
+	// Delete temp GeneralSceneDescription.xml file
+	IFileManager& FileManager = IFileManager::Get();
+	constexpr bool bRequireExists = true;
+	constexpr bool bEvenIfReadOnly = false;
+	constexpr bool bQuiet = true;
+	FileManager.Delete(*TempGeneralSceneDescriptionFilePathAndName, bRequireExists, bEvenIfReadOnly, bQuiet);
+
+	return true;
+}
+
+void FDMXMVRExporter::ZipGDTFs(const TSharedRef<FDMXZipper>& Zip, UDMXLibrary* DMXLibrary)
+{
+	check(DMXLibrary);
+
+	// Gather GDTFs
+	TArray<const UDMXGDTFAssetImportData*> GDTFAssetImportDataToExport;
+	const TArray<const UDMXEntityFixturePatch*> FixturePatches = DMXLibrary->GetEntitiesTypeCast<const UDMXEntityFixturePatch>();
+	for (const UDMXEntityFixturePatch* FixturePatch : FixturePatches)
+	{
+		const UDMXEntityFixtureType* FixtureType = FixturePatch->GetFixtureType();
+		if (!FixtureType)
+		{
+			continue;
+		}
+		else if (!FixtureType->GDTF)
+		{
+			UE_LOG(LogDMXEditor, Warning, TEXT("Cannot export Fixture Patch '%s' to MVR. Its Fixture Type '%s' has no GDTF set."), *FixturePatch->Name, *FixtureType->Name);
+			continue;
+		}
+
+		const UDMXGDTFAssetImportData* GDTFAssetImportData = FixtureType->GDTF->GetGDTFAssetImportData();
+		if (!GDTFAssetImportData || GDTFAssetImportData->GetRawSourceData().IsEmpty())
+		{
+			UE_LOG(LogDMXEditor, Warning, TEXT("Cannot export Fixture Patch '%s' to MVR. Its Fixture Type '%' is from a version prior to 5.1, and does not hold required GDTF data."), *FixturePatch->Name, *FixtureType->Name);
+			continue;
+		}
+
+		GDTFAssetImportDataToExport.AddUnique(GDTFAssetImportData);
+	}
+
+	// Zip GDTFs
+	TArray<FString> TempGDTFFilePathAndNames;
+	for (const UDMXGDTFAssetImportData* GDTFAssetImportData : GDTFAssetImportDataToExport)
+	{
+		if (GDTFAssetImportData->GetRawSourceData().Num() > 0)
+		{
+			const FString GDTFFilename = FPaths::GetCleanFilename(GDTFAssetImportData->GetSourceFilePathAndName());
+			Zip->AddFile(GDTFFilename, GDTFAssetImportData->GetRawSourceData());
+		}
+	}
+}
+
+void FDMXMVRExporter::ZipThirdPartyData(const TSharedRef<FDMXZipper>& Zip, const UDMXMVRGeneralSceneDescription* GeneralSceneDescription)
+{
+	const UDMXMVRAssetImportData* AssetImportData = GeneralSceneDescription->GetMVRAssetImportData();
+	if (!AssetImportData)
+	{
+		return;
+	}
+
+	if (AssetImportData && AssetImportData->GetRawSourceData().Num() > 0)
+	{
+		const TSharedRef<FDMXZipper> SourceZip = MakeShared<FDMXZipper>();
+		if (SourceZip->LoadFromData(AssetImportData->GetRawSourceData()))
+		{
+			for (const FString& SourceFileName : SourceZip->GetFiles())
+			{
+				// Don't add GeneralSceneDescription.xml and GDTFs
+				constexpr TCHAR GeneralSceneDescriptionFileName[] = TEXT("GeneralSceneDescription.xml");
+				constexpr TCHAR GDTFExtension[] = TEXT("gdtf");
+				if (SourceFileName.EndsWith(GeneralSceneDescriptionFileName) || FPaths::GetExtension(SourceFileName) == GDTFExtension)
+				{
+					continue;
+				}
+
+				TArray64<uint8> SourceFileData;
+				if (SourceZip->GetFileContent(SourceFileName, SourceFileData))
+				{
+					Zip->AddFile(SourceFileName, SourceFileData);
+				}
+			}
+		}
+	}
+}
+
+const TSharedPtr<FXmlFile> FDMXMVRExporter::CreateSourceGeneralSceneDescriptionXmlFile(const UDMXMVRGeneralSceneDescription* GeneralSceneDescription) const
+{
+	const UDMXMVRAssetImportData* AssetImportData = GeneralSceneDescription->GetMVRAssetImportData();
+	if (!AssetImportData)
+	{
+		return nullptr;
+	}
+
+	if (AssetImportData->GetRawSourceData().IsEmpty())
+	{
+		return nullptr;
+	}
+
+	const TSharedRef<FDMXZipper> DMXZipper = MakeShared<FDMXZipper>();
+	if (!DMXZipper->LoadFromData(AssetImportData->GetRawSourceData()))
+	{
+		return nullptr;
+	}
+
+	constexpr TCHAR GeneralSceneDescriptionFilename[] = TEXT("GeneralSceneDescription.xml");
+	FDMXZipper::FDMXScopedUnzipToTempFile UnzipTempFileScope(DMXZipper, GeneralSceneDescriptionFilename);
+
+	if (UnzipTempFileScope.TempFilePathAndName.IsEmpty())
+	{
+		return nullptr;
+	}
+
+	const TSharedRef<FXmlFile> XmlFile = MakeShared<FXmlFile>();
+	if (!XmlFile->LoadFile(UnzipTempFileScope.TempFilePathAndName))
+	{
+		return nullptr;
+	}
+
+	return XmlFile;
+}
diff --git a/Engine/Plugins/VirtualProduction/DMX/DMXEngine/Source/DMXEditor/Private/Exporters/DMXMVRExporter.h b/Engine/Plugins/VirtualProduction/DMX/DMXEngine/Source/DMXEditor/Private/Exporters/DMXMVRExporter.h
new file mode 100644
--- /dev/null
+++ b/Engine/Plugins/VirtualProduction/DMX/DMXEngine/Source/DMXEditor/Private/Exporters/DMXMVRExporter.h
@@ -0,0 +1,37 @@
+// Copyright Epic Games, Inc. All Rights Reserved.
+
+#pragma once
+
+#include "CoreMinimal.h"
+
+class FDMXZipper;
+class UDMXLibrary;
+class UDMXMVRAssetImportData;
+class UDMXMVRGeneralSceneDescription;
+
+class FXmlFile;
+
+
+/** Helper class to export a DMX Library as MVR file */
+class FDMXMVRExporter
+{
+public:
+	/** Exports the DMX Library as MVR File */
+	static bool Export(UDMXLibrary* DMXLibrary, const FString& FilePathAndName);
+
+private:
+	/** Exports the DMX Library as MVR File */
+	UE_NODISCARD bool ExportInternal(UDMXLibrary* DMXLibrary, const FString& FilePathAndName);
+
+	/** Zips the GeneralSceneDescription.xml */
+	UE_NODISCARD bool ZipGeneralSceneDescription(const TSharedRef<FDMXZipper>& Zip, const UDMXMVRGeneralSceneDescription* GeneralSceneDescription);
+
+	/** Zips the GDTFs from the Library */
+	void ZipGDTFs(const TSharedRef<FDMXZipper>& Zip, UDMXLibrary* DMXLibrary);
+
+	/** Zips 3rd Party Data from the MVR Asset Import Data */
+	void ZipThirdPartyData(const TSharedRef<FDMXZipper>& Zip, const UDMXMVRGeneralSceneDescription* GeneralSceneDescription);
+
+	/** Creates an General Scene Description Xml File from the MVR Source, as it was imported */
+	const TSharedPtr<FXmlFile> CreateSourceGeneralSceneDescriptionXmlFile(const UDMXMVRGeneralSceneDescription* GeneralSceneDescription) const;
+};
diff --git a/Engine/Plugins/VirtualProduction/DMX/DMXEngine/Source/DMXEditor/Private/Exporters/DMXMVRXmlMergeUtility.cpp b/Engine/Plugins/VirtualProduction/DMX/DMXEngine/Source/DMXEditor/Private/Exporters/DMXMVRXmlMergeUtility.cpp
new file mode 100644
--- /dev/null
+++ b/Engine/Plugins/VirtualProduction/DMX/DMXEngine/Source/DMXEditor/Private/Exporters/DMXMVRXmlMergeUtility.cpp
@@ -0,0 +1,319 @@
+// Copyright Epic Games, Inc. All Rights Reserved.
+
+#include "DMXMVRXmlMergeUtility.h"
+
+#include "DMXRuntimeLog.h"
+#include "DMXRuntimeUtils.h"
+#include "MVR/DMXMVRGeneralSceneDescription.h"
+#include "MVR/Types/DMXMVRFixtureNode.h"
+#include "MVR/Types/DMXMVRRootNode.h"
+#include "MVR/Types/DMXMVRUnrealEngineDataNode.h"
+
+#include "XmlFile.h"
+#include "XmlNode.h"
+#include "Algo/Find.h"
+
+
+TSharedPtr<FXmlFile> FDMXXmlMergeUtility::Merge(const UDMXMVRGeneralSceneDescription* GeneralSceneDescription, const TSharedRef<FXmlFile>& OtherXml)
+{
+	checkf(GeneralSceneDescription, TEXT("Invalid General Scene Description provided when trying to merge MVR data"));
+
+	FDMXXmlMergeUtility Instance;
+	return Instance.MergeInternal(GeneralSceneDescription, OtherXml);
+}
+
+
+TSharedPtr<FXmlFile> FDMXXmlMergeUtility::MergeInternal(const UDMXMVRGeneralSceneDescription* GeneralSceneDescription, const TSharedRef<FXmlFile>& OtherXml)
+{
+	TSharedPtr<FXmlFile> GeneralSceneDescriptionXml = GeneralSceneDescription->CreateXmlFile();
+	if (!GeneralSceneDescriptionXml.IsValid() || !OtherXml->GetRootNode())
+	{
+		return nullptr;
+	}
+
+	FXmlNode* RootNode = GeneralSceneDescriptionXml->GetRootNode();
+	if (!RootNode)
+	{
+		return GeneralSceneDescriptionXml;
+	}
+
+	// 1. Merge User Data
+	const TArray<FXmlNode*> AdditionalDataNodes = AcquireAdditionalData(OtherXml);
+	if (AdditionalDataNodes.Num() > 0)
+	{
+		constexpr TCHAR UserDataTag[] = TEXT("UserData");
+		FXmlNode* UserDataNode = RootNode->FindChildNode(UserDataTag);
+		if (!UserDataNode)
+		{
+			RootNode->AppendChildNode(UserDataTag);
+
+			UserDataNode = RootNode->GetChildrenNodes().Last();
+		}
+		check(UserDataNode);
+
+		for (FXmlNode* AdditionalDataNode : AdditionalDataNodes)
+		{
+			AddXmlNodeWithChildren(UserDataNode, AdditionalDataNode);
+		}
+	}
+	
+	// 2. Merge Scene Data
+	FXmlNode* SceneNode = GetSceneNode(GeneralSceneDescriptionXml.ToSharedRef());
+	if (!SceneNode)
+	{
+		constexpr TCHAR SceneTag[] = TEXT("Scene");
+		RootNode->AppendChildNode(SceneTag);
+
+		SceneNode = RootNode->GetChildrenNodes().Last();
+	}
+	check(SceneNode);
+
+	const FXmlNode* OtherSceneNode = GetSceneNode(OtherXml);
+	if (!OtherSceneNode)
+	{
+		return GeneralSceneDescriptionXml;
+	}
+
+	// Scene
+	if (SceneNode && OtherSceneNode)
+	{
+		MergeScenes(GeneralSceneDescription, SceneNode, OtherSceneNode);
+	}
+	
+	// AUXData
+	const FXmlNode* AdditionalAUXDataNode = AcquireAdditionalAUXData(OtherXml);
+	if (AdditionalAUXDataNode)
+	{
+		constexpr TCHAR AUXDataTag[] = TEXT("AUXData");
+		FXmlNode* AUXDataNode = SceneNode->FindChildNode(AUXDataTag);
+		if (!AUXDataNode)
+		{
+			RootNode->AppendChildNode(AUXDataTag);
+
+			AUXDataNode = RootNode->GetChildrenNodes().Last();
+		}
+		check(AUXDataNode);
+
+		for (const FXmlNode* AdditionalAUXDataChild : AUXDataNode->GetChildrenNodes())
+		{
+			AddXmlNodeWithChildren(AUXDataNode, AdditionalAUXDataChild);
+		}
+	}
+
+	return GeneralSceneDescriptionXml;
+}
+
+TArray<FXmlNode*> FDMXXmlMergeUtility::AcquireAdditionalData(const TSharedRef<FXmlFile>& Xml)
+{
+	TArray<FXmlNode*> Result;
+	const FXmlNode* RootNode = Xml->GetRootNode();
+	if (!RootNode)
+	{
+		return Result;
+	}
+
+	constexpr TCHAR UserDataTag[] = TEXT("UserData");
+	const FXmlNode* UserDataNode = RootNode->FindChildNode(UserDataTag);
+	if (!UserDataNode)
+	{
+		return Result;
+	}
+
+	// All children of 'UserData' are 'Data' nodes
+	Result = UserDataNode->GetChildrenNodes();
+
+	// Remove previous Unreal Engine Data nodes
+	Result.RemoveAll([](const FXmlNode* Node)
+		{	
+			constexpr TCHAR ProviderAttributeName[] = TEXT("provider");
+			const FString Provider = Node->GetAttribute(ProviderAttributeName);
+
+			if (Provider == UDMXMVRUnrealEngineDataNode::ProviderNameUnrealEngine)
+			{
+				return true;
+			}
+
+			return false;
+		});
+
+	return Result;
+}
+
+const FXmlNode* FDMXXmlMergeUtility::AcquireAdditionalAUXData(const TSharedRef<FXmlFile>& Xml)
+{
+	const FXmlNode* SceneNode = GetSceneNode(Xml);
+	if (!SceneNode)
+	{
+		return nullptr;
+	}
+
+	constexpr TCHAR AUXDataTag[] = TEXT("AUXData");
+	return SceneNode->FindChildNode(AUXDataTag);
+}
+
+void FDMXXmlMergeUtility::MergeScenes(const UDMXMVRGeneralSceneDescription* GeneralSceneDescription, FXmlNode* InOutPrimarySceneNode, const FXmlNode* InSecondarySceneNode)
+{
+	checkf(InOutPrimarySceneNode->GetTag() == TEXT("scene") && InSecondarySceneNode->GetTag() == TEXT("scene"), TEXT("Trying to merge scene nodes, but nodes are '%s' and '%s' nodes."), *InOutPrimarySceneNode->GetTag(), *InSecondarySceneNode->GetTag());
+
+	// Don't merge Fixtures, so the merged result is what sourced from the DMX Library and its General Scene Description object only
+	constexpr TCHAR FixtureTagToIgnore[] = TEXT("Fixture");
+
+	// If the primary has no Layer, just copy what's in the secondary
+	constexpr TCHAR LayersTag[] = TEXT("Layers");
+	const FXmlNode* PrimaryLayersNode = InOutPrimarySceneNode->FindChildNode(LayersTag);
+	const FXmlNode* SecondaryLayersNode = InSecondarySceneNode->FindChildNode(LayersTag);
+	if (!PrimaryLayersNode && SecondaryLayersNode)
+	{
+		AddXmlNodeWithChildren(InOutPrimarySceneNode, SecondaryLayersNode, FixtureTagToIgnore);
+		return;
+	}
+
+	// Merge scenes
+	checkf(PrimaryLayersNode, TEXT("Primary scene node has no layers node, but cases where there is no layers node should already have been handled."));
+	for (const FXmlNode* SecondaryChildNode : InSecondarySceneNode->GetChildrenNodes())
+	{
+		AddXmlNodeWithChildren(InOutPrimarySceneNode, SecondaryChildNode, FixtureTagToIgnore);
+	}
+}
+
+TArray<FGuid> FDMXXmlMergeUtility::GetFixtureUUIDs(const UDMXMVRGeneralSceneDescription* GeneralSceneDescription)
+{
+	check(GeneralSceneDescription);
+
+	TArray<UDMXMVRFixtureNode*> FixtureNodes;
+	GeneralSceneDescription->GetFixtureNodes(FixtureNodes);
+
+	TArray<FGuid> Result;
+	for (UDMXMVRFixtureNode* FixtureNode : FixtureNodes)
+	{
+		Result.Add(FixtureNode->UUID);
+	}
+	return Result;
+}
+
+FXmlNode* FDMXXmlMergeUtility::GetSceneNode(const TSharedRef<FXmlFile>& Xml)
+{
+	FXmlNode* RootNode = Xml->GetRootNode();
+	if (!RootNode)
+	{
+		return nullptr;
+	}
+
+	constexpr TCHAR ScenenTag[] = TEXT("Scene");
+	return RootNode->FindChildNode(ScenenTag);
+}
+
+void FDMXXmlMergeUtility::AddXmlNodeWithChildren(FXmlNode* ParentTarget, const FXmlNode* ChildToAdd, const FString& IgnoredTag)
+{		
+	// Ignore Tags where requested
+	if (ChildToAdd->GetTag() == IgnoredTag)
+	{
+		return;
+	}
+
+	FXmlNode* ExistingNode = FindMatchingMVRNode(ParentTarget, ChildToAdd);
+	if (ExistingNode)
+	{
+		for (const FXmlNode* ChildNode : TArray<const FXmlNode*>(ChildToAdd->GetChildrenNodes()))
+		{
+			AddXmlNodeWithChildren(ExistingNode, ChildNode, IgnoredTag);
+		}
+	}
+	else
+	{
+		ParentTarget->AppendChildNode(ChildToAdd->GetTag(), ChildToAdd->GetContent(), ChildToAdd->GetAttributes());
+		FXmlNode* NewNode = ParentTarget->GetChildrenNodes().Last();
+		checkf(NewNode, TEXT("Expected a new node being added to Parent, but the last node is null"));
+
+		// Recursive for children
+		for (const FXmlNode* ChildNode : TArray<const FXmlNode*>(ChildToAdd->GetChildrenNodes()))
+		{
+			AddXmlNodeWithChildren(NewNode, ChildNode, IgnoredTag);
+		}
+	}
+}
+
+FXmlNode* FDMXXmlMergeUtility::FindMatchingMVRNode(const FXmlNode* ParentNode, const FXmlNode* OtherChild)
+{
+	const TArray<FXmlNode*> ChildNodes = ParentNode->GetChildrenNodes();
+	FXmlNode* const* MatchingNodePtr = nullptr;
+
+	// Identical UUIDs			
+	constexpr TCHAR UUIDAttributeName[] = TEXT("UUID");
+	const FString OtherChildUUIDString = OtherChild->GetAttribute(UUIDAttributeName);
+	FGuid OtherChildUUID;
+	if (FGuid::Parse(OtherChildUUIDString, OtherChildUUID))
+	{
+		MatchingNodePtr = Algo::FindByPredicate(ChildNodes, [UUIDAttributeName, &OtherChildUUID](const FXmlNode* Child)
+			{
+				const FString ChildUUIDString = Child->GetAttribute(UUIDAttributeName);
+
+				FGuid ChildUUID;
+				if (FGuid::Parse(ChildUUIDString, ChildUUID))
+				{
+					if (ChildUUID == OtherChildUUID)
+					{
+						return true;
+					}
+				}
+
+				return false;
+			});
+	}
+	
+	// Identical Provider
+	if (!MatchingNodePtr)
+	{
+		constexpr TCHAR ProviderTag[] = TEXT("provider");
+		const FString OtherProvider = OtherChild->GetAttribute(ProviderTag);
+
+		if (!OtherProvider.IsEmpty())
+		{
+			MatchingNodePtr = Algo::FindByPredicate(ChildNodes, [ProviderTag, &OtherProvider](const FXmlNode* Child)
+				{
+					const FString Provider = Child->GetAttribute(ProviderTag);
+					if (Provider == OtherProvider)
+					{
+						return true;
+					}
+
+					return false;
+				});
+		}
+	}
+
+	// Specific tags that are unique in their hierarchy and match by tag
+	if (!MatchingNodePtr)
+	{
+		MatchingNodePtr = Algo::FindByPredicate(ChildNodes, [OtherChild](const FXmlNode* Child)
+			{		
+				constexpr TCHAR ProviderTag[] = TEXT("provider");
+				constexpr TCHAR LayersTag[] = TEXT("Layers");
+				constexpr TCHAR LayerTag[] = TEXT("Layer");
+				constexpr TCHAR ChildListTag[] = TEXT("ChildList");
+				constexpr TCHAR GroupObjectTag[] = TEXT("GroupObject");
+
+				const FString& ChildTag = Child->GetTag();
+				if (ChildTag == ProviderTag ||
+					ChildTag == LayersTag ||
+					ChildTag == LayerTag ||
+					ChildTag == ChildListTag ||
+					ChildTag == GroupObjectTag)
+				{
+					if (Child->GetTag() == OtherChild->GetTag())
+					{
+						return true;
+					}
+				}
+
+				return false;
+			});
+	}
+
+	if (MatchingNodePtr)
+	{
+		return *MatchingNodePtr;
+	}
+
+	return nullptr;
+}
diff --git a/Engine/Plugins/VirtualProduction/DMX/DMXEngine/Source/DMXEditor/Private/Exporters/DMXMVRXmlMergeUtility.h b/Engine/Plugins/VirtualProduction/DMX/DMXEngine/Source/DMXEditor/Private/Exporters/DMXMVRXmlMergeUtility.h
new file mode 100644
--- /dev/null
+++ b/Engine/Plugins/VirtualProduction/DMX/DMXEngine/Source/DMXEditor/Private/Exporters/DMXMVRXmlMergeUtility.h
@@ -0,0 +1,45 @@
+// Copyright Epic Games, Inc. All Rights Reserved.
+
+#pragma once
+
+#include "CoreMinimal.h"
+
+class UDMXMVRUnrealEngineDataNode;
+class UDMXMVRGeneralSceneDescription;
+
+class FXmlFile;
+class FXmlNode;
+
+
+/** Utility to merge two General Scene Description Xmls, prefers primary for duplicate entries. Not a generic xml merger. */
+class FDMXXmlMergeUtility
+{
+public:
+	/** Merges the secondary General Scene Description Xml into the GeneralSceneDescription from engine. Returns an Xml File of the merged files or nullptr if data couldn't be merged. */
+	static TSharedPtr<FXmlFile> Merge(const UDMXMVRGeneralSceneDescription* GeneralSceneDescription, const TSharedRef<FXmlFile>& OtherXml);
+
+private:
+	/** Merges two General Scene Description Xmls */
+	TSharedPtr<FXmlFile> MergeInternal(const UDMXMVRGeneralSceneDescription* GeneralSceneDescription, const TSharedRef<FXmlFile>& OtherXml);
+
+	/** Acquires additional Data nodes, considering the UnrealEngineDataNode already being present */
+	TArray<FXmlNode*> AcquireAdditionalData(const TSharedRef<FXmlFile>& Xml);
+
+	/** Acquires additional AUXData nodes (child of Scene) */
+	const FXmlNode* AcquireAdditionalAUXData(const TSharedRef<FXmlFile>& Xml);
+
+	/** Merges the secondary scene into the primary */
+	void MergeScenes(const UDMXMVRGeneralSceneDescription* GeneralSceneDescription, FXmlNode* InOutPrimarySceneNode, const FXmlNode* InSecondarySceneNode);
+
+	/** Acquires additional AUXData nodes (child of Scene) */
+	TArray<FGuid> GetFixtureUUIDs(const UDMXMVRGeneralSceneDescription* GeneralSceneDescription);
+
+	/** Returns the scene node of the GeneralSceneDescription Xml File */
+	FXmlNode* GetSceneNode(const TSharedRef<FXmlFile>& Xml);
+
+	/** Adds source node to target with all its children. Optionaly ignores Nodes that specify a Tag */
+	void AddXmlNodeWithChildren(FXmlNode* ParentTarget, const FXmlNode* ChildToAdd, const FString& IgnoredTag = FString());
+
+	/** Tries ot find a Node in Parent (of Xml File A) that matches Child (of Xml File B) */
+	FXmlNode* FindMatchingMVRNode(const FXmlNode* ParentNode, const FXmlNode* Child);
+};
diff --git a/Engine/Plugins/VirtualProduction/DMX/DMXEngine/Source/DMXEditor/Private/Factories/DMXLibraryFromMVRFactory.cpp b/Engine/Plugins/VirtualProduction/DMX/DMXEngine/Source/DMXEditor/Private/Factories/DMXLibraryFromMVRFactory.cpp
--- a/Engine/Plugins/VirtualProduction/DMX/DMXEngine/Source/DMXEditor/Private/Factories/DMXLibraryFromMVRFactory.cpp
+++ b/Engine/Plugins/VirtualProduction/DMX/DMXEngine/Source/DMXEditor/Private/Factories/DMXLibraryFromMVRFactory.cpp
@@ -4,7 +4,9 @@
 
 #include "DMXEditorLog.h"
 #include "DMXEditorModule.h"
+#include "DMXEditorSettings.h"
 #include "DMXEditorUtils.h"
+#include "DMXZipper.h"
 #include "Factories/DMXGDTFFactory.h"
 #include "Library/DMXEntityFixturePatch.h"
 #include "Library/DMXEntityFixtureType.h"
@@ -13,7 +15,7 @@
 #include "Library/DMXLibrary.h"
 #include "MVR/DMXMVRAssetImportData.h"
 #include "MVR/DMXMVRGeneralSceneDescription.h"
-#include "MVR/DMXMVRUnzip.h"
+#include "MVR/Types/DMXMVRFixtureNode.h"
 
 #include "AssetRegistry/AssetRegistryModule.h"
 #include "AssetToolsModule.h"
@@ -43,6 +45,15 @@ UObject* UDMXLibraryFromMVRFactory::FactoryCreateFile(UClass* InClass, UObject*
 	bOutOperationCanceled = false;
 	CurrentFilename = InFilename;
 
+	if (!FPaths::FileExists(InFilename))
+	{
+		UE_LOG(LogDMXEditor, Error, TEXT("Failed to create DMX Library for MVR '%s'. Cannot find file."), *InFilename);
+		return nullptr;
+	}
+	UDMXEditorSettings* DMXEditorSettings = GetMutableDefault<UDMXEditorSettings>();
+	check(DMXEditorSettings);
+	DMXEditorSettings->LastMVRImportPath = FPaths::GetPath(InFilename);
+
 	UDMXLibrary* NewDMXLibrary = CreateDMXLibraryAsset(Parent, Flags, InFilename);
 	if (!NewDMXLibrary)
 	{
@@ -50,15 +61,19 @@ UObject* UDMXLibraryFromMVRFactory::FactoryCreateFile(UClass* InClass, UObject*
 		return nullptr;
 	}
 
-	const TSharedPtr<FDMXMVRUnzip> MVRUnzip = FDMXMVRUnzip::CreateFromFile(InFilename);
-	if (!MVRUnzip.IsValid())
+	const TSharedRef<FDMXZipper> Zip = MakeShared<FDMXZipper>();
+	if (!Zip->LoadFromFile(InFilename))
 	{
 		UE_LOG(LogDMXEditor, Error, TEXT("Cannot read MVR '%s'. File is not a valid MVR."), *InFilename);
 		return nullptr;
 	}
 
 	TArray64<uint8> XMLData;
-	MVRUnzip->GetFileContent(TEXT("GeneralSceneDescription.xml"), XMLData);
+	if (!Zip->GetFileContent(TEXT("GeneralSceneDescription.xml"), XMLData))
+	{
+		UE_LOG(LogDMXEditor, Error, TEXT("Cannot read General Scene Description from MVR '%s'. File is not a valid MVR."), *InFilename);
+		return nullptr;
+	}
 
 	// MVR implicitly adpots UTF-8 encoding of Xml Files by adopting the GDT standard (DIN-15800).
 	// Content is NOT null-terminated; we need to specify lengths here.
@@ -75,7 +90,7 @@ UObject* UDMXLibraryFromMVRFactory::FactoryCreateFile(UClass* InClass, UObject*
 	UDMXMVRAssetImportData* MVRAssetImportData = GeneralSceneDescription->GetMVRAssetImportData();
 	MVRAssetImportData->SetSourceFile(InFilename);
 
-	TArray<UDMXImportGDTF*> GDTFs = CreateGDTFAssets(Parent, Flags, MVRUnzip.ToSharedRef(), *GeneralSceneDescription);
+	TArray<UDMXImportGDTF*> GDTFs = CreateGDTFAssets(Parent, Flags, Zip, *GeneralSceneDescription);
 	InitDMXLibrary(NewDMXLibrary, GDTFs, GeneralSceneDescription);
 
 	return NewDMXLibrary;
@@ -126,7 +141,7 @@ UDMXLibrary* UDMXLibraryFromMVRFactory::CreateDMXLibraryAsset(UObject* Parent, E
 	return NewDMXLibrary;
 }
 
-TArray<UDMXImportGDTF*> UDMXLibraryFromMVRFactory::CreateGDTFAssets(UObject* Parent, EObjectFlags Flags, const TSharedRef<FDMXMVRUnzip>& MVRUnzip, const UDMXMVRGeneralSceneDescription& GeneralSceneDescription)
+TArray<UDMXImportGDTF*> UDMXLibraryFromMVRFactory::CreateGDTFAssets(UObject* Parent, EObjectFlags Flags, const TSharedRef<FDMXZipper>& Zip, const UDMXMVRGeneralSceneDescription& GeneralSceneDescription)
 {
 	const FString Path = FPaths::GetPath(DMXLibraryPackageName) + TEXT("/GDTFs");
 
@@ -168,12 +183,14 @@ TArray<UDMXImportGDTF*> UDMXLibraryFromMVRFactory::CreateGDTFAssets(UObject* Par
 	}
 
 	// Import GDTF Assets that aren't yet imported in the reimport procedure
-	for (const FDMXMVRFixture& MVRFixture : GeneralSceneDescription.GetMVRFixtures())
+	TArray<UDMXMVRFixtureNode*> FixtureNodes;
+	GeneralSceneDescription.GetFixtureNodes(FixtureNodes);
+	for (const UDMXMVRFixtureNode* FixtureNode : FixtureNodes)
 	{
 		// Don't import the same GDTF twice
-		if (!ImportedGDTFNames.Contains(MVRFixture.GDTFSpec))
+		if (!ImportedGDTFNames.Contains(FixtureNode->GDTFSpec))
 		{
-			const FDMXMVRUnzip::FDMXScopedUnzipToTempFile ScopedUnzipGDTF(MVRUnzip, MVRFixture.GDTFSpec);
+			const FDMXZipper::FDMXScopedUnzipToTempFile ScopedUnzipGDTF(Zip, FixtureNode->GDTFSpec);
 			if (!ScopedUnzipGDTF.TempFilePathAndName.IsEmpty())
 			{
 				constexpr bool bRemovePathFromDesiredName = true;
@@ -190,12 +207,12 @@ TArray<UDMXImportGDTF*> UDMXLibraryFromMVRFactory::CreateGDTFAssets(UObject* Par
 				Package->FullyLoad();
 
 				bool bCanceled;
-				UObject* NewGDTFObject = GDTFFactory->FactoryCreateFile(UDMXImportGDTF::StaticClass(), Package, FName(*MVRFixture.GDTFSpec), Flags | RF_Public, ScopedUnzipGDTF.TempFilePathAndName, nullptr, GWarn, bCanceled);
+				UObject* NewGDTFObject = GDTFFactory->FactoryCreateFile(UDMXImportGDTF::StaticClass(), Package, FName(*FixtureNode->GDTFSpec), Flags | RF_Public, ScopedUnzipGDTF.TempFilePathAndName, nullptr, GWarn, bCanceled);
 				
 				if (UDMXImportGDTF* NewGDTF = Cast<UDMXImportGDTF>(NewGDTFObject))
 				{
 					ImportedGDTFs.Add(NewGDTF);
-					ImportedGDTFNames.Add(MVRFixture.GDTFSpec);
+					ImportedGDTFNames.Add(FixtureNode->GDTFSpec);
 
 					FAssetRegistryModule::AssetCreated(NewGDTF);
 					Package->MarkPackageDirty();
@@ -243,22 +260,24 @@ void UDMXLibraryFromMVRFactory::InitDMXLibrary(UDMXLibrary* DMXLibrary, const TA
 	}
 
 	// Create Fixture Patches for the MVR Fixtures
-	for (const FDMXMVRFixture& MVRFixture : GeneralSceneDescription->GetMVRFixtures())
+	TArray< UDMXMVRFixtureNode*> FixtureNodes;
+	GeneralSceneDescription->GetFixtureNodes(FixtureNodes);
+	for (const UDMXMVRFixtureNode* FixtureNode : FixtureNodes)
 	{
-		if (!GDTFSpecToFixtureTypeMap.Contains(MVRFixture.GDTFSpec))
+		if (!GDTFSpecToFixtureTypeMap.Contains(FixtureNode->GDTFSpec))
 		{
 			continue;
 		}
 
-		UDMXEntityFixtureType* FixtureType = GDTFSpecToFixtureTypeMap[MVRFixture.GDTFSpec];
-		UDMXEntityFixturePatch* const* FixturePatchPtr = DMXLibrary->GetEntitiesTypeCast<UDMXEntityFixturePatch>().FindByPredicate([MVRFixture, FixtureType](UDMXEntityFixturePatch* FixturePatch)
+		UDMXEntityFixtureType* FixtureType = GDTFSpecToFixtureTypeMap[FixtureNode->GDTFSpec];
+		UDMXEntityFixturePatch* const* FixturePatchPtr = DMXLibrary->GetEntitiesTypeCast<UDMXEntityFixturePatch>().FindByPredicate([FixtureNode, FixtureType](UDMXEntityFixturePatch* FixturePatch)
 			{
 				const FDMXFixtureMode* ActiveModePtr = FixturePatch->GetActiveMode();
 				return
 					FixturePatch->GetFixtureType() == FixtureType &&
-					FixturePatch->GetUniverseID() == MVRFixture.GetUniverseID() &&
-					FixturePatch->GetStartingChannel() == MVRFixture.GetStartingChannel() &&
-					ActiveModePtr && ActiveModePtr->ModeName == MVRFixture.GDTFMode;
+					FixturePatch->GetUniverseID() == FixtureNode->GetUniverseID() &&
+					FixturePatch->GetStartingChannel() == FixtureNode->GetStartingChannel() &&
+					ActiveModePtr && ActiveModePtr->ModeName == FixtureNode->GDTFMode;
 			});
 
 		if (!FixtureTypeToColorMap.Contains(FixtureType))
@@ -276,9 +295,9 @@ void UDMXLibraryFromMVRFactory::InitDMXLibrary(UDMXLibrary* DMXLibrary, const TA
 
 		if (!FixturePatchPtr)
 		{
-			int32 ActiveModeIndex = FixtureType->Modes.IndexOfByPredicate([MVRFixture](const FDMXFixtureMode& Mode)
+			int32 ActiveModeIndex = FixtureType->Modes.IndexOfByPredicate([FixtureNode](const FDMXFixtureMode& Mode)
 				{
-					return Mode.ModeName == MVRFixture.GDTFMode;
+					return Mode.ModeName == FixtureNode->GDTFMode;
 				});
 
 			if (ActiveModeIndex != INDEX_NONE)
@@ -286,17 +305,17 @@ void UDMXLibraryFromMVRFactory::InitDMXLibrary(UDMXLibrary* DMXLibrary, const TA
 				FDMXEntityFixturePatchConstructionParams FixturePatchConstructionParams;
 				FixturePatchConstructionParams.ActiveMode = ActiveModeIndex;
 				FixturePatchConstructionParams.FixtureTypeRef = FDMXEntityFixtureTypeRef(FixtureType);
-				FixturePatchConstructionParams.UniverseID = MVRFixture.GetUniverseID();
-				FixturePatchConstructionParams.StartingAddress = MVRFixture.GetStartingChannel();
-				FixturePatchConstructionParams.MVRFixtureUUID = MVRFixture.UUID;
+				FixturePatchConstructionParams.UniverseID = FixtureNode->GetUniverseID();
+				FixturePatchConstructionParams.StartingAddress = FixtureNode->GetStartingChannel();
+				FixturePatchConstructionParams.MVRFixtureUUID = FixtureNode->UUID;
 
-				UDMXEntityFixturePatch* FixturePatch = UDMXEntityFixturePatch::CreateFixturePatchInLibrary(FixturePatchConstructionParams, MVRFixture.Name);
+				UDMXEntityFixturePatch* FixturePatch = UDMXEntityFixturePatch::CreateFixturePatchInLibrary(FixturePatchConstructionParams, FixtureNode->Name);
 				FixturePatch->EditorColor = FixtureTypeToColorMap.FindChecked(FixtureType);
 				FixturePatchPtr = &FixturePatch;
 			}
 			else
 			{
-				UE_LOG(LogDMXEditor, Warning, TEXT("Skipped creating a Fixture Patch for '%s', as no valid Mode could be imported when importing %s."), *MVRFixture.Name, *MVRFixture.GDTFSpec)
+				UE_LOG(LogDMXEditor, Warning, TEXT("Skipped creating a Fixture Patch for '%s', GDTF '%s' has no valid Mode."), *FixtureNode->Name, *FixtureNode->GDTFSpec)
 			}
 		}
 	}
diff --git a/Engine/Plugins/VirtualProduction/DMX/DMXEngine/Source/DMXEditor/Private/Toolbars/DMXEditorToolbar.cpp b/Engine/Plugins/VirtualProduction/DMX/DMXEngine/Source/DMXEditor/Private/Toolbars/DMXEditorToolbar.cpp
--- a/Engine/Plugins/VirtualProduction/DMX/DMXEngine/Source/DMXEditor/Private/Toolbars/DMXEditorToolbar.cpp
+++ b/Engine/Plugins/VirtualProduction/DMX/DMXEngine/Source/DMXEditor/Private/Toolbars/DMXEditorToolbar.cpp
@@ -1,7 +1,11 @@
 // Copyright Epic Games, Inc. All Rights Reserved.
 
 #include "Toolbars/DMXEditorToolbar.h"
+
 #include "DMXEditor.h"
+#include "DMXEditorStyle.h"
+#include "Commands/DMXEditorCommands.h"
+
 #include "Framework/MultiBox/MultiBoxBuilder.h"
 
 #define LOCTEXT_NAMESPACE "KismetToolbar"
@@ -14,7 +18,7 @@ void FDMXEditorToolbar::AddCompileToolbar(TSharedPtr<FExtender> Extender)
 		"Asset",
 		EExtensionHook::After,
 		DMXEditorPtr->GetToolkitCommands(),
-		FToolBarExtensionDelegate::CreateSP(this, &FDMXEditorToolbar::FillCompileToolbar));
+		FToolBarExtensionDelegate::CreateSP(this, &FDMXEditorToolbar::FillDMXLibraryToolbar));
 }
 
 FSlateIcon FDMXEditorToolbar::GetStatusImage() const
@@ -27,10 +31,15 @@ FText FDMXEditorToolbar::GetStatusTooltip() const
 	return LOCTEXT("Default_Status", "Good to go");
 }
 
-void FDMXEditorToolbar::FillCompileToolbar(FToolBarBuilder& ToolbarBuilder)
+void FDMXEditorToolbar::FillDMXLibraryToolbar(FToolBarBuilder& ToolbarBuilder)
 {
-	// TODO. Implement custom toolbar
-	ToolbarBuilder.BeginSection("CompileToolbar");
+	ToolbarBuilder.BeginSection("DMXLibraryToolbar");
+		ToolbarBuilder.AddToolBarButton(
+		FDMXEditorCommands::Get().ExportDMXLibrary,
+		NAME_None,
+		TAttribute<FText>(),
+		TAttribute<FText>(),
+		FSlateIcon(FDMXEditorStyle::Get().GetStyleSetName(), "Icons.DMXLibraryToolbar.Export"));
 	ToolbarBuilder.EndSection();
 }
 
diff --git a/Engine/Plugins/VirtualProduction/DMX/DMXEngine/Source/DMXEditor/Private/Toolbars/DMXEditorToolbar.h b/Engine/Plugins/VirtualProduction/DMX/DMXEngine/Source/DMXEditor/Private/Toolbars/DMXEditorToolbar.h
--- a/Engine/Plugins/VirtualProduction/DMX/DMXEngine/Source/DMXEditor/Private/Toolbars/DMXEditorToolbar.h
+++ b/Engine/Plugins/VirtualProduction/DMX/DMXEngine/Source/DMXEditor/Private/Toolbars/DMXEditorToolbar.h
@@ -28,7 +28,7 @@ public:
 	FText GetStatusTooltip() const;
 
 private:
-	void FillCompileToolbar(FToolBarBuilder& ToolbarBuilder);
+	void FillDMXLibraryToolbar(FToolBarBuilder& ToolbarBuilder);
 
 protected:
 	/** Pointer back to the blueprint editor tool that owns us */
diff --git a/Engine/Plugins/VirtualProduction/DMX/DMXEngine/Source/DMXEditor/Private/Widgets/FixturePatch/DMXMVRFixtureListItem.cpp b/Engine/Plugins/VirtualProduction/DMX/DMXEngine/Source/DMXEditor/Private/Widgets/FixturePatch/DMXMVRFixtureListItem.cpp
--- a/Engine/Plugins/VirtualProduction/DMX/DMXEngine/Source/DMXEditor/Private/Widgets/FixturePatch/DMXMVRFixtureListItem.cpp
+++ b/Engine/Plugins/VirtualProduction/DMX/DMXEngine/Source/DMXEditor/Private/Widgets/FixturePatch/DMXMVRFixtureListItem.cpp
@@ -7,6 +7,7 @@
 #include "Library/DMXEntityFixturePatch.h"
 #include "Library/DMXEntityReference.h"
 #include "Library/DMXLibrary.h"
+#include "MVR/Types/DMXMVRFixtureNode.h"
 
 #include "ScopedTransaction.h"
 
@@ -41,11 +42,7 @@ FDMXMVRFixtureListItem::FDMXMVRFixtureListItem(TWeakPtr<FDMXEditor> InDMXEditor,
 		WeakFixturePatch = *FixturePatchPtr;
 	}
 	
-	FDMXMVRFixture* MVRFixturePtr = FindMVRFixture();
-	if (MVRFixturePtr)
-	{
-		CachedMVRFixture = *MVRFixturePtr;
-	}
+	MVRFixtureNode = FindMVRFixture();
 }
 
 const FGuid& FDMXMVRFixtureListItem::GetMVRUUID() const
@@ -98,47 +95,39 @@ void FDMXMVRFixtureListItem::SetFixturePatchName(const FString& InDesiredName, F
 
 int32 FDMXMVRFixtureListItem::GetUnitNumber() const
 {
-	return CachedMVRFixture.UnitNumber;
+	return MVRFixtureNode->UnitNumber;
 }
 
 void FDMXMVRFixtureListItem::SetUnitNumber(int32 InUnitNumber)
 {
-	FDMXMVRFixture* MVRFixturePtr = FindMVRFixture();
 	UDMXLibrary* DMXLibrary = GetDMXLibrary();
-	if (MVRFixturePtr && DMXLibrary)
+	if (MVRFixtureNode && DMXLibrary)
 	{
 		const FScopedTransaction SetMVRFixtureUnitNumberTransaction(LOCTEXT("SetMVRFixtureUnitNumberTransaction", "Set MVR Fixture UnitNumber"));
 		DMXLibrary->PreEditChange(UDMXLibrary::StaticClass()->FindPropertyByName(UDMXLibrary::GetGeneralSceneDescriptionPropertyName()));
 
-		MVRFixturePtr->UnitNumber = InUnitNumber;
+		MVRFixtureNode->UnitNumber = InUnitNumber;
 
 		DMXLibrary->PostEditChange();
-
-		// Update Cache for Getter
-		CachedMVRFixture = *MVRFixturePtr;
 	}
 }	
 
 const FString& FDMXMVRFixtureListItem::GetMVRFixtureName() const
 {
-	return CachedMVRFixture.Name;
+	return MVRFixtureNode->Name;
 }
 
 bool FDMXMVRFixtureListItem::SetMVRFixtureName(const FString& Name)
 {
-	FDMXMVRFixture* MVRFixturePtr = FindMVRFixture();
 	UDMXLibrary* DMXLibrary = GetDMXLibrary();
-	if (!Name.IsEmpty() && MVRFixturePtr && DMXLibrary)
+	if (!Name.IsEmpty() && MVRFixtureNode && DMXLibrary)
 	{
 		const FScopedTransaction SetMVRFixtureNameTransaction(LOCTEXT("SetMVRFixtureNameTransaction", "Set MVR Fixture Name"));
 		DMXLibrary->PreEditChange(UDMXLibrary::StaticClass()->FindPropertyByName(UDMXLibrary::GetGeneralSceneDescriptionPropertyName()));
 
-		MVRFixturePtr->Name = Name;
+		MVRFixtureNode->Name = Name;
 
 		DMXLibrary->PostEditChange();
-
-		// Update Cache for Getter
-		CachedMVRFixture = *MVRFixturePtr;
 	}
 
 	return false;
@@ -158,9 +147,8 @@ void FDMXMVRFixtureListItem::SetFixtureType(UDMXEntityFixtureType* FixtureType)
 {
 	UDMXLibrary* DMXLibrary = GetDMXLibrary();
 	UDMXEntityFixturePatch* FixturePatch = WeakFixturePatch.Get();
-	FDMXMVRFixture* MVRFixturePtr = FindMVRFixture();
 
-	if (DMXLibrary && FixturePatch && MVRFixturePtr)
+	if (DMXLibrary && FixturePatch)
 	{
 		if (FixturePatch->GetFixtureType() == FixtureType)
 		{
@@ -190,9 +178,8 @@ void FDMXMVRFixtureListItem::SetModeIndex(int32 ModeIndex)
 {
 	UDMXLibrary* DMXLibrary = GetDMXLibrary();
 	UDMXEntityFixturePatch* FixturePatch = WeakFixturePatch.Get();
-	FDMXMVRFixture* MVRFixturePtr = FindMVRFixture();
 
-	if (DMXLibrary && FixturePatch && MVRFixturePtr)
+	if (DMXLibrary && FixturePatch)
 	{
 		if (ModeIndex == FixturePatch->GetActiveModeIndex())
 		{
@@ -233,9 +220,8 @@ void FDMXMVRFixtureListItem::SetAddresses(int32 Universe, int32 Address)
 {
 	UDMXLibrary* DMXLibrary = GetDMXLibrary();
 	UDMXEntityFixturePatch* FixturePatch = WeakFixturePatch.Get();
-	FDMXMVRFixture* MVRFixturePtr = FindMVRFixture();
 
-	if (DMXLibrary && FixturePatch && MVRFixturePtr)
+	if (DMXLibrary && FixturePatch)
 	{
 		if (FixturePatch->GetUniverseID() == Universe &&
 			FixturePatch->GetStartingChannel() == Address)
@@ -262,7 +248,6 @@ void FDMXMVRFixtureListItem::SetAddresses(int32 Universe, int32 Address)
 
 		FixturePatch->PostEditChange();
 
-
 		// Select the universe in Fixture Patch Shared Data
 		FixturePatchSharedData->SelectUniverse(Universe);
 	}
@@ -358,6 +343,11 @@ UDMXLibrary* FDMXMVRFixtureListItem::GetDMXLibrary() const
 	return nullptr;
 }
 
+void FDMXMVRFixtureListItem::AddReferencedObjects(FReferenceCollector& Collector)
+{
+	Collector.AddReferencedObject(MVRFixtureNode);
+}
+
 void FDMXMVRFixtureListItem::DuplicateFixturePatchesInternal(TWeakPtr<FDMXEditor> WeakDMXEditor, const TArray<UDMXEntityFixturePatch*>& FixturePatchesToDuplicate, const FText& TransactionText)
 {
 	const TSharedPtr<FDMXEditor> DMXEditor = WeakDMXEditor.Pin();
@@ -437,19 +427,6 @@ void FDMXMVRFixtureListItem::DuplicateFixturePatchesInternal(TWeakPtr<FDMXEditor
 		UDMXEntityFixturePatch* NewFixturePatch = UDMXEntityFixturePatch::CreateFixturePatchInLibrary(ConstructionParams, FixturePatchToDuplicate->Name);
 		NewFixturePatches.Add(NewFixturePatch);
 
-		// Duplicate the MVR Fixture if it exists
-		UDMXMVRGeneralSceneDescription* GeneralSceneDescription = DMXLibrary->GetLazyGeneralSceneDescription();
-		if (ensureAlwaysMsgf(GeneralSceneDescription, TEXT("Trying to duplicate MVR Fixture in the General Scene Description of '%s', but the General Scene Desription is not valid."), *DMXLibrary->GetName()))
-		{
-			FDMXMVRFixture* SourceMVRFixturePtr = GeneralSceneDescription->FindMVRFixture(FixturePatchToDuplicate->GetMVRFixtureUUID());
-			if (SourceMVRFixturePtr)
-			{
-				FDMXMVRFixture MVRFixture = *SourceMVRFixturePtr;
-				MVRFixture.UUID = NewFixturePatch->GetMVRFixtureUUID();
-				DMXLibrary->GetLazyGeneralSceneDescription()->AddMVRFixture(MVRFixture);
-			}
-		}
-
 		Address += ChannelSpan;
 	}
 
@@ -461,7 +438,7 @@ void FDMXMVRFixtureListItem::DuplicateFixturePatchesInternal(TWeakPtr<FDMXEditor
 	}
 }
 
-FDMXMVRFixture* FDMXMVRFixtureListItem::FindMVRFixture() const
+UDMXMVRFixtureNode* FDMXMVRFixtureListItem::FindMVRFixture() const
 {
 	UDMXLibrary* DMXLibrary = GetDMXLibrary();
 
@@ -477,7 +454,7 @@ FDMXMVRFixture* FDMXMVRFixtureListItem::FindMVRFixture() const
 		return nullptr;
 	}
 
-	return GeneralSceneDescription->FindMVRFixture(MVRFixtureUUID);
+	return GeneralSceneDescription->FindFixtureNode(MVRFixtureUUID);
 }
 
 #undef LOCTEXT_NAMESPACE
diff --git a/Engine/Plugins/VirtualProduction/DMX/DMXEngine/Source/DMXEditor/Private/Widgets/FixturePatch/DMXMVRFixtureListItem.h b/Engine/Plugins/VirtualProduction/DMX/DMXEngine/Source/DMXEditor/Private/Widgets/FixturePatch/DMXMVRFixtureListItem.h
--- a/Engine/Plugins/VirtualProduction/DMX/DMXEngine/Source/DMXEditor/Private/Widgets/FixturePatch/DMXMVRFixtureListItem.h
+++ b/Engine/Plugins/VirtualProduction/DMX/DMXEngine/Source/DMXEditor/Private/Widgets/FixturePatch/DMXMVRFixtureListItem.h
@@ -5,6 +5,7 @@
 #include "MVR/DMXMVRGeneralSceneDescription.h"
 
 #include "CoreMinimal.h"
+#include "UObject/GCObject.h"
 #include "UObject/WeakObjectPtr.h"
 
 class FDMXEditor;
@@ -12,11 +13,13 @@ class FDMXFixturePatchSharedData;
 class UDMXEntityFixturePatch;
 class UDMXEntityFixtureType;
 class UDMXLibrary;
+class UDMXMVRFixtureNode;
 
 
 /** An MVR Fixture as an Item in a List. A Primary Items is the first MVR Fixture in a Patch. Secondary Items are subsequent MVR Fixtures */
 class FDMXMVRFixtureListItem
-	: public TSharedFromThis<FDMXMVRFixtureListItem>
+	: public FGCObject
+	, public TSharedFromThis<FDMXMVRFixtureListItem>
 {
 public:	
 	/** Constructor */
@@ -91,18 +94,28 @@ public:
 	/** Error Status Text of the Item */
 	FText ErrorStatusText;
 
+protected:
+	// FGCObject interface
+	virtual void AddReferencedObjects(FReferenceCollector& Collector) override;
+	virtual FString GetReferencerName() const override
+	{
+		return TEXT("FDMXMVRFixtureListItem");
+	}
+	// End of FGCObject interface
+
+
 private:
 	/** Duplicates Fixture Patches */
 	static void DuplicateFixturePatchesInternal(TWeakPtr<FDMXEditor> WeakDMXEditor, const TArray<UDMXEntityFixturePatch*>& FixturePatchesToDuplicate, const FText& TransactionText);
 
-	/** Returns the live MVR Fixture in the library */
-	FDMXMVRFixture* FindMVRFixture() const;
+	/** Returns the MVR Fixture Node that corresponds to this item */
+	UDMXMVRFixtureNode* FindMVRFixture() const;
 
 	/** True if this is in an even group */
 	bool bIsEvenGroup = false;
 
-	/** Cached MVR Fixture to speed up getters. Should not be used by setters. Use RefreshCachedMVRFixture to update. */
-	FDMXMVRFixture CachedMVRFixture;
+	/** Cached MVR Fixture to speed up getters. */
+	UDMXMVRFixtureNode* MVRFixtureNode = nullptr;
 
 	/** The MVRFixtureUUID this Item handles */
 	FGuid MVRFixtureUUID;
diff --git a/Engine/Plugins/VirtualProduction/DMX/DMXEngine/Source/DMXEditor/Public/Commands/DMXEditorCommands.h b/Engine/Plugins/VirtualProduction/DMX/DMXEngine/Source/DMXEditor/Public/Commands/DMXEditorCommands.h
--- a/Engine/Plugins/VirtualProduction/DMX/DMXEngine/Source/DMXEditor/Public/Commands/DMXEditorCommands.h
+++ b/Engine/Plugins/VirtualProduction/DMX/DMXEngine/Source/DMXEditor/Public/Commands/DMXEditorCommands.h
@@ -19,7 +19,9 @@ public:
 	// Documentation related
 	TSharedPtr<FUICommandInfo> GoToDocumentation;
 
-	// Entity Editor related
+	// DMX Library related
+	TSharedPtr<FUICommandInfo> ExportDMXLibrary;
+
 	TSharedPtr<FUICommandInfo> AddNewEntityFixtureType;
 	TSharedPtr<FUICommandInfo> AddNewEntityFixturePatch;
 	
diff --git a/Engine/Plugins/VirtualProduction/DMX/DMXEngine/Source/DMXEditor/Public/DMXEditor.h b/Engine/Plugins/VirtualProduction/DMX/DMXEngine/Source/DMXEditor/Public/DMXEditor.h
--- a/Engine/Plugins/VirtualProduction/DMX/DMXEngine/Source/DMXEditor/Public/DMXEditor.h
+++ b/Engine/Plugins/VirtualProduction/DMX/DMXEngine/Source/DMXEditor/Public/DMXEditor.h
@@ -104,14 +104,14 @@ protected:
 	virtual void RegisterApplicationModes(UDMXLibrary* DMXLibrary, bool bShouldOpenInDefaultsMode, bool bNewlyCreated = false);
 
 private:
-	UDMXLibrary* GetEditableDMXLibrary() const;
+	/** Exports the DMX Library */
+	void ExportDMXLibrary() const;
 
 	//~ Generate Editor widgets for tabs
 	TSharedRef<SDMXLibraryEditorTab> CreateDMXLibraryEditorTab();
 	TSharedRef<SDMXFixtureTypeEditor> CreateFixtureTypeEditor();
 	TSharedRef<SDMXFixturePatchEditor> CreateFixturePatchEditor();
 
-private:
 	/** The toolbar builder class */
 	TSharedPtr<FDMXEditorToolbar> Toolbar;
 
diff --git a/Engine/Plugins/VirtualProduction/DMX/DMXEngine/Source/DMXEditor/Public/DMXEditorSettings.h b/Engine/Plugins/VirtualProduction/DMX/DMXEngine/Source/DMXEditor/Public/DMXEditorSettings.h
--- a/Engine/Plugins/VirtualProduction/DMX/DMXEngine/Source/DMXEditor/Public/DMXEditorSettings.h
+++ b/Engine/Plugins/VirtualProduction/DMX/DMXEngine/Source/DMXEditor/Public/DMXEditorSettings.h
@@ -126,11 +126,16 @@ class UDMXEditorSettings : public UObject
 	GENERATED_BODY()
 
 
-	// MVR Fixture List
+	// MVR
 public:
 	UPROPERTY(Config)
-	FDMXMVRFixtureListSettings MVRFixtureListSettings;
+	FString LastMVRImportPath;
+
+	UPROPERTY(Config)
+	FString LastMVRExportPath;
 
+	UPROPERTY(Config)
+	FDMXMVRFixtureListSettings MVRFixtureListSettings;
 
 	// Fixture Type Functions Editor
 public:
diff --git a/Engine/Plugins/VirtualProduction/DMX/DMXEngine/Source/DMXEditor/Public/Factories/DMXLibraryFromMVRFactory.h b/Engine/Plugins/VirtualProduction/DMX/DMXEngine/Source/DMXEditor/Public/Factories/DMXLibraryFromMVRFactory.h
--- a/Engine/Plugins/VirtualProduction/DMX/DMXEngine/Source/DMXEditor/Public/Factories/DMXLibraryFromMVRFactory.h
+++ b/Engine/Plugins/VirtualProduction/DMX/DMXEngine/Source/DMXEditor/Public/Factories/DMXLibraryFromMVRFactory.h
@@ -6,7 +6,7 @@
 
 #include "DMXLibraryFromMVRFactory.generated.h"
 
-class FDMXMVRUnzip;
+class FDMXZipper;
 class UDMXImportGDTF;
 class UDMXLibrary;
 class UDMXMVRGeneralSceneDescription;
@@ -37,7 +37,7 @@ private:
 	UDMXLibrary* CreateDMXLibraryAsset(UObject* Parent, EObjectFlags Flags, const FString& InFilename);
 
 	/** Creates GDTF assets from the MVR */
-	TArray<UDMXImportGDTF*> CreateGDTFAssets(UObject* Parent, EObjectFlags Flags, const TSharedRef<FDMXMVRUnzip>& MVRUnzip, const UDMXMVRGeneralSceneDescription& GeneralSceneDescription);
+	TArray<UDMXImportGDTF*> CreateGDTFAssets(UObject* Parent, EObjectFlags Flags, const TSharedRef<FDMXZipper>& Zip, const UDMXMVRGeneralSceneDescription& GeneralSceneDescription);
 
 	/** Initializes the DMX Library from the General Scene Description and GDTF assets */
 	void InitDMXLibrary(UDMXLibrary* DMXLibrary, const TArray<UDMXImportGDTF*>& GDTFAssets, UDMXMVRGeneralSceneDescription* GeneralSceneDescription) const;
diff --git a/Engine/Plugins/VirtualProduction/DMX/DMXEngine/Source/DMXRuntime/Private/DMXRuntimeUtils.cpp b/Engine/Plugins/VirtualProduction/DMX/DMXEngine/Source/DMXRuntime/Private/DMXRuntimeUtils.cpp
--- a/Engine/Plugins/VirtualProduction/DMX/DMXEngine/Source/DMXRuntime/Private/DMXRuntimeUtils.cpp
+++ b/Engine/Plugins/VirtualProduction/DMX/DMXEngine/Source/DMXRuntime/Private/DMXRuntimeUtils.cpp
@@ -8,6 +8,139 @@
 
 #define LOCTEXT_NAMESPACE "FDMXRuntimeUtils"
 
+FDMXOptionalTransform FDMXRuntimeUtils::ParseGDTFMatrix(const FString& String)
+{
+	auto GetVectorStringsFromMatrixStringLambda([](const FString& MatrixString) -> TArray<FString>
+		{
+			TArray<FString> Result;
+			int32 SubstringIndex = 0;
+			for (int32 SubStringNumber = 0; SubStringNumber < 4; SubStringNumber++)
+			{
+				SubstringIndex = MatrixString.Find(TEXT("{"), ESearchCase::IgnoreCase, ESearchDir::FromStart, SubstringIndex) + 1;
+				const int32 EndIndex = MatrixString.Find(TEXT("}"), ESearchCase::IgnoreCase, ESearchDir::FromStart, SubstringIndex);
+
+				const FString SubString = MatrixString.Mid(SubstringIndex, EndIndex - SubstringIndex);
+
+				if (!SubString.IsEmpty())
+				{
+					Result.Add(SubString);
+				}
+			}
+
+			return Result;
+		});
+
+	auto GetVectorFromVectorStringLambda([](const FString& VectorString) -> FVector
+		{
+			const int32 XIndex = 0;
+			const int32 YIndex = VectorString.Find(TEXT(","), ESearchCase::IgnoreCase, ESearchDir::FromStart, 0) + 1;
+			const int32 ZIndex = VectorString.Find(TEXT(","), ESearchCase::IgnoreCase, ESearchDir::FromStart, YIndex) + 1;
+
+			const FString XString = "X=" + VectorString.Mid(XIndex, YIndex - XIndex);
+			const FString YString = "Y=" + VectorString.Mid(YIndex, ZIndex - YIndex);
+			const FString ZString = "Z=" + VectorString.Mid(ZIndex, VectorString.Len() - ZIndex);
+
+			FVector Vector;
+			if (Vector.InitFromString(XString + YString + ZString))
+			{
+				return Vector;
+			}
+
+			return FVector::ZeroVector;
+		});
+
+	FDMXOptionalTransform OptionalResult;
+	const TArray<FString> SubStrings = GetVectorStringsFromMatrixStringLambda(String);
+	if (SubStrings.Num() != 4)
+	{
+		return OptionalResult;
+	}
+
+	// From millimeters to centimeters
+	const FVector U = GetVectorFromVectorStringLambda(SubStrings[0]) / 10.0;
+	const FVector V = GetVectorFromVectorStringLambda(SubStrings[1]) / 10.0;
+	const FVector W = GetVectorFromVectorStringLambda(SubStrings[2]) / 10.0;
+	const FVector O = GetVectorFromVectorStringLambda(SubStrings[3]) / 10.0;
+
+	FMatrix Matrix;
+	Matrix.M[0][0] = U.X; 
+	Matrix.M[0][1] = U.Y;
+	Matrix.M[0][2] = U.Z;
+	Matrix.M[0][3] = 0.0;
+
+	Matrix.M[1][0] = V.X;
+	Matrix.M[1][1] = V.Y;
+	Matrix.M[1][2] = V.Z;
+	Matrix.M[1][3] = 0.0;
+
+	Matrix.M[2][0] = W.X;
+	Matrix.M[2][1] = W.Y;
+	Matrix.M[2][2] = W.Z;
+	Matrix.M[2][3] = 0.0;
+
+	Matrix.M[3][0] = O.X;
+	Matrix.M[3][1] = -O.Y; // From GDTF's right hand to UE's left hand coordinate system
+	Matrix.M[3][2] = O.Z;
+	Matrix.M[3][3] = 1.0;
+
+	FTransform Result(Matrix);
+
+	// Mm to cm does not affect scale
+	Result.SetScale3D(Result.GetScale3D() * 10.0);
+
+	// GDTFs are facing down on the Z-Axis, but UE Actors are facing up 
+	const FQuat InvertUpRotationQuaternion = FQuat(Result.GetRotation().GetAxisY(), PI);
+	Result.SetRotation(Result.GetRotation() * InvertUpRotationQuaternion);
+
+	OptionalResult = Result;
+	return OptionalResult;
+}
+
+FString FDMXRuntimeUtils::ConvertTransformToGDTF4x3MatrixString(FTransform Transform)
+{
+	if (!ensureAlwaysMsgf(Transform.IsValid(), TEXT("Got invalid tranform when trying to generate a GDTF Matrix String from it.")))
+	{
+		Transform = FTransform::Identity;
+	}
+
+	// GDTFs are facing down on the Z-Axis, but UE Actors are facing up 
+	const FQuat InvertUpRotationQuaternion = FQuat(Transform.GetRotation().GetAxisY(), PI);
+	Transform.SetRotation(Transform.GetRotation() * InvertUpRotationQuaternion);
+
+	FMatrix Matrix = Transform.ToMatrixWithScale();
+	Matrix.M[3][1] = -Matrix.M[3][1]; // From UE's left hand to GDTF's right hand coordinate system
+
+	FString Result;
+	constexpr int32 FractionalDigits = 6; // As other soft- and hardwares
+
+	Result += TEXT("{");
+	Result += FString::SanitizeFloat(Matrix.M[0][0], FractionalDigits) + TEXT(",");
+	Result += FString::SanitizeFloat(Matrix.M[0][1], FractionalDigits) + TEXT(",");
+	Result += FString::SanitizeFloat(Matrix.M[0][2], FractionalDigits);
+	Result += TEXT("}");
+
+	Result += TEXT("{");
+	Result += FString::SanitizeFloat(Matrix.M[1][0], FractionalDigits) + TEXT(",");
+	Result += FString::SanitizeFloat(Matrix.M[1][1], FractionalDigits) + TEXT(",");
+	Result += FString::SanitizeFloat(Matrix.M[1][2], FractionalDigits);
+	Result += TEXT("}");
+
+	Result += TEXT("{");
+	Result += FString::SanitizeFloat(Matrix.M[2][0], FractionalDigits) + TEXT(",");
+	Result += FString::SanitizeFloat(Matrix.M[2][1], FractionalDigits) + TEXT(",");
+	Result += FString::SanitizeFloat(Matrix.M[2][2], FractionalDigits);
+	Result += TEXT("}");
+
+	// From centimeters to milimeters
+	Result += TEXT("{");
+	Result += FString::SanitizeFloat(Matrix.M[3][0] * 10.0, FractionalDigits) + TEXT(",");
+	Result += FString::SanitizeFloat(Matrix.M[3][1] * 10.0, FractionalDigits) + TEXT(",");
+	Result += FString::SanitizeFloat(Matrix.M[3][2] * 10.0, FractionalDigits);
+	Result += TEXT("}");
+
+	return Result;
+}
+
 FString FDMXRuntimeUtils::GenerateUniqueNameFromExisting(const TSet<FString>& InExistingNames, const FString& InBaseName)
 {
 	if (!InBaseName.IsEmpty() && !InExistingNames.Contains(InBaseName))
diff --git a/Engine/Plugins/VirtualProduction/DMX/DMXEngine/Source/DMXRuntime/Private/Library/DMXGDTFAssetImportData.cpp b/Engine/Plugins/VirtualProduction/DMX/DMXEngine/Source/DMXRuntime/Private/Library/DMXGDTFAssetImportData.cpp
--- a/Engine/Plugins/VirtualProduction/DMX/DMXEngine/Source/DMXRuntime/Private/Library/DMXGDTFAssetImportData.cpp
+++ b/Engine/Plugins/VirtualProduction/DMX/DMXEngine/Source/DMXRuntime/Private/Library/DMXGDTFAssetImportData.cpp
@@ -8,7 +8,8 @@
 #if WITH_EDITOR
 FString UDMXGDTFAssetImportData::GetSourceFilePathAndName() const
 {
-	return GetFirstFilename();
+	// Not GetFirstFilename(), we want the name even if the file doesn't exist on disk locally.
+	return SourceData.SourceFiles.Num() > 0 ? SourceData.SourceFiles[0].RelativeFilename : FString();
 }
 #endif // WITH_EDITOR
 
@@ -24,7 +25,7 @@ void UDMXGDTFAssetImportData::SetSourceFile(const FString& FilePathAndName)
 		return;
 	}
 
-	RawZipFile.Reset();
-	FFileHelper::LoadFileToArray(RawZipFile, *FilePathAndName);
+	RawSourceData.ByteArray.Reset();
+	FFileHelper::LoadFileToArray(RawSourceData.ByteArray, *FilePathAndName);
 }
 #endif // WITH_EDITOR
diff --git a/Engine/Plugins/VirtualProduction/DMX/DMXEngine/Source/DMXRuntime/Private/Library/DMXLibrary.cpp b/Engine/Plugins/VirtualProduction/DMX/DMXEngine/Source/DMXRuntime/Private/Library/DMXLibrary.cpp
--- a/Engine/Plugins/VirtualProduction/DMX/DMXEngine/Source/DMXRuntime/Private/Library/DMXLibrary.cpp
+++ b/Engine/Plugins/VirtualProduction/DMX/DMXEngine/Source/DMXRuntime/Private/Library/DMXLibrary.cpp
@@ -569,15 +569,12 @@ void UDMXLibrary::UpdatePorts()
 
 void UDMXLibrary::SetMVRGeneralSceneDescription(UDMXMVRGeneralSceneDescription* NewGeneralSceneDescription)
 {
-	// Remove all Fixture Patches that don't exist in the General Scene Description
-	Entities.RemoveAll([NewGeneralSceneDescription](UDMXEntity* Entity)
-		{
-			if (UDMXEntityFixturePatch* FixturePatch = Cast<UDMXEntityFixturePatch>(Entity))
-			{
-				return NewGeneralSceneDescription->FindMVRFixture(FixturePatch->GetMVRFixtureUUID()) == nullptr;
-			}
-			return false;
-		});
+	// Remove all Fixture Patches
+	TArray<UDMXEntityFixturePatch*> FixturePatches = GetEntitiesTypeCast<UDMXEntityFixturePatch>();
+	for (UDMXEntityFixturePatch* FixturePatch : FixturePatches)
+	{
+		UDMXEntityFixturePatch::RemoveFixturePatchFromLibrary(FDMXEntityFixturePatchRef(FixturePatch));
+	}
 
 	GeneralSceneDescription = NewGeneralSceneDescription;
 }
diff --git a/Engine/Plugins/VirtualProduction/DMX/DMXEngine/Source/DMXRuntime/Private/MVR/DMXMVRAssetImportData.cpp b/Engine/Plugins/VirtualProduction/DMX/DMXEngine/Source/DMXRuntime/Private/MVR/DMXMVRAssetImportData.cpp
--- a/Engine/Plugins/VirtualProduction/DMX/DMXEngine/Source/DMXRuntime/Private/MVR/DMXMVRAssetImportData.cpp
+++ b/Engine/Plugins/VirtualProduction/DMX/DMXEngine/Source/DMXRuntime/Private/MVR/DMXMVRAssetImportData.cpp
@@ -18,7 +18,7 @@ void UDMXMVRAssetImportData::SetSourceFile(const FString& FilePathAndName)
 		return;
 	}
 
-	RawZipFile.Reset();
-	FFileHelper::LoadFileToArray(RawZipFile, *FilePathAndName);
+	RawSourceData.ByteArray.Reset();
+	FFileHelper::LoadFileToArray(RawSourceData.ByteArray, *FilePathAndName);
 }
 #endif // WITH_EDITOR
diff --git a/Engine/Plugins/VirtualProduction/DMX/DMXEngine/Source/DMXRuntime/Private/MVR/DMXMVRFixture.cpp b/Engine/Plugins/VirtualProduction/DMX/DMXEngine/Source/DMXRuntime/Private/MVR/DMXMVRFixture.cpp
deleted file mode 100644
diff --git a/Engine/Plugins/VirtualProduction/DMX/DMXEngine/Source/DMXRuntime/Private/MVR/DMXMVRGeneralSceneDescription.cpp b/Engine/Plugins/VirtualProduction/DMX/DMXEngine/Source/DMXRuntime/Private/MVR/DMXMVRGeneralSceneDescription.cpp
--- a/Engine/Plugins/VirtualProduction/DMX/DMXEngine/Source/DMXRuntime/Private/MVR/DMXMVRGeneralSceneDescription.cpp
+++ b/Engine/Plugins/VirtualProduction/DMX/DMXEngine/Source/DMXRuntime/Private/MVR/DMXMVRGeneralSceneDescription.cpp
@@ -10,20 +10,48 @@
 #include "Library/DMXImportGDTF.h"
 #include "Library/DMXLibrary.h"
 #include "MVR/DMXMVRAssetImportData.h"
+#include "MVR/Types/DMXMVRChildListNode.h"
+#include "MVR/Types/DMXMVRFixtureNode.h"
+#include "MVR/Types/DMXMVRGroupObjectNode.h"
+#include "MVR/Types/DMXMVRLayerNode.h"
+#include "MVR/Types/DMXMVRLayersNode.h"
+#include "MVR/Types/DMXMVRParametricObjectNodeBase.h"
+#include "MVR/Types/DMXMVRRootNode.h"
+#include "MVR/Types/DMXMVRSceneNode.h"
 
 #include "XmlFile.h"
-#include "XmlNode.h"
 #include "EditorFramework/AssetImportData.h"
 #include "Misc/Paths.h"
 
 
 UDMXMVRGeneralSceneDescription::UDMXMVRGeneralSceneDescription()
 {
+	RootNode = CreateDefaultSubobject<UDMXMVRRootNode>("MVRRootNode");
+
 #if WITH_EDITORONLY_DATA
-	MVRAssetImportData = NewObject<UDMXMVRAssetImportData>(this, TEXT("MVRAssetImportData"), RF_Public);
+	MVRAssetImportData = CreateDefaultSubobject<UDMXMVRAssetImportData>(TEXT("MVRAssetImportData"));
 #endif
 }
 
+void UDMXMVRGeneralSceneDescription::GetFixtureNodes(TArray<UDMXMVRFixtureNode*>& OutFixtureNodes) const
+{
+	checkf(RootNode, TEXT("Unexpected: MVR General Scene Description Root Node is invalid."));
+
+	OutFixtureNodes.Reset();
+	RootNode->GetFixtureNodes(OutFixtureNodes);
+}
+
+UDMXMVRFixtureNode* UDMXMVRGeneralSceneDescription::FindFixtureNode(const FGuid& FixtureUUID) const
+{
+	checkf(RootNode, TEXT("Unexpected: MVR General Scene Description Root Node is invalid."));
+
+	if (UDMXMVRParametricObjectNodeBase** ObjectNodePtr = RootNode->FindParametricObjectNodeByUUID(FixtureUUID))
+	{
+		return Cast<UDMXMVRFixtureNode>(*ObjectNodePtr);
+	}
+	return nullptr;
+}
+
 #if WITH_EDITOR
 UDMXMVRGeneralSceneDescription* UDMXMVRGeneralSceneDescription::CreateFromXmlFile(TSharedRef<FXmlFile> GeneralSceneDescriptionXml, UObject* Outer, FName Name, EObjectFlags Flags)
 {
@@ -44,35 +72,9 @@ UDMXMVRGeneralSceneDescription* UDMXMVRGeneralSceneDescription::CreateFromDMXLib
 }
 #endif // WITH_EDITOR
 
-FDMXMVRFixture* UDMXMVRGeneralSceneDescription::FindMVRFixture(const FGuid& MVRFixtureUUID)
-{
-	return MVRFixtures.FindByPredicate([MVRFixtureUUID](const FDMXMVRFixture& MVRFixture)
-		{
-			return MVRFixture.UUID == MVRFixtureUUID;
-		});
-}
-
-void UDMXMVRGeneralSceneDescription::AddMVRFixture(FDMXMVRFixture& MVRFixture)
-{
-	if (ensureAlwaysMsgf(MVRFixture.UUID.IsValid(), TEXT("Trying to add MVR Fixture '%s' to General Scene Description, but the UUID is invalid."), *MVRFixture.Name))
-	{
-		return;
-	}
-
-	const bool bUUIDAlreadyContained = FindMVRFixture(MVRFixture.UUID) != nullptr;
-	if (ensureAlwaysMsgf(bUUIDAlreadyContained, TEXT("Trying to add MVR Fixture '%s' to General Scene Description, but its UUID is already containted in the General Scene Description."), *MVRFixture.Name))
-	{
-		return;
-	}
-
-	MVRFixtures.Add(MVRFixture);
-}
-
 #if WITH_EDITOR
 void UDMXMVRGeneralSceneDescription::WriteDMXLibraryToGeneralSceneDescription(const UDMXLibrary& DMXLibrary)
 {
-	TArray<FDMXMVRFixture> MVRFixturesInDMXLibrary;
-
 	TArray<UDMXEntityFixturePatch*> FixturePatches = DMXLibrary.GetEntitiesTypeCast<UDMXEntityFixturePatch>();
 	FixturePatches.Sort([](const UDMXEntityFixturePatch& FixturePatchA, const UDMXEntityFixturePatch& FixturePatchB)
 		{
@@ -83,17 +85,21 @@ void UDMXMVRGeneralSceneDescription::WriteDMXLibraryToGeneralSceneDescription(co
 			return bUniverseIsSmaller || (bUniverseIsEqual && bAddressIsSmaller);
 		});
 
+	// Remove Fixture Nodes no longer defined in the DMX Library
 	TArray<FGuid> MVRFixtureUUIDsInUse;
 	for (UDMXEntityFixturePatch* FixturePatch : FixturePatches)
 	{
 		MVRFixtureUUIDsInUse.Add(FixturePatch->GetMVRFixtureUUID());
 	}
-
-	// Remove entries no longer present in the DMX Library
-	MVRFixtures.RemoveAll([&MVRFixtureUUIDsInUse](const FDMXMVRFixture& MVRFixture)
+	TArray<UDMXMVRFixtureNode*> FixtureNodes;
+	RootNode->GetFixtureNodes(FixtureNodes);
+	for (UDMXMVRFixtureNode* FixtureNode : FixtureNodes)
+	{
+		if (!FixtureNode || !MVRFixtureUUIDsInUse.Contains(FixtureNode->UUID))
 		{
-			return !MVRFixtureUUIDsInUse.Contains(MVRFixture.UUID);
-		});
+			RootNode->RemoveParametricObjectNode(FixtureNode);
+		}
+	}
 
 	// Create or update MVR Fixtures for each Fixture Patch's MVR Fixture UUIDs
 	for (UDMXEntityFixturePatch* FixturePatch : FixturePatches)
@@ -106,128 +112,98 @@ void UDMXMVRGeneralSceneDescription::WriteDMXLibraryToGeneralSceneDescription(co
 }
 #endif // WITH_EDITOR
 
-void UDMXMVRGeneralSceneDescription::Serialize(FArchive& Ar)
+#if WITH_EDITOR
+TSharedPtr<FXmlFile> UDMXMVRGeneralSceneDescription::CreateXmlFile() const
 {
-	Super::Serialize(Ar);
-
-	Ar << MVRFixtures;
+	return RootNode ? RootNode->CreateXmlFile() : nullptr;
 }
+#endif // WITH_EDITOR
 
 #if WITH_EDITOR
 void UDMXMVRGeneralSceneDescription::WriteFixturePatchToGeneralSceneDescription(const UDMXEntityFixturePatch& FixturePatch)
 {
-	if (const UDMXLibrary* DMXLibrary = FixturePatch.GetParentLibrary())
+	checkf(RootNode, TEXT("Unexpected: MVR General Scene Description Root Node is invalid."));
+
+	const UDMXLibrary* DMXLibrary = FixturePatch.GetParentLibrary();
+	if (!DMXLibrary)
 	{
-		const FGuid& MVRFixtureUUID = FixturePatch.GetMVRFixtureUUID();
-		FDMXMVRFixture* MVRFixturePtr = FindMVRFixture(MVRFixtureUUID);
+		return;
+	}
 
-		FDMXMVRFixture MVRFixture;
-		if (MVRFixturePtr)
-		{
-			MVRFixture = *MVRFixturePtr;
-		}
-		else
-		{
-			const TArray<int32> UnitNumbersInUse = GetUnitNumbersInUse(*DMXLibrary);
-			const int32 UnitNumber = UnitNumbersInUse.Num() > 0 ? UnitNumbersInUse.Last() + 1 : 1;
+	const FGuid& MVRFixtureUUID = FixturePatch.GetMVRFixtureUUID();
+	UDMXMVRParametricObjectNodeBase** ParametricObjectNodePtr = RootNode->FindParametricObjectNodeByUUID(MVRFixtureUUID);
+
+	UDMXMVRFixtureNode* MVRFixtureNode = nullptr;
+	if (ParametricObjectNodePtr)
+	{
+		MVRFixtureNode = Cast<UDMXMVRFixtureNode>(*ParametricObjectNodePtr);
+	}
 
-			MVRFixture.Name = FixturePatch.Name;
-			MVRFixture.UUID = MVRFixtureUUID;
+	if (!MVRFixtureNode)
+	{
+		UDMXMVRChildListNode& AnyChildList = RootNode->GetOrCreateFirstChildListNode();
+		MVRFixtureNode = AnyChildList.CreateParametricObject<UDMXMVRFixtureNode>();
 
-			// Create a new Unit Number and add it to the array of Unit Numbers in use
-			MVRFixture.UnitNumber = UnitNumbersInUse.Last() + 1;
-		}
+		const TArray<int32> UnitNumbersInUse = GetUnitNumbersInUse(*DMXLibrary);
+		const int32 UnitNumber = UnitNumbersInUse.Num() > 0 ? UnitNumbersInUse.Last() + 1 : 1;
 
-		MVRFixture.SetUniverseID(FixturePatch.GetUniverseID());
-		MVRFixture.SetStartingChannel(FixturePatch.GetStartingChannel());
+		MVRFixtureNode->Name = FixturePatch.Name;
+		MVRFixtureNode->UUID = MVRFixtureUUID;
 
-		bool bSetGDTFSpec = false;
-		if (UDMXEntityFixtureType* FixtureType = FixturePatch.GetFixtureType())
+		// Create a new Unit Number and add it to the array of Unit Numbers in use
+		MVRFixtureNode->UnitNumber = UnitNumber;
+	}
+	check(MVRFixtureNode);
+
+	MVRFixtureNode->SetUniverseID(FixturePatch.GetUniverseID());
+	MVRFixtureNode->SetStartingChannel(FixturePatch.GetStartingChannel());
+
+	bool bSetGDTFSpec = false;
+	if (UDMXEntityFixtureType* FixtureType = FixturePatch.GetFixtureType())
+	{
+		if (UDMXImportGDTF* GDTF = FixtureType->GDTF)
 		{
-			if (UDMXImportGDTF* GDTF = FixtureType->GDTF)
+			const FString SourceFilename = [GDTF]()
 			{
-				const FString SourceFilename = [GDTF]()
+				if (GDTF && GDTF->GetGDTFAssetImportData())
 				{
-					if (GDTF && GDTF->GetGDTFAssetImportData())
-					{
-						return GDTF->GetGDTFAssetImportData()->GetSourceFilePathAndName();
-					}
-					return FString();
-				}();
-				MVRFixture.GDTFSpec = FPaths::GetBaseFilename(SourceFilename);
-			}
-
-			const int32 ModeIndex = FixturePatch.GetActiveModeIndex();
-			if (ensureAlwaysMsgf(FixtureType->Modes.IsValidIndex(ModeIndex), TEXT("Trying to write Active Mode of to General Scene Description, but the Mode Index is invalid")))
-			{
-				MVRFixture.GDTFMode = FixtureType->Modes[ModeIndex].ModeName;
-			}
+					return GDTF->GetGDTFAssetImportData()->GetSourceFilePathAndName();
+				}
+				return FString();
+			}();
+			MVRFixtureNode->GDTFSpec = FPaths::GetCleanFilename(SourceFilename);
+			bSetGDTFSpec = true;
 		}
 
-		if (!bSetGDTFSpec)
+		const int32 ModeIndex = FixturePatch.GetActiveModeIndex();
+		if (ensureAlwaysMsgf(FixtureType->Modes.IsValidIndex(ModeIndex), TEXT("Trying to write Active Mode of to General Scene Description, but the Mode Index is invalid")))
 		{
-			// No point to set a mode when there's no GDTF
-			MVRFixture.GDTFMode = TEXT("");
-			MVRFixture.GDTFSpec = TEXT("");
+			MVRFixtureNode->GDTFMode = FixtureType->Modes[ModeIndex].ModeName;
 		}
+	}
 
-		if (MVRFixturePtr)
-		{
-			(*MVRFixturePtr) = MVRFixture;
-		}
-		else
-		{
-			MVRFixtures.Add(MVRFixture);
-		}
+	if (!bSetGDTFSpec)
+	{
+		// Don't set a mode when there's no GDTF
+		MVRFixtureNode->GDTFMode = TEXT("");
+		MVRFixtureNode->GDTFSpec = TEXT("");
 	}
 }
 #endif // WITH_EDITOR
 
 #if WITH_EDITOR
-void UDMXMVRGeneralSceneDescription::ParseGeneralSceneDescriptionXml(const TSharedRef<FXmlFile>& GeneralSceneDescription)
+void UDMXMVRGeneralSceneDescription::ParseGeneralSceneDescriptionXml(const TSharedRef<FXmlFile>& GeneralSceneDescriptionXml)
 {
-	if (const FXmlNode* RootNode = GeneralSceneDescription->GetRootNode())
-	{
-		static const FString NodeName_Scene = TEXT("Scene");
-		if (const FXmlNode* SceneNode = RootNode->FindChildNode(NodeName_Scene))
-		{
-			static const FString NodeName_Layers = TEXT("Layers");
-			if (const FXmlNode* LayersNode = SceneNode->FindChildNode(NodeName_Layers))
-			{
-				for (const FXmlNode* LayerNode : LayersNode->GetChildrenNodes())
-				{
-					static const FString NodeName_Layer = TEXT("Layer");
-					if (LayerNode->GetTag() == NodeName_Layer)
-					{
-						for (const FXmlNode* ChildListNode : LayerNode->GetChildrenNodes())
-						{
-							static const FString NodeName_ChildList = TEXT("ChildList");
-							if (ChildListNode->GetTag() == NodeName_ChildList)
-							{
-								for (const FXmlNode* FixtureNode : ChildListNode->GetChildrenNodes())
-								{
-									static const FString NodeName_Fixture = TEXT("Fixture");
-									if (FixtureNode->GetTag() == NodeName_Fixture)
-									{
-										FDMXMVRFixture MVRFixture(FixtureNode);
-										if (MVRFixture.IsValid())
-										{
-											MVRFixtures.Add(MoveTemp(MVRFixture));
-										}
-									}
-								}
-							}
-						}
-					}
-				}
-			}
-		}
-	}
+	checkf(RootNode, TEXT("Unexpected: MVR General Scene Description Root Node is invalid."));
+
+	RootNode->InitializeFromGeneralSceneDescriptionXml(GeneralSceneDescriptionXml);
 }
 #endif // WITH_EDITOR
 
 TArray<int32> UDMXMVRGeneralSceneDescription::GetUnitNumbersInUse(const UDMXLibrary& DMXLibrary)
 {
+	checkf(RootNode, TEXT("Unexpected: MVR General Scene Description Root Node is invalid."));
+
 	TArray<UDMXEntityFixturePatch*> FixturePatches = DMXLibrary.GetEntitiesTypeCast<UDMXEntityFixturePatch>();
 
 	TArray<FGuid> MVRFixtureUUIDsInUse;
@@ -236,27 +212,22 @@ TArray<int32> UDMXMVRGeneralSceneDescription::GetUnitNumbersInUse(const UDMXLibr
 		MVRFixtureUUIDsInUse.Add(FixturePatch->GetMVRFixtureUUID());
 	}
 
-	// Fetch Unit Numbers in use to create new ones when adding new MVR Fixtures
 	TArray<int32> UnitNumbersInUse;
 	for (const FGuid& MVRFixtureUUID : MVRFixtureUUIDsInUse)
 	{
-		const FDMXMVRFixture* MVRFixturePtr = FindMVRFixture(MVRFixtureUUID);
-		if (MVRFixturePtr)
+		if (UDMXMVRParametricObjectNodeBase** ObjectNodePtr = RootNode->FindParametricObjectNodeByUUID(MVRFixtureUUID))
 		{
-			UnitNumbersInUse.Add(MVRFixturePtr->UnitNumber);
-		}
-	}
-	if (UnitNumbersInUse.Num() == 0)
-	{
-		UnitNumbersInUse.Add(0);
-	}
-	else
-	{
-		UnitNumbersInUse.Sort([](int32 UnitNumberA, int32 UnitNumberB)
+			if (UDMXMVRFixtureNode* FixtureNode = Cast<UDMXMVRFixtureNode>(*ObjectNodePtr))
 			{
-				return UnitNumberA < UnitNumberB;
-			});
+				UnitNumbersInUse.Add(FixtureNode->UnitNumber);
+			}
+		}
 	}
 
+	UnitNumbersInUse.Sort([](int32 UnitNumberA, int32 UnitNumberB)
+		{
+			return UnitNumberA < UnitNumberB;
+		});
+
 	return UnitNumbersInUse;
 }
diff --git a/Engine/Plugins/VirtualProduction/DMX/DMXEngine/Source/DMXRuntime/Private/MVR/DMXMVRSceneActor.cpp b/Engine/Plugins/VirtualProduction/DMX/DMXEngine/Source/DMXRuntime/Private/MVR/DMXMVRSceneActor.cpp
--- a/Engine/Plugins/VirtualProduction/DMX/DMXEngine/Source/DMXRuntime/Private/MVR/DMXMVRSceneActor.cpp
+++ b/Engine/Plugins/VirtualProduction/DMX/DMXEngine/Source/DMXRuntime/Private/MVR/DMXMVRSceneActor.cpp
@@ -10,6 +10,7 @@
 #include "MVR/DMXMVRFixtureActorLibrary.h"
 #include "MVR/DMXMVRAssetUserData.h"
 #include "MVR/DMXMVRFixtureActorInterface.h"
+#include "MVR/Types/DMXMVRFixtureNode.h"
 
 #include "DatasmithAssetUserData.h"
 #include "Components/SceneComponent.h"
@@ -197,15 +198,13 @@ void ADMXMVRSceneActor::SetDMXLibrary(UDMXLibrary* NewDMXLibrary)
 		}
 
 		const FGuid& MVRFixtureUUID = FixturePatch->GetMVRFixtureUUID();
-		const FDMXMVRFixture* MVRFixturePtr = GeneralSceneDescription->FindMVRFixture(MVRFixtureUUID);
-		if (!MVRFixturePtr)
+		const UDMXMVRFixtureNode* FixtureNode = GeneralSceneDescription->FindFixtureNode(MVRFixtureUUID);
+		if (!FixtureNode)
 		{
 			continue;
 		}
-		const FDMXMVRFixture& MVRFixture = *MVRFixturePtr;
-
-		const FTransform Transform = MVRFixture.Transform.IsSet() ? MVRFixture.Transform.GetValue() : FTransform::Identity;
 
+		const FTransform Transform = FixtureNode->GetTransformAbsolute();
 		SpawnMVRActor(ActorClass, FixturePatch, Transform);
 	}
 }
diff --git a/Engine/Plugins/VirtualProduction/DMX/DMXEngine/Source/DMXRuntime/Private/MVR/DMXMVRUnzip.cpp b/Engine/Plugins/VirtualProduction/DMX/DMXEngine/Source/DMXRuntime/Private/MVR/DMXMVRUnzip.cpp
deleted file mode 100644
diff --git a/Engine/Plugins/VirtualProduction/DMX/DMXEngine/Source/DMXRuntime/Private/MVR/DMXMVRVersion.cpp b/Engine/Plugins/VirtualProduction/DMX/DMXEngine/Source/DMXRuntime/Private/MVR/DMXMVRVersion.cpp
new file mode 100644
--- /dev/null
+++ b/Engine/Plugins/VirtualProduction/DMX/DMXEngine/Source/DMXRuntime/Private/MVR/DMXMVRVersion.cpp
@@ -0,0 +1,14 @@
+// Copyright Epic Games, Inc. All Rights Reserved.
+
+#include "MVR/DMXMVRVersion.h"
+
+
+const FString FDMXMVRVersion::GetMajorVersionAsString() 
+{ 
+	return FString::FromInt(FDMXMVRVersion::MajorVersion); 
+}
+
+const FString FDMXMVRVersion::GetMinorVersionAsString() 
+{ 
+	return FString::FromInt(FDMXMVRVersion::MinorVersion); 
+}
diff --git a/Engine/Plugins/VirtualProduction/DMX/DMXEngine/Source/DMXRuntime/Private/MVR/Types/DMXMVRChildListNode.cpp b/Engine/Plugins/VirtualProduction/DMX/DMXEngine/Source/DMXRuntime/Private/MVR/Types/DMXMVRChildListNode.cpp
new file mode 100644
--- /dev/null
+++ b/Engine/Plugins/VirtualProduction/DMX/DMXEngine/Source/DMXRuntime/Private/MVR/Types/DMXMVRChildListNode.cpp
@@ -0,0 +1,141 @@
+// Copyright Epic Games, Inc. All Rights Reserved.
+
+#include "MVR/Types/DMXMVRChildListNode.h"
+
+#include "MVR/Types/DMXMVRFixtureNode.h"
+#include "MVR/Types/DMXMVRGroupObjectNode.h"
+#include "MVR/Types/DMXMVRParametricObjectNodeBase.h"
+
+#include "XmlNode.h"
+#include "Algo/Find.h"
+
+
+void UDMXMVRChildListNode::InitializeFromChildListXmlNode(const FXmlNode& ChildListXmlNode)
+{
+	checkf(ChildListXmlNode.GetTag() == TEXT("ChildList"), TEXT("Trying to initialize a Child List Node from Xml Node, but the Node's Tag is '%s' instead of 'ChildList'."), *ChildListXmlNode.GetTag());
+
+	ParametricObjectNodes.Reset();
+	for (const FXmlNode* XmlNode : ChildListXmlNode.GetChildrenNodes())
+	{
+		constexpr TCHAR GroupObjectTag[] = TEXT("GroupObject");
+		constexpr TCHAR FixtureTag[] = TEXT("Fixture");
+
+		if (XmlNode->GetTag() == GroupObjectTag)
+		{
+			UDMXMVRGroupObjectNode* NewGroupObjectNode = CreateParametricObject<UDMXMVRGroupObjectNode>();
+			NewGroupObjectNode->InitializeFromGroupObjectXmlNode(*XmlNode);
+		}
+		else if (XmlNode->GetTag() == FixtureTag)
+		{
+			UDMXMVRFixtureNode* NewFixtureNode = CreateParametricObject<UDMXMVRFixtureNode>();
+			NewFixtureNode->InitializeFromFixtureXmlNode(*XmlNode);
+		}
+	}
+}
+
+void UDMXMVRChildListNode::CreateXmlNodeInParent(FXmlNode& ParentNode) const
+{
+	checkf(ParentNode.GetTag() == TEXT("layer") || ParentNode.GetTag() == TEXT("groupobject"), TEXT("Child List Nodes have to be created in a Layer Or Group Object node, but parent node is %s."), *ParentNode.GetTag());
+
+	// Don't export empty child lists
+	if (ParametricObjectNodes.IsEmpty())
+	{
+		return;
+	}
+
+	constexpr TCHAR Tag[] = TEXT("ChildList");
+	ParentNode.AppendChildNode(Tag);
+	FXmlNode* ChildListXmlNode = ParentNode.GetChildrenNodes().Last();
+	check(ChildListXmlNode);
+
+	for (const UDMXMVRParametricObjectNodeBase* ParametricObject : ParametricObjectNodes)
+	{
+		ParametricObject->CreateXmlNodeInParent(*ChildListXmlNode);
+	}
+}
+
+void UDMXMVRChildListNode::GetFixtureNodes(TArray<UDMXMVRFixtureNode*>& OutFixtureNodes) const
+{
+	for (UDMXMVRParametricObjectNodeBase* Node : ParametricObjectNodes)
+	{
+		if (UDMXMVRFixtureNode* FixtureNode = Cast<UDMXMVRFixtureNode>(Node))
+		{
+			OutFixtureNodes.Add(FixtureNode);
+		}
+		else if (UDMXMVRGroupObjectNode* GroupObjectNode = Cast<UDMXMVRGroupObjectNode>(Node))
+		{
+			// For group objects, also get their Fixture Nodes
+			GroupObjectNode->GetFixtureNodes(OutFixtureNodes);
+		}
+	}
+}
+
+bool UDMXMVRChildListNode::Contains(UDMXMVRParametricObjectNodeBase* ParametricObjectNode)
+{
+	for (UDMXMVRParametricObjectNodeBase* Other : ParametricObjectNodes)
+	{
+		if (UDMXMVRGroupObjectNode* GroupObjectNode = Cast<UDMXMVRGroupObjectNode>(Other))
+		{		
+			// For group objects, also test if they contain the node
+			if (GroupObjectNode->Contains(ParametricObjectNode))
+			{
+				return true;
+			}
+		}
+
+		if (Other == ParametricObjectNode)
+		{
+			return true;
+		}
+	}
+
+	return false;
+}
+
+UDMXMVRParametricObjectNodeBase** UDMXMVRChildListNode::FindParametricObjectNodeByUUID(const FGuid& UUID)
+{
+	UDMXMVRParametricObjectNodeBase** ParametricObjectNodePtr = Algo::FindByPredicate(ParametricObjectNodes, [UUID](UDMXMVRParametricObjectNodeBase* ParametricObjectNode)
+		{
+			return ParametricObjectNode->UUID == UUID;
+		});
+	if (ParametricObjectNodePtr)
+	{
+		return ParametricObjectNodePtr;
+	}
+
+	for (UDMXMVRParametricObjectNodeBase* Node : ParametricObjectNodes)
+	{
+		if (UDMXMVRGroupObjectNode* GroupObjectNode = Cast<UDMXMVRGroupObjectNode>(Node))
+		{
+			ParametricObjectNodePtr = GroupObjectNode->FindParametricObjectNodeByUUID(UUID);
+			if (ParametricObjectNodePtr)
+			{
+				return ParametricObjectNodePtr;
+			}
+		}
+	}
+
+	return nullptr;
+}
+
+bool UDMXMVRChildListNode::RemoveParametricObject(UDMXMVRParametricObjectNodeBase* ParametricObjectNode)
+{
+	const int32 NumRemovedElements = ParametricObjectNodes.RemoveSingle(ParametricObjectNode);
+	if (NumRemovedElements == 1)
+	{
+		return true;
+	}
+
+	for (UDMXMVRParametricObjectNodeBase* Node : ParametricObjectNodes)
+	{
+		if (UDMXMVRGroupObjectNode* GroupObjectNode = Cast<UDMXMVRGroupObjectNode>(Node))
+		{
+			if (GroupObjectNode->RemoveParametricObjectNode(ParametricObjectNode))
+			{
+				return true;
+			}
+		}
+	}
+
+	return false;
+}
diff --git a/Engine/Plugins/VirtualProduction/DMX/DMXEngine/Source/DMXRuntime/Private/MVR/Types/DMXMVRFixtureNode.cpp b/Engine/Plugins/VirtualProduction/DMX/DMXEngine/Source/DMXRuntime/Private/MVR/Types/DMXMVRFixtureNode.cpp
new file mode 100644
--- /dev/null
+++ b/Engine/Plugins/VirtualProduction/DMX/DMXEngine/Source/DMXRuntime/Private/MVR/Types/DMXMVRFixtureNode.cpp
@@ -0,0 +1,762 @@
+// Copyright Epic Games, Inc. All Rights Reserved.
+
+#include "MVR/Types/DMXMVRFixtureNode.h"
+
+#include "DMXProtocolCommon.h"
+#include "DMXRuntimeLog.h"
+#include "DMXRuntimeUtils.h"
+
+#include "XmlNode.h"
+
+
+namespace UE::DMXRuntime::DMXMVRFixture::Private
+{
+	/** Parses the Name from the MVR Fixture Node, returns true on success. */
+	bool ParseName(const FXmlNode& InFixtureNode, FString& OutName)
+	{
+		constexpr TCHAR NameTag[] = TEXT("Name");
+		OutName = InFixtureNode.GetAttribute(NameTag);
+		if (!OutName.IsEmpty())
+		{
+			return true;
+		}
+
+		return false;
+	}
+
+	/** Parses the UUID from the MVR Fixture Node, returns true on success. */
+	bool ParseUUID(const FXmlNode& InFixtureNode, FGuid& OutUUID)
+	{
+		constexpr TCHAR UUIDTag[] = TEXT("UUID");
+		const FString UUIDString = InFixtureNode.GetAttribute(UUIDTag);
+		if (FGuid::Parse(UUIDString, OutUUID))
+		{
+			return true;
+		}
+		return false;
+	}
+
+	/** Parses the Matrix from the MVR Fixture Node, optional, always succeeds. */
+	void ParseMatrix(const FXmlNode& InFixtureNode, FDMXOptionalTransform& OutTransform)
+	{
+		OutTransform.Reset();
+
+		constexpr TCHAR MatrixTag[] = TEXT("Matrix");
+		const FXmlNode* MatrixNode = InFixtureNode.FindChildNode(MatrixTag);
+
+		if (!MatrixNode)
+		{
+			return;
+		}
+
+		OutTransform = FDMXRuntimeUtils::ParseGDTFMatrix(MatrixNode->GetContent());
+	}
+	
+	/** Parses the GDTFSpec from the MVR Fixture Node, returns true on success. */
+	bool ParseGDTFSpec(const FXmlNode& InFixtureNode, FString& OutGDTFSpec)
+	{
+		constexpr TCHAR GDTFSpecTag[] = TEXT("GDTFSpec");
+		const FXmlNode* GDTFSpecNode = InFixtureNode.FindChildNode(GDTFSpecTag);
+		if (GDTFSpecNode)
+		{
+			OutGDTFSpec = GDTFSpecNode->GetContent();
+			return true;
+		}
+
+		return false;
+	}
+
+	/** Parses the GDTFMode from the MVR Fixture Node, returns true on success. */
+	bool ParseGDTFMode(const FXmlNode& InFixtureNode, FString& OutGDTFMode)
+	{
+		constexpr TCHAR GDTFModeTag[] = TEXT("GDTFMode");
+		const FXmlNode* GDTFModeNode = InFixtureNode.FindChildNode(GDTFModeTag);
+		if (GDTFModeNode)
+		{
+			OutGDTFMode = GDTFModeNode->GetContent();
+			return true;
+		}
+		return true;
+	}
+
+	/** Parses Focus from the MVR Fixture Node, optional, always succeeds. */
+	void ParseFocus(const FXmlNode& InFixtureNode, FDMXOptionalGuid& OutFocus)
+	{
+		OutFocus.Reset();
+
+		constexpr TCHAR FocusTag[] = TEXT("Focus");
+		const FXmlNode* FocusNode = InFixtureNode.FindChildNode(FocusTag);
+		if (FocusNode)
+		{
+			const FString FocusString = InFixtureNode.GetContent();
+			FGuid FocusGuid;
+			if (FGuid::Parse(FocusString, FocusGuid))
+			{
+				OutFocus = FocusGuid;
+			}
+		};
+	}
+
+	/** Parses Cast Shadows from the MVR Fixture Node, optional, always succeeds. */
+	void ParseCastShadow(const FXmlNode& InFixtureNode, FDMXOptionalBool& OutCastShadow)
+	{
+		OutCastShadow.Reset();
+
+		constexpr TCHAR CastShadowTag[] = TEXT("CastShadow");
+		const FXmlNode* CastShadowNode = InFixtureNode.FindChildNode(CastShadowTag);
+		if (CastShadowNode)
+		{
+			bool bCastShadow = 
+				CastShadowNode->GetContent() != TEXT("0") ||
+				CastShadowNode->GetContent().Equals(TEXT("true"), ESearchCase::IgnoreCase);
+
+			OutCastShadow = bCastShadow;
+		}
+	}
+
+	/** Parses Position of from MVR Fixture Node, optional, always succeeds. */
+	void ParsePosition(const FXmlNode& InFixtureNode, FDMXOptionalGuid& OutPosition)
+	{
+		OutPosition.Reset();
+
+		constexpr TCHAR PositionTag[] = TEXT("Position");
+		const FXmlNode* PositionNode = InFixtureNode.FindChildNode(PositionTag);
+		if (PositionNode)
+		{
+			FGuid PositionGuid;
+			if (FGuid::Parse(PositionNode->GetContent(), PositionGuid))
+			{
+				OutPosition = PositionGuid;
+			}
+		}
+	}
+
+	/** Parses the Fixture Id from the MVR Fixture Node, returns true on success. */
+	bool ParseFixtureId(const FXmlNode& InFixtureNode, FString& OutFixtureId)
+	{
+		OutFixtureId.Reset();
+
+		constexpr TCHAR FixtureIdTag[] = TEXT("FixtureId");
+		const FXmlNode* FixtureIdNode = InFixtureNode.FindChildNode(FixtureIdTag);
+		if (FixtureIdNode)
+		{
+			OutFixtureId = FixtureIdNode->GetContent();
+			return true;
+		}
+
+		return false;
+	}
+
+	/** Parses the Unit Number from the MVR Fixture Node, returns true on success. */
+	bool ParseUnitNumber(const FXmlNode& InFixtureNode, int32& OutUnitNumber)
+	{
+		// @todo: Temp fix for the larger fixture id/unit number mismatch
+		constexpr TCHAR UnitNumberTag[] = TEXT("FixtureId");
+		const FXmlNode* UnitNumberNode = InFixtureNode.FindChildNode(UnitNumberTag);
+		if (UnitNumberNode)
+		{
+			int32 UnitNumber;
+			if (LexTryParseString(UnitNumber, *UnitNumberNode->GetContent()))
+			{
+				OutUnitNumber = UnitNumber;
+				return true;
+			}
+		}
+		return false;
+	}
+
+	/** Parses the Addresses from the MVR Fixture Node, defaulted, always succeeds */
+	bool ParseAddresses(const FXmlNode& InFixtureNode, FDMXMVRFixtureAddresses& OutAddresses)
+	{
+		constexpr TCHAR AddressesTag[] = TEXT("Addresses");
+		const FXmlNode* AddressesNode = InFixtureNode.FindChildNode(AddressesTag);
+
+		constexpr TCHAR AddressTag[] = TEXT("Address");
+		const FXmlNode* AddressNode = AddressesNode ? AddressesNode->FindChildNode(AddressTag) : nullptr;
+
+		if (AddressNode)
+		{
+			int32 Universe = -1;
+			int32 Address = -1;
+
+			auto ParseAddressStringLambda = [](const FXmlNode& InAddressNode, int32& OutUniverse, int32& OutAddress) -> const bool
+			{
+				const FString InAddressString = InAddressNode.GetContent();
+
+				int32 SeparatorIndex = INDEX_NONE;
+				if (InAddressString.FindChar('.', SeparatorIndex))
+				{
+					if (SeparatorIndex > 0 && SeparatorIndex < InAddressString.Len() - 2)
+					{
+						const FString UniverseString = InAddressString.Left(SeparatorIndex);
+						const FString AddressString = InAddressString.RightChop(SeparatorIndex + 1);
+
+						const bool bValidUniverse = LexTryParseString(OutUniverse, *UniverseString) && OutUniverse >= 0 && OutUniverse <= DMX_MAX_UNIVERSE;
+						const bool bValidAddress = LexTryParseString(OutAddress, *UniverseString) && OutAddress >= 0 && OutAddress <= DMX_MAX_ADDRESS;
+
+						return bValidUniverse && bValidAddress;
+					}
+				}
+				else if (LexTryParseString(OutAddress, *InAddressString))
+				{
+					static const int32 UniverseSize = 512;
+
+					OutUniverse = OutAddress / UniverseSize;
+					OutAddress = OutAddress % UniverseSize;
+
+					return 
+						OutUniverse >= 0 && 
+						OutUniverse <= DMX_MAX_UNIVERSE &&
+						OutAddress >= 0 && 
+						OutAddress <= DMX_MAX_ADDRESS;
+				}
+
+				return false;
+			}(*AddressNode, Universe, Address);
+			if (ParseAddressStringLambda)
+			{
+				OutAddresses.Universe = Universe;
+				OutAddresses.Address = Address;
+
+				return true;
+			}
+		}
+		return false;
+	}
+
+	/** Parses the CIE Color from the MVR Fixture Node, optional, always succeeds. */
+	void ParseCIEColor(const FXmlNode& InFixtureNode, FDMXOptionalColorCIE1931& OutColorCIE1931)
+	{
+		OutColorCIE1931.Reset();
+
+		constexpr TCHAR CIEColorTag[] = TEXT("CIEColor");
+		const FXmlNode* CIEColorNode = InFixtureNode.FindChildNode(CIEColorTag);
+		if (CIEColorNode)
+		{
+			const FString CIEColorString = CIEColorNode->GetContent();
+			if (CIEColorString.IsEmpty())
+			{
+				return;
+			}
+
+			TArray<FString> CIEColorArray;
+			CIEColorString.ParseIntoArray(CIEColorArray, TEXT(","));
+
+			if (CIEColorArray.Num() != 3	)
+			{
+				return;
+			}
+
+			FDMXColorCIE1931xyY Result;
+			bool bSuccess = LexTryParseString<float>(Result.X, *CIEColorArray[0]);
+			bSuccess = bSuccess && LexTryParseString<float>(Result.Y, *CIEColorArray[1]);
+			bSuccess = bSuccess && LexTryParseString<float>(Result.YY, *CIEColorArray[2]);
+
+			if (bSuccess)
+			{
+				OutColorCIE1931 = Result;
+			}
+		}
+	}
+
+	/** Parses the Fixture Type Id from MVR Fixture Node, optional, always succeeds. */
+	void ParseFixtureTypeId(const FXmlNode& InFixtureNode, FDMXOptionalInt32& OutFixtureTypeId)
+	{
+		OutFixtureTypeId.Reset();
+
+		constexpr TCHAR FixtureTypeIdTag[] = TEXT("FixtureTypeId");
+		const FXmlNode* FixtureTypeIdNode = InFixtureNode.FindChildNode(FixtureTypeIdTag);
+		if (FixtureTypeIdNode)
+		{
+			int32 FixtureTypeId;
+			if (LexTryParseString(FixtureTypeId, *FixtureTypeIdNode->GetContent()))
+			{
+				OutFixtureTypeId = FixtureTypeId;
+			}
+		}
+	}
+
+	/** Parses the Custom Id from MVR Fixture Node, optional, always succeeds. */
+	void ParseCustomId(const FXmlNode& InFixtureNode, FDMXOptionalInt32& OutCustomId)
+	{
+		OutCustomId.Reset();
+
+		constexpr TCHAR PositionTag[] = TEXT("Position");
+		const FXmlNode* PositionNode = InFixtureNode.FindChildNode(PositionTag);
+		if (PositionNode)
+		{
+			int32 FixtureTypeId;
+			if (LexTryParseString(FixtureTypeId, *PositionNode->GetContent()))
+			{
+				OutCustomId = FixtureTypeId;
+			}
+		}
+	}
+
+	/** Parses the Mapping from MVR Fixture Node, optional, always succeeds. */
+	void ParseMappings(const FXmlNode& InFixtureNode, TArray<FDMXMVRFixtureMapping>& OutMappings)
+	{
+		OutMappings.Reset();
+
+		constexpr TCHAR MappingsTag[] = TEXT("Mappings");
+		const FXmlNode* MappingsNode = InFixtureNode.FindChildNode(MappingsTag);
+		if (MappingsNode)
+		{
+			for (const FXmlNode* MappingNode : MappingsNode->GetChildrenNodes())
+			{
+				FDMXMVRFixtureMapping Mapping;
+
+				constexpr TCHAR ChildNodenName_LinkDef[] = TEXT("LinkDef");
+				const FXmlNode* LinkDefNode = MappingNode->FindChildNode(ChildNodenName_LinkDef);
+				if (LinkDefNode && !FGuid::Parse(LinkDefNode->GetContent(), Mapping.LinkedDef))
+				{
+					continue;
+				}
+
+				constexpr TCHAR ChildNodenName_UX[] = TEXT("UX");
+				const FXmlNode* UXNode = InFixtureNode.FindChildNode(ChildNodenName_UX);
+				int32 UX;
+				if (UXNode && LexTryParseString(UX, *UXNode->GetContent()))
+				{
+					Mapping.UX = UX;
+				}
+
+				constexpr TCHAR ChildNodenName_UY[] = TEXT("UY");
+				const FXmlNode* UYNode = InFixtureNode.FindChildNode(ChildNodenName_UY);
+				int32 UY;
+				if (UYNode && LexTryParseString(UY, *UYNode->GetContent()))
+				{
+					Mapping.UY = UY;
+				}
+
+				constexpr TCHAR ChildNodenName_OX[] = TEXT("OX");
+				const FXmlNode* OXNode = InFixtureNode.FindChildNode(ChildNodenName_OX);
+				int32 OX;
+				if (OXNode && LexTryParseString(OX, *OXNode->GetContent()))
+				{
+					Mapping.OX = OX;
+				}
+
+				constexpr TCHAR ChildNodenName_OY[] = TEXT("OY");
+				const FXmlNode* OYNode = InFixtureNode.FindChildNode(ChildNodenName_OY);
+				int32 OY;
+				if (OYNode && LexTryParseString(OY, *OYNode->GetContent()))
+				{
+					Mapping.OY = OY;
+				}
+
+
+				constexpr TCHAR ChildNodenName_RZ[] = TEXT("RZ");
+				const FXmlNode* RZNode = InFixtureNode.FindChildNode(ChildNodenName_RZ);
+				int32 RZ;
+				if (RZNode && !LexTryParseString(RZ, *RZNode->GetContent()))
+				{
+					Mapping.RZ = RZ;
+				}
+
+				OutMappings.Add(Mapping);
+			}
+		}
+	}
+
+	/** Parses the Gobo from MVR Fixture Node, optional, always succeeds. */
+	void ParseGobo(const FXmlNode& InFixtureNode, FDMXOptionalMVRFixtureGobo& OutGobo)
+	{
+		OutGobo.Value.Reset();
+
+		constexpr TCHAR GoboTag[] = TEXT("Gobo");
+		const FXmlNode* GoboNode = InFixtureNode.FindChildNode(GoboTag);
+		if (GoboNode)
+		{
+			FDMXMVRFixtureGobo FixtureGobo;
+			const FString& GoboResource = GoboNode->GetContent();
+			FixtureGobo.Value = GoboResource;
+
+			constexpr TCHAR RotationTag[] = TEXT("Rotation");
+			const FXmlNode* GoboRotationNode = GoboNode->FindChildNode(RotationTag);
+			float GoboRotation;
+			if (GoboRotationNode && LexTryParseString(GoboRotation, *GoboRotationNode->GetContent()))
+			{
+				FixtureGobo.Rotation = GoboRotation;
+			}
+
+			OutGobo.Value = FixtureGobo;
+		}
+	}
+}
+
+void FDMXMVRFixtureAddresses::CreateXmlNodeInParent(FXmlNode& ParentNode) const
+{
+	const int32 AbsoluteChannel = Universe * DMX_UNIVERSE_SIZE + Address;
+	constexpr TCHAR Tag[] = TEXT("Address");
+	const FString Content = FString::FromInt(AbsoluteChannel);
+
+	constexpr TCHAR BreakAttributeName[] = TEXT("rotation");
+	constexpr TCHAR BreakAttributeValue[] = TEXT("0");
+	const TArray<FXmlAttribute> Attributes =
+	{
+		FXmlAttribute(BreakAttributeName, BreakAttributeValue)
+	};
+
+	ParentNode.AppendChildNode(Tag, Content, Attributes);
+}
+
+void FDMXOptionalMVRFixtureGobo::CreateXmlNodeInParent(FXmlNode& ParentNode) const
+{
+	const TOptional<float>& OptionalRotation = Value.GetValue().Rotation.Value;
+	const FString& Content = Value.GetValue().Value; // The resource
+
+	if (!Content.IsEmpty() && OptionalRotation.IsSet())
+	{
+		// Add self
+		constexpr TCHAR Tag[] = TEXT("Gobo");
+
+		TArray<FXmlAttribute> Attributes;
+		if (OptionalRotation.IsSet())
+		{
+			constexpr TCHAR RotationAttributeName[] = TEXT("rotation");
+			constexpr int32 MinFractionalDigitsOfValue = 6; // As other soft- and hardwares
+			const FString RotationAttributeValue = FString::SanitizeFloat(OptionalRotation.GetValue(), MinFractionalDigitsOfValue);
+
+			Attributes.Add(FXmlAttribute(RotationAttributeName, RotationAttributeValue));
+		}
+		ParentNode.AppendChildNode(Tag, Content, Attributes);
+	}
+}
+
+void FDMXMVRFixtureMapping::CreateXmlNodeInParent(FXmlNode& ParentNode) const
+{
+	constexpr TCHAR Tag[] = TEXT("Mapping");
+	constexpr TCHAR LinkDefAttributeName[] = TEXT("linkedDef");
+	const TArray<FXmlAttribute> Attributes
+	{
+		FXmlAttribute(LinkDefAttributeName, LinkedDef.ToString(EGuidFormats::DigitsWithHyphens))
+	};
+	ParentNode.AppendChildNode(Tag);
+	FXmlNode* MappingXmlNode = ParentNode.GetChildrenNodes().Last();
+	check(MappingXmlNode);
+
+	if (UX.IsSet())
+	{
+		constexpr TCHAR UxTag[] = TEXT("ux");
+		const FString ChildContent = FString::FromInt(UX.GetValue());
+		MappingXmlNode->AppendChildNode(UxTag, ChildContent);
+	}
+
+	if (UY.IsSet())
+	{
+		constexpr TCHAR UyTag[] = TEXT("uy");
+		const FString ChildContent = FString::FromInt(UY.GetValue());
+		MappingXmlNode->AppendChildNode(UyTag, ChildContent);
+	}
+
+	if (OX.IsSet())
+	{
+		constexpr TCHAR OxTag[] = TEXT("ox");
+		const FString ChildContent = FString::FromInt(OX.GetValue());
+		MappingXmlNode->AppendChildNode(OxTag, ChildContent);
+	}
+
+	if (OY.IsSet())
+	{
+		constexpr TCHAR OyTag[] = TEXT("oy");
+		const FString ChildContent = FString::FromInt(OY.GetValue());
+		MappingXmlNode->AppendChildNode(OyTag, ChildContent);
+	}
+
+	if (RZ.IsSet())
+	{
+		constexpr TCHAR RzTag[] = TEXT("rz");
+		const FString ChildContent = FString::FromInt(RZ.GetValue());
+		MappingXmlNode->AppendChildNode(RzTag, ChildContent);
+	}
+}
+
+bool UDMXMVRFixtureNode::InitializeFromFixtureXmlNode(const FXmlNode& FixtureXmlNode)
+{
+	checkf(FixtureXmlNode.GetTag() == TEXT("fixture"), TEXT("Trying to initialize a Fixture Node from Xml Node, but the Node's Tag is '%s' instead of 'Fixture'."), *FixtureXmlNode.GetTag());
+
+	constexpr TCHAR NodeName_Fixture[] = TEXT("Fixture");
+	if (!ensureAlwaysMsgf(FixtureXmlNode.GetTag() == NodeName_Fixture, TEXT("Trying to parse MVR Fixture Node, but node name doesn't match.")))
+	{
+		return false;
+	}
+
+	// Parse nodes. Those that don't fail are either defaulted or optional.
+	using namespace UE::DMXRuntime::DMXMVRFixture::Private;
+	if (!ParseName(FixtureXmlNode, Name))
+	{
+		// Non-optional
+		UE_LOG(LogDMXRuntime, Warning, TEXT("Cannot parse Fixture Node as MVR Fixture Name. Failed to parse MVR Node %s."), *FixtureXmlNode.GetTag());
+		return false;
+	}
+
+	if(!ParseUUID(FixtureXmlNode, UUID))
+	{
+		// Non-optional
+		UE_LOG(LogDMXRuntime, Warning, TEXT("Cannot parse Fixture '%s' from MVR. Failed to parse MVR UUID."), *Name);
+		return false;
+	}
+
+	ParseMatrix(FixtureXmlNode, Matrix);
+	if (Matrix.IsSet() && !Matrix.GetValue().IsValid())
+	{
+		// Mend invalid transforms
+		UE_LOG(LogDMXRuntime, Warning, TEXT("Imported Fixture '%s' from MVR but its Transformation Matrix is not valid, using Identity instead."), *Name);
+		Matrix = FTransform::Identity;
+	}
+
+	if (!ParseGDTFSpec(FixtureXmlNode, GDTFSpec))
+	{
+		// Non-optional
+		UE_LOG(LogDMXRuntime, Warning, TEXT("Cannot parse Fixture '%s' from MVR. No valid GDTF Spec specified."), *Name);
+		return false;
+	}
+
+	if (!ParseGDTFMode(FixtureXmlNode, GDTFMode))
+	{
+		// Non-optional
+		UE_LOG(LogDMXRuntime, Warning, TEXT("Cannot parse Fixture '%s' from MVR. No valid GDTF Mode specified."), *Name);
+		return false;
+	}
+
+	ParseFocus(FixtureXmlNode, Focus);
+	ParseCastShadow(FixtureXmlNode, CastShadow);
+	ParsePosition(FixtureXmlNode, Position);
+	
+	if (!ParseFixtureId(FixtureXmlNode, FixtureID))
+	{
+		// Non-optional, mend with empty string
+		UE_LOG(LogDMXRuntime, Warning, TEXT("Imported Fixture '%s' from MVR. No valid Fixture Id specified. Fixture Id set to Fixture Name."), *Name);
+		FixtureID = TEXT("");
+	}
+
+	if (!ParseUnitNumber(FixtureXmlNode, UnitNumber))
+	{
+		// Non-optional, but defaulted by standard
+		UE_LOG(LogDMXRuntime, Warning, TEXT("Imported Fixture '%s' from MVR. No valid Unit Number specified. Unit Number set to ''0'."), *Name);
+		UnitNumber = 0;
+	}
+
+	if (!ParseAddresses(FixtureXmlNode, Addresses))
+	{
+		// Non-optional
+		UE_LOG(LogDMXRuntime, Warning, TEXT("Cannot parse Fixture '%s' from MVR. No valid Addresses specified."), *Name);
+		return false;
+	}
+
+	ParseCIEColor(FixtureXmlNode, CIEColor);
+	ParseFixtureTypeId(FixtureXmlNode, FixtureTypeId);
+	ParseCustomId(FixtureXmlNode, CustomId);
+	ParseMappings(FixtureXmlNode, Mappings);
+	ParseGobo(FixtureXmlNode, Gobo);
+
+	return true;
+}
+
+void UDMXMVRFixtureNode::CreateXmlNodeInParent(FXmlNode& ParentNode) const
+{
+	checkf(ParentNode.GetTag() == TEXT("childlist"), TEXT("Fixture Nodes have to be created in a Child List, but parent node is %s."), *ParentNode.GetTag());
+
+	constexpr bool bLogInvalidReason = true;
+	if (!IsValid(bLogInvalidReason))
+	{
+		return;
+	}
+
+	constexpr TCHAR FixtureTag[] = TEXT("Fixture");
+	ParentNode.AppendChildNode(FixtureTag);
+	FXmlNode* FixtureXmlNode = ParentNode.GetChildrenNodes().Last();
+	check(FixtureXmlNode);
+
+	TArray<FXmlAttribute> Attributes;
+
+	// Name attribute
+	constexpr TCHAR NameAttributeName[] = TEXT("name");
+	FXmlAttribute NameAttribute(NameAttributeName, Name);
+	Attributes.Add(NameAttribute);
+
+	// UUID attribute
+	if (!ensureAlwaysMsgf(UUID.IsValid(), TEXT("Cannot export MVR Fixture '%s'. Invalid UUID specified."), *Name))
+	{
+		return;
+	}
+
+	constexpr TCHAR UUIDAttributeName[] = TEXT("uuid");
+	FXmlAttribute UUIDAttribute(UUIDAttributeName, UUID.ToString(EGuidFormats::DigitsWithHyphens));
+	Attributes.Add(UUIDAttribute);
+
+	FixtureXmlNode->SetAttributes(Attributes);
+
+
+	////////////////////
+	// Child Nodes
+
+	// Matrix 
+	if (Matrix.IsSet())
+	{
+		constexpr TCHAR Tag[] = TEXT("Matrix");
+		const FString Content = FDMXRuntimeUtils::ConvertTransformToGDTF4x3MatrixString(Matrix.GetValue());
+		FixtureXmlNode->AppendChildNode(Tag, Content);
+	}
+
+	// GDTFSpec
+	{
+		constexpr TCHAR Tag[] = TEXT("GDTFSpec");
+		FixtureXmlNode->AppendChildNode(Tag, GDTFSpec);
+	}
+
+	// GDTFMode
+	{
+		constexpr TCHAR Tag[] = TEXT("GDTFMode");
+		FixtureXmlNode->AppendChildNode(Tag, GDTFMode);
+	}
+
+	// Focus
+	if (Focus.IsSet())
+	{
+		constexpr TCHAR Tag[] = TEXT("Focus");
+		const FString Content = Focus.GetValue().ToString(EGuidFormats::DigitsWithHyphens);
+		FixtureXmlNode->AppendChildNode(Tag, Content);
+	}
+
+	// CastShadow
+	if (CastShadow.IsSet())
+	{
+		constexpr TCHAR Tag[] = TEXT("CastShadow");
+		const FString Content = CastShadow.GetValue() ? TEXT("true") : TEXT("false");
+		FixtureXmlNode->AppendChildNode(Tag, Content);
+	}
+
+	// Position
+	if (Position.IsSet())
+	{
+		constexpr TCHAR Tag[] = TEXT("Position");
+		const FString Content = Position.GetValue().ToString(EGuidFormats::DigitsWithHyphens);
+		FixtureXmlNode->AppendChildNode(Tag, Content);
+	}
+
+	// FixtureID
+	{
+		constexpr TCHAR Tag[] = TEXT("FixtureID");
+		FixtureXmlNode->AppendChildNode(Tag, FixtureID);
+	}
+
+	// UnitNumber
+	{
+		constexpr TCHAR Tag[] = TEXT("UnitNumber");
+		const FString Content = FString::FromInt(UnitNumber);
+		FixtureXmlNode->AppendChildNode(Tag, Content);
+	}
+
+	// Addresses
+	{
+		constexpr TCHAR Tag[] = TEXT("Addresses");
+		FixtureXmlNode->AppendChildNode(Tag);
+		FXmlNode* AddressesXmlNode = FixtureXmlNode->GetChildrenNodes().Last();
+		check(AddressesXmlNode);
+
+		Addresses.CreateXmlNodeInParent(*AddressesXmlNode);
+	}
+
+	// CIEColor
+	if (CIEColor.IsSet())
+	{
+		constexpr TCHAR Tag[] = TEXT("CIEColor");
+		const FString Content = CIEColor.GetValue().ToString();
+		FixtureXmlNode->AppendChildNode(Tag, Content);
+	}
+
+	// FixtureTypeId
+	if (FixtureTypeId.IsSet())
+	{
+		constexpr TCHAR Tag[] = TEXT("FixtureTypeId");
+		const FString Content = FString::FromInt(FixtureTypeId.GetValue());
+		FixtureXmlNode->AppendChildNode(Tag, Content);
+	}
+
+	// CustomId
+	if (CustomId.IsSet())
+	{
+		constexpr TCHAR Tag[] = TEXT("CustomId");
+		const FString Content = FString::FromInt(CustomId.GetValue());
+		FixtureXmlNode->AppendChildNode(Tag, Content);
+	}
+
+	// Mappings
+	if (Mappings.Num() > 0)
+	{
+		constexpr TCHAR Tag[] = TEXT("Mappings");
+		FixtureXmlNode->AppendChildNode(Tag);
+		FXmlNode* MappingsXmlNode = FixtureXmlNode->GetChildrenNodes().Last();
+		check(MappingsXmlNode);
+
+		for (const FDMXMVRFixtureMapping& Mapping : Mappings)
+		{
+			Mapping.CreateXmlNodeInParent(*MappingsXmlNode);
+		}
+	}
+
+	// Gobo
+	if (Gobo.Value.IsSet())
+	{
+		constexpr TCHAR Tag[] = TEXT("Gobo");
+		FixtureXmlNode->AppendChildNode(Tag);
+		FXmlNode* GoboXmlNode = FixtureXmlNode->GetChildrenNodes().Last();
+		check(GoboXmlNode);
+
+		Gobo.CreateXmlNodeInParent(*GoboXmlNode);
+	}
+}
+
+bool UDMXMVRFixtureNode::IsValid(bool bLogInvalidReason) const
+{
+	const FString BeautifulName = Name.IsEmpty() ? TEXT("unnamed") : Name;
+
+	if (!UUID.IsValid())
+	{
+		if (bLogInvalidReason)
+		{
+			UE_LOG(LogDMXRuntime, Warning, TEXT("Fixture '%' is not a valid MVR Fixture, its Fixture UUID is invalid."), *BeautifulName);
+		}
+		return false;
+	}
+	
+	if (GDTFSpec.IsEmpty())
+	{
+		if (bLogInvalidReason)
+		{
+			UE_LOG(LogDMXRuntime, Warning, TEXT("Fixture '%' is not a valid MVR Fixture, its GDTFSpec is empty."), *BeautifulName);
+		}
+		return false;
+	}
+
+	if (GDTFMode.IsEmpty())
+	{
+		if (bLogInvalidReason)
+		{
+			UE_LOG(LogDMXRuntime, Warning, TEXT("Fixture '%' is not a valid MVR Fixture, its GDTFMode is empty."), *BeautifulName);
+		}
+		return false;
+	}
+
+	if (Addresses.Universe < 0 || Addresses.Universe > DMX_MAX_UNIVERSE)
+	{		
+		if (bLogInvalidReason)
+		{
+			UE_LOG(LogDMXRuntime, Warning, TEXT("Fixture '%' is not a valid MVR Fixture, its Universe is < 0 or larger than the Range supported by UE."), *BeautifulName);
+		}
+		return false;
+	}
+
+	if (Addresses.Address < 1 || Addresses.Address > DMX_MAX_ADDRESS)
+	{
+		if (bLogInvalidReason)
+		{
+			UE_LOG(LogDMXRuntime, Warning, TEXT("Fixture '%' is not a valid MVR Fixture, its Address is < 1 or > 512."), *BeautifulName);
+		}
+		return false;
+	}
+
+	return true;
+}
diff --git a/Engine/Plugins/VirtualProduction/DMX/DMXEngine/Source/DMXRuntime/Private/MVR/Types/DMXMVRGroupObjectNode.cpp b/Engine/Plugins/VirtualProduction/DMX/DMXEngine/Source/DMXRuntime/Private/MVR/Types/DMXMVRGroupObjectNode.cpp
new file mode 100644
--- /dev/null
+++ b/Engine/Plugins/VirtualProduction/DMX/DMXEngine/Source/DMXRuntime/Private/MVR/Types/DMXMVRGroupObjectNode.cpp
@@ -0,0 +1,116 @@
+// Copyright Epic Games, Inc. All Rights Reserved.
+
+#include "MVR/Types/DMXMVRGroupObjectNode.h"
+
+#include "DMXRuntimeLog.h"
+#include "DMXRuntimeUtils.h"
+#include "MVR/Types/DMXMVRChildListNode.h"
+#include "MVR/Types/DMXMVRFixtureNode.h"
+#include "MVR/Types/DMXMVRParametricObjectNodeBase.h"
+
+#include "XmlNode.h"
+
+
+UDMXMVRGroupObjectNode::UDMXMVRGroupObjectNode()
+{
+	ChildListNode = CreateDefaultSubobject<UDMXMVRChildListNode>("ChildListNode");
+}
+
+void UDMXMVRGroupObjectNode::InitializeFromGroupObjectXmlNode(const FXmlNode& GroupObjectXmlNode)
+{
+	checkf(GroupObjectXmlNode.GetTag() == TEXT("groupobject"), TEXT("Trying to initialize a Group Object Node from Xml Node, but the Node's Tag is '%s' instead of 'GroupObject'."), *GroupObjectXmlNode.GetTag());
+	checkf(ChildListNode, TEXT("Default Subobject 'ChildListNode' is invalid, this is not expected."));
+
+	// Init self
+	constexpr TCHAR AttributeName_UUID[] = TEXT("UUID");
+	const FString UUIDString = GroupObjectXmlNode.GetAttribute(AttributeName_UUID);
+	if (!FGuid::Parse(UUIDString, UUID))
+	{
+		UE_LOG(LogDMXRuntime, Warning, TEXT("Found invalid UUID of Group Object when parsing MVR File. Skipping Group Object."))
+		return;
+	}
+	constexpr TCHAR AttributeName_Name[] = TEXT("Name");
+	const FString NameString = GroupObjectXmlNode.GetAttribute(AttributeName_Name);
+	if (!NameString.IsEmpty())
+	{
+		Name.Value = NameString;
+	}
+
+	static const FString AttributeName_Matrix = TEXT("Matrix");
+	const FString MatrixString = GroupObjectXmlNode.GetAttribute(AttributeName_Matrix);
+	Matrix = FDMXRuntimeUtils::ParseGDTFMatrix(MatrixString);
+
+	// Init Child List Node
+	static const FString NodeName_ChildList = TEXT("ChildList");
+	const FXmlNode* ChildListXmlNode = GroupObjectXmlNode.FindChildNode(NodeName_ChildList);
+
+	if (ChildListXmlNode)
+	{
+		ChildListNode->InitializeFromChildListXmlNode(*ChildListXmlNode);
+	}
+}
+
+void UDMXMVRGroupObjectNode::CreateXmlNodeInParent(FXmlNode& ParentNode) const
+{
+	checkf(ParentNode.GetTag() == TEXT("childlist"), TEXT("Group Object Nodes have to be created in a Child List node, but parent node is %s."), *ParentNode.GetTag());
+	
+	// Don't export empty group objects
+	if (!ChildListNode || ChildListNode->GetParametricObjectNodes().IsEmpty())
+	{
+		return;
+	}
+
+	constexpr TCHAR GroupObjectTag[] = TEXT("GroupObject");
+	ParentNode.AppendChildNode(GroupObjectTag);
+	FXmlNode* GroupObjectXmlNode = ParentNode.GetChildrenNodes().Last();
+	check(GroupObjectXmlNode);
+
+	if (ChildListNode)
+	{
+		ChildListNode->CreateXmlNodeInParent(*GroupObjectXmlNode);
+	}
+}
+
+void UDMXMVRGroupObjectNode::GetFixtureNodes(TArray<UDMXMVRFixtureNode*>& OutFixtureNodes) const
+{
+	checkf(ChildListNode, TEXT("Default Subobject 'ChildListNode' is invalid, this is not expected."));
+
+	ChildListNode->GetFixtureNodes(OutFixtureNodes);
+}
+
+bool UDMXMVRGroupObjectNode::Contains(UDMXMVRParametricObjectNodeBase* ParametricObjectNode)
+{
+	checkf(ChildListNode, TEXT("Default Subobject 'ChildListNode' is invalid, this is not expected."));
+
+	for (UDMXMVRParametricObjectNodeBase* Other : ChildListNode->GetParametricObjectNodes())
+	{
+		if (UDMXMVRGroupObjectNode* GroupObjectNode = Cast<UDMXMVRGroupObjectNode>(Other))
+		{
+			if (GroupObjectNode->Contains(ParametricObjectNode))
+			{
+				return true;
+			}
+		}
+
+		if (Other == ParametricObjectNode)
+		{
+			return true;
+		}
+	}
+
+	return false;
+}
+
+UDMXMVRParametricObjectNodeBase** UDMXMVRGroupObjectNode::FindParametricObjectNodeByUUID(const FGuid& InUUID)
+{
+	checkf(ChildListNode, TEXT("Default Subobject 'ChildListNode' is invalid, this is not expected."));
+
+	return ChildListNode->FindParametricObjectNodeByUUID(InUUID);
+}
+
+bool UDMXMVRGroupObjectNode::RemoveParametricObjectNode(UDMXMVRParametricObjectNodeBase* ParametricObjectNode)
+{
+	checkf(ChildListNode, TEXT("Default Subobject 'ChildListNode' is invalid, this is not expected."));
+
+	return ChildListNode->RemoveParametricObject(ParametricObjectNode);
+}
diff --git a/Engine/Plugins/VirtualProduction/DMX/DMXEngine/Source/DMXRuntime/Private/MVR/Types/DMXMVRLayerNode.cpp b/Engine/Plugins/VirtualProduction/DMX/DMXEngine/Source/DMXRuntime/Private/MVR/Types/DMXMVRLayerNode.cpp
new file mode 100644
--- /dev/null
+++ b/Engine/Plugins/VirtualProduction/DMX/DMXEngine/Source/DMXRuntime/Private/MVR/Types/DMXMVRLayerNode.cpp
@@ -0,0 +1,133 @@
+// Copyright Epic Games, Inc. All Rights Reserved.
+
+#include "MVR/Types/DMXMVRLayerNode.h"
+
+#include "DMXRuntimeLog.h"
+#include "DMXRuntimeUtils.h"
+#include "MVR/Types/DMXMVRChildListNode.h"
+#include "MVR/Types/DMXMVRFixtureNode.h"
+#include "MVR/Types/DMXMVRParametricObjectNodeBase.h"
+
+#include "XmlNode.h"
+
+
+UDMXMVRLayerNode::UDMXMVRLayerNode()
+{
+	ChildListNode = CreateDefaultSubobject<UDMXMVRChildListNode>("ChildListNode");
+}
+
+void UDMXMVRLayerNode::InitializeFromLayerXmlNode(const FXmlNode& LayerXmlNode)
+{
+	checkf(ChildListNode, TEXT("Default Subobject 'ChildListNode' is invalid, this is not expected."));
+	checkf(LayerXmlNode.GetTag() == TEXT("layer"), TEXT("Trying to initialize a Layer Node from Xml Node, but the Node's Tag is '%s' instead of 'Layer'."), *LayerXmlNode.GetTag());
+
+	// Init self
+	constexpr TCHAR AttributeName_UUID[] = TEXT("UUID");
+	const FString UUIDString = LayerXmlNode.GetAttribute(AttributeName_UUID);
+	if (!FGuid::Parse(UUIDString, UUID))
+	{
+		UE_LOG(LogDMXRuntime, Warning, TEXT("Found invalid UUID of Layer when parsing MVR File. Skipping Layer."))
+		return;
+	}
+	constexpr TCHAR AttributeName_Name[] = TEXT("Name");
+	const FString NameString = LayerXmlNode.GetAttribute(AttributeName_Name);
+	if (!NameString.IsEmpty())
+	{
+		Name.Value = NameString;
+	}
+	
+	// Init Matrix
+	static const FString AttributeName_Matrix = TEXT("Matrix");
+	const FString MatrixString = LayerXmlNode.GetAttribute(AttributeName_Matrix);
+	Matrix = FDMXRuntimeUtils::ParseGDTFMatrix(MatrixString);
+
+	// Init Child List Node
+	static const FString NodeName_ChildList = TEXT("ChildList");
+	const FXmlNode* ChildListXmlNode = LayerXmlNode.FindChildNode(NodeName_ChildList);
+
+	if (ChildListXmlNode)
+	{
+		ChildListNode->InitializeFromChildListXmlNode(*ChildListXmlNode);
+	}
+}
+
+void UDMXMVRLayerNode::CreateXmlNodeInParent(FXmlNode& ParentNode) const
+{
+	checkf(ChildListNode, TEXT("Default Subobject 'ChildListNode' is invalid, this is not expected."));
+	checkf(ParentNode.GetTag() == TEXT("layers"), TEXT("Layer Nodes have to be created in a Layers node, but parent node is %s."), *ParentNode.GetTag());
+
+	TArray<FXmlAttribute> Attributes;
+
+	// Name attribute
+	if (Name.IsSet())
+	{
+		constexpr TCHAR AttributeName[] = TEXT("name");
+		FXmlAttribute Attribute = FXmlAttribute(AttributeName, Name.GetValue());
+		Attributes.Add(Attribute);
+	}
+	const FString NonOptionalName = Name.IsSet() ? Name.GetValue() : TEXT("Unnamed"); // Hold a name for logging
+
+	// UUID attribute
+	{
+		if (!UUID.IsValid())
+		{
+			UE_LOG(LogDMXRuntime, Warning, TEXT("Cannot export MVR Layer '%s'. Invalid UUID specified."), *NonOptionalName);
+			return;
+		}
+
+		constexpr TCHAR AttributeName[] = TEXT("uuid");
+		FXmlAttribute Attribute = FXmlAttribute(AttributeName, UUID.ToString(EGuidFormats::DigitsWithHyphens));
+		Attributes.Add(Attribute);
+	}
+
+	// Add self
+	constexpr TCHAR Tag[] = TEXT("Layer");
+	ParentNode.AppendChildNode(Tag);
+	FXmlNode* LayerXmlNode = ParentNode.GetChildrenNodes().Last();
+	check(LayerXmlNode);
+
+	// Add Matrix child Node
+	if (Matrix.IsSet())
+	{
+		constexpr TCHAR MatrixTag[] = TEXT("Matrix");
+		const FString MatrixContent = FDMXRuntimeUtils::ConvertTransformToGDTF4x3MatrixString(Matrix.GetValue());
+		ParentNode.AppendChildNode(MatrixTag, MatrixContent);
+		FXmlNode* MatrixXmlNode = ParentNode.GetChildrenNodes().Last();
+		check(MatrixXmlNode);
+	}
+
+	// Add Child List child Node
+	if (ChildListNode)
+	{
+		ChildListNode->CreateXmlNodeInParent(*LayerXmlNode);
+	}
+}
+
+void UDMXMVRLayerNode::GetFixtureNodes(TArray<UDMXMVRFixtureNode*>& OutFixtureNodes) const
+{
+	checkf(ChildListNode, TEXT("Unexpected: ChildListNode is instanced, but its pointer is invalid."));
+	ChildListNode->GetFixtureNodes(OutFixtureNodes);
+}
+
+bool UDMXMVRLayerNode::Contains(UDMXMVRParametricObjectNodeBase* ParametricObjectNode)
+{
+	checkf(ChildListNode, TEXT("Unexpected: ChildListNode is instanced, but its pointer is invalid."));
+	if (ChildListNode->Contains(ParametricObjectNode))
+	{
+		return true;
+	}
+
+	return false;
+}
+
+UDMXMVRParametricObjectNodeBase** UDMXMVRLayerNode::FindParametricObjectNodeByUUID(const FGuid& InUUID) const
+{
+	checkf(ChildListNode, TEXT("Unexpected: ChildListNode is instanced, but its pointer is invalid."));
+	UDMXMVRParametricObjectNodeBase** ParametricObjectNodePtr = ChildListNode->FindParametricObjectNodeByUUID(InUUID);
+	if (ParametricObjectNodePtr)
+	{
+		return ParametricObjectNodePtr;
+	}
+
+	return nullptr;
+}
diff --git a/Engine/Plugins/VirtualProduction/DMX/DMXEngine/Source/DMXRuntime/Private/MVR/Types/DMXMVRLayersNode.cpp b/Engine/Plugins/VirtualProduction/DMX/DMXEngine/Source/DMXRuntime/Private/MVR/Types/DMXMVRLayersNode.cpp
new file mode 100644
--- /dev/null
+++ b/Engine/Plugins/VirtualProduction/DMX/DMXEngine/Source/DMXRuntime/Private/MVR/Types/DMXMVRLayersNode.cpp
@@ -0,0 +1,73 @@
+// Copyright Epic Games, Inc. All Rights Reserved.
+
+#include "MVR/Types/DMXMVRLayersNode.h"
+
+#include "MVR/Types/DMXMVRChildListNode.h"
+#include "MVR/Types/DMXMVRLayerNode.h"
+#include "MVR/Types/DMXMVRParametricObjectNodeBase.h"
+
+#include "XmlNode.h"
+
+
+void UDMXMVRLayersNode::InitializeFromLayersXmlNode(const FXmlNode& LayersXmlNode)
+{
+	checkf(LayersXmlNode.GetTag() == TEXT("layers"), TEXT("Trying to initialize a Layers Node from Xml Node, but the Node's Tag is '%s' instead of 'Layers'."), *LayersXmlNode.GetTag());
+
+	// Init Layer Nodes
+	LayerNodes.Reset();
+	for (const FXmlNode* LayerXmlNode : LayersXmlNode.GetChildrenNodes())
+	{
+		UDMXMVRLayerNode* NewLayer = CreateLayer();
+		NewLayer->InitializeFromLayerXmlNode(*LayerXmlNode);
+	}
+}
+
+void UDMXMVRLayersNode::CreateXmlNodeInParent(FXmlNode& ParentNode) const
+{
+	checkf(ParentNode.GetTag() == TEXT("scene"), TEXT("The Layers Node has to be created in a Scene node, but parent node is %s."), *ParentNode.GetTag());
+
+	constexpr TCHAR Tag[] = TEXT("Layers");
+	ParentNode.AppendChildNode(Tag);
+	FXmlNode* LayersXmlNode = ParentNode.GetChildrenNodes().Last();
+	check(LayersXmlNode);
+
+	for (const UDMXMVRLayerNode* LayerNode : LayerNodes)
+	{
+		LayerNode->CreateXmlNodeInParent(*LayersXmlNode);
+	}
+}
+
+void UDMXMVRLayersNode::GetFixtureNodes(TArray<UDMXMVRFixtureNode*>& OutFixtureNodes) const
+{
+	for (UDMXMVRLayerNode* Layer : LayerNodes)
+	{
+		Layer->GetFixtureNodes(OutFixtureNodes);
+	}
+}
+
+UDMXMVRParametricObjectNodeBase** UDMXMVRLayersNode::FindParametricObjectNodeByUUID(const FGuid& UUID) const
+{
+	for (UDMXMVRLayerNode* Layer : LayerNodes)
+	{
+		UDMXMVRParametricObjectNodeBase** ParametricObjectNodePtr = Layer->FindParametricObjectNodeByUUID(UUID);
+		if (ParametricObjectNodePtr)
+		{
+			return ParametricObjectNodePtr;
+		}
+	}
+	return nullptr;
+}
+
+UDMXMVRLayerNode* UDMXMVRLayersNode::CreateLayer()
+{
+	UDMXMVRLayerNode* NewLayerNode = NewObject<UDMXMVRLayerNode>(this);
+	LayerNodes.Add(NewLayerNode);
+
+	return NewLayerNode;
+}
+
+void UDMXMVRLayersNode::RemoveLayer(UDMXMVRLayerNode* LayerNodeToRemove)
+{
+	const int32 NumRemovedElements = LayerNodes.RemoveSingle(LayerNodeToRemove);
+	ensureMsgf(NumRemovedElements == 1, TEXT("Tried to remove Layer from MVR Layers Node, but the Layer was not contained in the Layers Node."));
+}
diff --git a/Engine/Plugins/VirtualProduction/DMX/DMXEngine/Source/DMXRuntime/Private/MVR/Types/DMXMVRParametricObjectNodeBase.cpp b/Engine/Plugins/VirtualProduction/DMX/DMXEngine/Source/DMXRuntime/Private/MVR/Types/DMXMVRParametricObjectNodeBase.cpp
new file mode 100644
--- /dev/null
+++ b/Engine/Plugins/VirtualProduction/DMX/DMXEngine/Source/DMXRuntime/Private/MVR/Types/DMXMVRParametricObjectNodeBase.cpp
@@ -0,0 +1,88 @@
+// Copyright Epic Games, Inc. All Rights Reserved.
+
+#include "MVR/Types/DMXMVRParametricObjectNodeBase.h"
+
+#include "MVR/Types/DMXMVRGroupObjectNode.h"
+#include "MVR/Types/DMXMVRLayerNode.h"
+
+
+UDMXMVRLayerNode* UDMXMVRParametricObjectNodeBase::GetLayer()
+{
+	return const_cast<UDMXMVRLayerNode*>(const_cast<const UDMXMVRParametricObjectNodeBase*>(this)->GetLayer());
+}
+
+const UDMXMVRLayerNode* UDMXMVRParametricObjectNodeBase::GetLayer() const
+{
+	// Traverse up the Outer until we find the Layer
+	for (const UObject* Object = this; Object; Object = Object->GetOuter())
+	{
+		if (Object->GetClass() == UDMXMVRLayerNode::StaticClass())
+		{
+			return Cast<UDMXMVRLayerNode>(Object);
+		}
+		Object = Object->GetOuter();
+	}
+
+	ensureMsgf(0, TEXT("Unexpected: No outer Layer found for Parametric Object Node."));
+	return nullptr;
+}
+
+UDMXMVRGroupObjectNode* UDMXMVRParametricObjectNodeBase::GetGroup()
+{
+	return const_cast<UDMXMVRGroupObjectNode*>(const_cast<const UDMXMVRParametricObjectNodeBase*>(this)->GetGroup());
+}
+
+const UDMXMVRGroupObjectNode* UDMXMVRParametricObjectNodeBase::GetGroup() const
+{
+	// Traverse up the Outer until we find the Layer
+	for (const UObject* Object = this; Object; Object = Object->GetOuter())
+	{
+		if (Object->GetClass() == UDMXMVRLayerNode::StaticClass())
+		{
+			return Cast<UDMXMVRGroupObjectNode>(Object);
+		}
+		Object = Object->GetOuter();
+	}
+
+	return nullptr;
+}
+
+FTransform UDMXMVRParametricObjectNodeBase::GetTransformAbsolute() const
+{
+	FTransform Result = Matrix.IsSet() ? Matrix.GetValue() : FTransform::Identity;
+
+	for (const UDMXMVRGroupObjectNode* Group = GetGroup(); Group; Group = Group->GetGroup())
+	{
+		if (Group->Matrix.IsSet())
+		{
+			Result *= Group->Matrix.GetValue();
+		}
+	}
+
+	const UDMXMVRLayerNode* Layer = GetLayer();
+	if (Layer && Layer->Matrix.IsSet())
+	{
+		Result *= Layer->Matrix.GetValue();
+	}
+
+	return Result;
+}
+
+void UDMXMVRParametricObjectNodeBase::SetTransformAbsolute(FTransform NewTransform)
+{
+	for (UDMXMVRGroupObjectNode* Group = GetGroup(); Group; Group = Group->GetGroup())
+	{
+		if (Group->Matrix.IsSet())
+		{
+			NewTransform = NewTransform.Inverse() * Group->Matrix.GetValue();
+		}
+	}
+
+	UDMXMVRLayerNode* Layer = GetLayer();
+	if (Layer && Layer->Matrix.IsSet())
+	{
+		NewTransform = NewTransform.Inverse() * Layer->Matrix.GetValue();
+	}
+
+	Matrix = NewTransform;
+}
diff --git a/Engine/Plugins/VirtualProduction/DMX/DMXEngine/Source/DMXRuntime/Private/MVR/Types/DMXMVRRootNode.cpp b/Engine/Plugins/VirtualProduction/DMX/DMXEngine/Source/DMXRuntime/Private/MVR/Types/DMXMVRRootNode.cpp
new file mode 100644
--- /dev/null
+++ b/Engine/Plugins/VirtualProduction/DMX/DMXEngine/Source/DMXRuntime/Private/MVR/Types/DMXMVRRootNode.cpp
@@ -0,0 +1,180 @@
+// Copyright Epic Games, Inc. All Rights Reserved.
+
+#include "MVR/Types/DMXMVRRootNode.h"
+
+#include "DMXRuntimeLog.h"
+#include "MVR/DMXMVRVersion.h"
+#include "MVR/Types/DMXMVRChildListNode.h"
+#include "MVR/Types/DMXMVRFixtureNode.h"
+#include "MVR/Types/DMXMVRLayerNode.h"
+#include "MVR/Types/DMXMVRLayersNode.h"
+#include "MVR/Types/DMXMVRSceneNode.h"
+#include "MVR/Types/DMXMVRUserDataNode.h"
+
+#include "XmlFile.h"
+#include "XmlNode.h"
+#include "Misc/MessageDialog.h"
+
+
+UDMXMVRRootNode::UDMXMVRRootNode()
+{
+	UserDataNode = CreateDefaultSubobject<UDMXMVRUserDataNode>("UserDataNode");
+	SceneNode = CreateDefaultSubobject<UDMXMVRSceneNode>("SceneNode");
+}
+
+bool UDMXMVRRootNode::InitializeFromGeneralSceneDescriptionXml(const TSharedRef<FXmlFile>& GeneralSceneDescriptionXml)
+{
+	checkf(SceneNode, TEXT("Default Subobject 'SceneNode' is invalid, this is not expected."));
+
+	// Init self
+	FXmlNode* RootNode = GeneralSceneDescriptionXml->GetRootNode();
+	if (!RootNode)
+	{
+		UE_LOG(LogDMXRuntime, Warning, TEXT("Cannot parse MVR file. Root node is invalid."));
+		return false;
+	}
+
+	constexpr TCHAR VerMajorAttributeName[] = TEXT("verMajor");
+	const FString VerMajorString = RootNode->GetAttribute(VerMajorAttributeName);
+	int32 VerMajor;
+	if (!LexTryParseString<int32>(VerMajor, *VerMajorString))
+	{
+		UE_LOG(LogDMXRuntime, Warning, TEXT("Invalid MVR GeneralSceneDescription.xml. Major Version is not set."));
+		return false;
+	}
+
+	constexpr TCHAR VerMinorAttributeName[] = TEXT("verMinor");
+	const FString VerMinorString = RootNode->GetAttribute(VerMinorAttributeName);
+	int32 VerMinor;
+	if (!LexTryParseString<int32>(VerMinor, *VerMinorString))
+	{
+		UE_LOG(LogDMXRuntime, Warning, TEXT("Invalid MVR GeneralSceneDescription.xml. Minor Version is not set."));
+		return false;
+	}
+
+	if (FDMXMVRVersion::MajorVersion < VerMajor ||
+		(FDMXMVRVersion::MajorVersion == VerMajor && FDMXMVRVersion::MinorVersion < VerMinor))
+	{
+		const FText VerMajorText = FText::FromString(VerMajorString);
+		const FText VerMinorText = FText::FromString(VerMinorString);
+		const FText EngineMajorVerText = FText::FromString(FDMXMVRVersion::GetMajorVersionAsString());
+		const FText EngineMinorVerText = FText::FromString(FDMXMVRVersion::GetMinorVersionAsString());
+
+		const EAppReturnType::Type DialogResult = FMessageDialog::Open(EAppMsgType::YesNo, FText::Format(NSLOCTEXT("DMXMVRRootNode", "TryLoadMVRWithVersionNewThanEngineDialog", "Warning: Version '{0}.{1}' of MVR is newer than the MVR Version supported by the Engine, '{2}.{3}'. Do you want to try to load the MVR anyways (not recommended)?"), VerMajorText, VerMinorText, EngineMajorVerText, EngineMinorVerText));
+		if (DialogResult == EAppReturnType::No)
+		{
+			return false;
+		}
+	}
+
+	// Init Scene Node
+	static const FString NodeName_Scene = TEXT("Scene");
+	FXmlNode* SceneXmlNode = RootNode->FindChildNode(NodeName_Scene);
+	if(SceneXmlNode)
+	{
+		SceneNode->InitializeFromSceneXmlNode(*SceneXmlNode);
+	}
+
+	return true;
+}
+
+TSharedPtr<FXmlFile> UDMXMVRRootNode::CreateXmlFile()
+{
+	checkf(UserDataNode, TEXT("Default Subobject 'UserDataNode' is invalid, this is not expected."));
+	checkf(SceneNode, TEXT("Default Subobject 'SceneNode' is invalid, this is not expected."));
+
+	const FString Buffer = "<?xml version=\"1.0\" encoding=\"UTF - 8\" standalone=\"no\" ?>\n<GeneralSceneDescription>\n</GeneralSceneDescription>";
+
+	TSharedRef<FXmlFile> XmlFile = MakeShared<FXmlFile>();
+	const bool bCreatedNewFile = XmlFile->LoadFile(Buffer, EConstructMethod::ConstructFromBuffer);
+
+	FXmlNode* RootNode = XmlFile->GetRootNode();
+	if (!ensureAlwaysMsgf(bCreatedNewFile && RootNode, TEXT("Unexpected failed to create a new MVR General Scene Description in memory. MVR Export failed.")))
+	{
+		return nullptr;
+	}
+	check(RootNode);
+
+	// Version the Root Node
+	constexpr TCHAR VerMajorAttributeName[] = TEXT("verMajor");
+	constexpr TCHAR VerMinorAttributeName[] = TEXT("verMinor");
+	TArray<FXmlAttribute> Attributes =
+	{
+		FXmlAttribute(VerMajorAttributeName, FDMXMVRVersion::GetMajorVersionAsString()),
+		FXmlAttribute(VerMinorAttributeName, FDMXMVRVersion::GetMinorVersionAsString())
+	};
+	RootNode->SetAttributes(Attributes);
+	
+	// Export Children
+	if (UserDataNode)
+	{
+		UserDataNode->CreateXmlNodeInParent(*RootNode);
+	}
+
+	if (SceneNode)
+	{
+		SceneNode->CreateXmlNodeInParent(*RootNode);
+	}
+
+	return XmlFile;
+}
+
+void UDMXMVRRootNode::GetFixtureNodes(TArray<UDMXMVRFixtureNode*>& OutFixtureNodes) const
+{
+	checkf(SceneNode, TEXT("Default Subobject 'SceneNode' is invalid, this is not expected."));
+	UDMXMVRLayersNode* LayersNode = SceneNode ? SceneNode->GetLayersNode() : nullptr;
+	checkf(LayersNode, TEXT("Cannot Find Layer node in Scene Node."));
+	LayersNode->GetFixtureNodes(OutFixtureNodes);
+}
+
+UDMXMVRChildListNode& UDMXMVRRootNode::GetOrCreateFirstChildListNode()
+{
+	checkf(SceneNode, TEXT("Default Subobject 'SceneNode' is invalid, this is not expected."));
+	UDMXMVRLayersNode* LayersNode = SceneNode ? SceneNode->GetLayersNode() : nullptr;
+	checkf(LayersNode, TEXT("Cannot Find Layers node in Scene Node. This is not expected."));
+
+	UDMXMVRLayerNode* LayerNode = nullptr;
+	if (LayersNode->GetLayerNodes().IsEmpty())
+	{
+		LayerNode = LayersNode->CreateLayer();
+	}
+	else
+	{
+		LayerNode = LayersNode->GetLayerNodes()[0];
+	}
+	check(LayerNode);
+
+	UDMXMVRChildListNode* ChildListNode = LayerNode->GetChildListNode();
+	checkf(ChildListNode, TEXT("Cannot Find Child List node in Layer Node. This is not expected."));
+	return *ChildListNode;
+}
+
+UDMXMVRParametricObjectNodeBase** UDMXMVRRootNode::FindParametricObjectNodeByUUID(const FGuid& UUID) const
+{
+	UDMXMVRLayersNode* LayersNode = SceneNode ? SceneNode->GetLayersNode() : nullptr;
+	if (LayersNode)
+	{
+		return LayersNode->FindParametricObjectNodeByUUID(UUID);
+	}
+	return nullptr;
+}
+
+bool UDMXMVRRootNode::RemoveParametricObjectNode(UDMXMVRParametricObjectNodeBase* ParametricObjectNode)
+{
+	checkf(SceneNode, TEXT("Default Subobject 'SceneNode' is invalid, this is not expected."));
+	UDMXMVRLayersNode* LayersNode = SceneNode ? SceneNode->GetLayersNode() : nullptr;
+	checkf(LayersNode, TEXT("Cannot Find Layers node in Scene Node. This is not expected."));
+
+	for (UDMXMVRLayerNode* Layer : LayersNode->GetLayerNodes())
+	{
+		if (UDMXMVRChildListNode* ChildList = Layer->GetChildListNode())
+		{
+			if (ChildList->RemoveParametricObject(ParametricObjectNode))
+			{
+				return true;
+			}
+		}
+	}
+
+	return false;
+}
diff --git a/Engine/Plugins/VirtualProduction/DMX/DMXEngine/Source/DMXRuntime/Private/MVR/Types/DMXMVRSceneNode.cpp b/Engine/Plugins/VirtualProduction/DMX/DMXEngine/Source/DMXRuntime/Private/MVR/Types/DMXMVRSceneNode.cpp
new file mode 100644
--- /dev/null
+++ b/Engine/Plugins/VirtualProduction/DMX/DMXEngine/Source/DMXRuntime/Private/MVR/Types/DMXMVRSceneNode.cpp
@@ -0,0 +1,44 @@
+// Copyright Epic Games, Inc. All Rights Reserved.
+
+#include "MVR/Types/DMXMVRSceneNode.h"
+
+#include "MVR/Types/DMXMVRLayersNode.h"
+
+#include "XmlNode.h"
+
+
+UDMXMVRSceneNode::UDMXMVRSceneNode()
+{
+	LayersNode = CreateDefaultSubobject<UDMXMVRLayersNode>("LayersNode");
+}
+
+void UDMXMVRSceneNode::InitializeFromSceneXmlNode(const FXmlNode& SceneXmlNode)
+{
+	checkf(LayersNode, TEXT("Default Subobject 'LayersNode' is invalid, this is not expected."));
+	checkf(SceneXmlNode.GetTag() == TEXT("scene"), TEXT("Trying to initialize a Scene Node from Xml Node, but the Node's Tag is '%s' instead of 'Scene'."), *SceneXmlNode.GetTag());
+
+	// Init Layers Node
+	static const FString LayersTag = TEXT("Layers");
+	const FXmlNode* LayersXmlNode = SceneXmlNode.FindChildNode(LayersTag);
+	if (LayersXmlNode)
+	{
+		LayersNode->InitializeFromLayersXmlNode(*LayersXmlNode);
+	}
+}
+
+void UDMXMVRSceneNode::CreateXmlNodeInParent(FXmlNode& ParentNode) const
+{
+	checkf(LayersNode, TEXT("Default Subobject 'LayersNode' is invalid, this is not expected."));
+	checkf(ParentNode.GetTag() == TEXT("generalscenedescription"), TEXT("The Scene Node has to be created in a General Scene Description node, but parent node is %s."), *ParentNode.GetTag());
+
+	constexpr TCHAR Tag[] = TEXT("Scene");
+	constexpr TCHAR Content[] = TEXT("Content");
+	ParentNode.AppendChildNode(Tag, Content);
+	FXmlNode* SceneXmlNode = ParentNode.GetChildrenNodes().Last();
+	check(SceneXmlNode);
+
+	if (LayersNode)
+	{
+		LayersNode->CreateXmlNodeInParent(*SceneXmlNode);
+	}
+}
diff --git a/Engine/Plugins/VirtualProduction/DMX/DMXEngine/Source/DMXRuntime/Private/MVR/Types/DMXMVRUnrealEngineDataNode.cpp b/Engine/Plugins/VirtualProduction/DMX/DMXEngine/Source/DMXRuntime/Private/MVR/Types/DMXMVRUnrealEngineDataNode.cpp
new file mode 100644
--- /dev/null
+++ b/Engine/Plugins/VirtualProduction/DMX/DMXEngine/Source/DMXRuntime/Private/MVR/Types/DMXMVRUnrealEngineDataNode.cpp
@@ -0,0 +1,37 @@
+// Copyright Epic Games, Inc. All Rights Reserved.
+
+#include "MVR/Types/DMXMVRUnrealEngineDataNode.h"
+
+#include "XmlNode.h"
+#include "Misc/EngineVersion.h"
+
+
+void UDMXMVRUnrealEngineDataNode::CreateXmlNodeInParent(FXmlNode& ParentNode) const
+{
+	checkf(ParentNode.GetTag() == TEXT("userdata"), TEXT("The Data Node has to be created in a General Scene Description node, but parent node is %s."), *ParentNode.GetTag());
+
+	constexpr TCHAR Tag[] = TEXT("Data");
+	ParentNode.AppendChildNode(Tag);
+	FXmlNode* DataXmlNode = ParentNode.GetChildrenNodes().Last();
+	check(DataXmlNode);
+
+	constexpr TCHAR ProviderAttributeName[] = TEXT("provider");
+	constexpr TCHAR VerAttributeName[] = TEXT("ver");
+	const FString VersionString =
+		FString::FromInt(FEngineVersion::Current().GetMajor()) +
+		TEXT(".") +
+		FString::FromInt(FEngineVersion::Current().GetMinor()) +
+		TEXT(".") +
+		FString::FromInt(FEngineVersion::Current().GetChangelist());
+
+	const TArray<FXmlAttribute> Attributes = 
+	{ 
+		FXmlAttribute(ProviderAttributeName, ProviderNameUnrealEngine),
+		FXmlAttribute(VerAttributeName, VersionString)
+	};
+	DataXmlNode->SetAttributes(Attributes);
+
+	// 'Date' child node
+	constexpr TCHAR ChildTag_Date[] = TEXT("Date");
+	DataXmlNode->AppendChildNode(ChildTag_Date, FDateTime::Now().ToString());
+}
diff --git a/Engine/Plugins/VirtualProduction/DMX/DMXEngine/Source/DMXRuntime/Private/MVR/Types/DMXMVRUserDataNode.cpp b/Engine/Plugins/VirtualProduction/DMX/DMXEngine/Source/DMXRuntime/Private/MVR/Types/DMXMVRUserDataNode.cpp
new file mode 100644
--- /dev/null
+++ b/Engine/Plugins/VirtualProduction/DMX/DMXEngine/Source/DMXRuntime/Private/MVR/Types/DMXMVRUserDataNode.cpp
@@ -0,0 +1,26 @@
+// Copyright Epic Games, Inc. All Rights Reserved.
+
+#include "MVR/Types/DMXMVRUserDataNode.h"
+
+#include "MVR/Types/DMXMVRUnrealEngineDataNode.h"
+
+#include "XmlNode.h"
+
+
+UDMXMVRUserDataNode::UDMXMVRUserDataNode()
+{
+	UnrealEngineDataNode = CreateDefaultSubobject<UDMXMVRUnrealEngineDataNode>("UnrealEngineDataNode");
+}
+
+void UDMXMVRUserDataNode::CreateXmlNodeInParent(FXmlNode& ParentNode) const
+{
+	checkf(UnrealEngineDataNode, TEXT("Default Subobject 'UnrealEngineDataNode' is invalid, this is not expected."));
+	checkf(ParentNode.GetTag() == TEXT("generalscenedescription"), TEXT("The User Data Node has to be created in the General Scene Description node, but parent node is %s."), *ParentNode.GetTag());
+
+	constexpr TCHAR Tag[] = TEXT("UserData");
+	ParentNode.AppendChildNode(Tag);
+	FXmlNode* UserDataXmlNode = ParentNode.GetChildrenNodes().Last();
+	check(UserDataXmlNode);
+
+	UnrealEngineDataNode->CreateXmlNodeInParent(*UserDataXmlNode);
+}
diff --git a/Engine/Plugins/VirtualProduction/DMX/DMXEngine/Source/DMXRuntime/Public/DMXOptionalTypes.h b/Engine/Plugins/VirtualProduction/DMX/DMXEngine/Source/DMXRuntime/Public/DMXOptionalTypes.h
new file mode 100644
--- /dev/null
+++ b/Engine/Plugins/VirtualProduction/DMX/DMXEngine/Source/DMXRuntime/Public/DMXOptionalTypes.h
@@ -0,0 +1,526 @@
+// Copyright Epic Games, Inc. All Rights Reserved.
+
+#pragma once
+
+#include "DMXTypes.h"
+
+#include "DMXOptionalTypes.generated.h"
+
+
+//////////////////////////////
+// Wrappers for TOptional<Types> common to DMX related formats such as GDTF and MVR.
+
+USTRUCT()
+struct DMXRUNTIME_API FDMXOptionalBool
+{
+	GENERATED_BODY()
+
+	/** The Optional Value */
+	TOptional<bool> Value;
+
+	/** Resets the optional */
+	void Reset() { Value.Reset(); }
+	
+	/** Returns true if the optional is set */
+	bool IsSet() const { return Value.IsSet(); }
+	
+	/** Returns the Value of the optional, only call when the value is set. */
+	const bool& GetValue() const { return Value.GetValue(); }
+
+	const FDMXOptionalBool& operator=(const FDMXOptionalBool& Rhs)
+	{
+		Value = Rhs.Value;
+		return *this;
+	}
+	const FDMXOptionalBool& operator=(const bool& Transform)
+	{
+		Value = Transform;
+		return *this;
+	}
+
+	friend bool operator==(const FDMXOptionalBool& A, const FDMXOptionalBool& B)
+	{
+		const bool bBothSet = A.Value.IsSet() && B.Value.IsSet();
+		if (bBothSet)
+		{
+			return A.Value.GetValue() == B.Value.GetValue();
+		}
+		else
+		{
+			return A.Value.IsSet() == B.Value.IsSet();
+		}
+	}
+
+	friend bool operator!=(const FDMXOptionalBool& A, const FDMXOptionalBool& B)
+	{
+		return !(A == B);
+	}
+
+	bool Serialize(FArchive& Ar)
+	{
+		Ar << Value;
+		return true;
+	}
+
+	friend FArchive& operator<<(FArchive& Ar, FDMXOptionalBool& Struct)
+	{
+		Ar << Struct.Value;
+		return Ar;
+	}
+};
+
+template<>
+struct TStructOpsTypeTraits<FDMXOptionalBool>
+	: public TStructOpsTypeTraitsBase2<FDMXOptionalBool>
+{
+	enum
+	{
+		WithSerializer = true,
+		WithCopy = true,
+		WithIdenticalViaEquality = true
+	};
+};
+
+USTRUCT()
+struct DMXRUNTIME_API FDMXOptionalInt32
+{
+	GENERATED_BODY()
+
+	/** The Optional Value */
+	TOptional<int32> Value;
+
+	/** Resets the optional */
+	void Reset() { Value.Reset(); }
+	
+	/** Returns true if the optional is set */
+	bool IsSet() const { return Value.IsSet(); }
+	
+	/** Returns the Value of the optional, only call when the value is set. */
+	const int32& GetValue() const { return Value.GetValue(); }
+
+	const FDMXOptionalInt32& operator=(const FDMXOptionalInt32& Rhs)
+	{
+		Value = Rhs.Value;
+		return *this;
+	}
+	const FDMXOptionalInt32& operator=(const int32& Transform)
+	{
+		Value = Transform;
+		return *this;
+	}
+
+	friend bool operator==(const FDMXOptionalInt32& A, const FDMXOptionalInt32& B)
+	{
+		const bool bBothSet = A.Value.IsSet() && B.Value.IsSet();
+		if (bBothSet)
+		{
+			return A.Value.GetValue() == B.Value.GetValue();
+		}
+		else
+		{
+			return A.Value.IsSet() == B.Value.IsSet();
+		}
+	}
+
+	friend bool operator!=(const FDMXOptionalInt32& A, const FDMXOptionalInt32& B)
+	{
+		return !(A == B);
+	}
+
+	bool Serialize(FArchive& Ar)
+	{
+		Ar << Value;
+		return true;
+	}
+	
+	FArchive& operator<<(FArchive& Ar)
+	{
+		Ar << Value;
+		return Ar;
+	}
+
+	friend FArchive& operator<<(FArchive& Ar, FDMXOptionalInt32& Struct)
+	{
+		Ar << Struct.Value;
+		return Ar;
+	}
+};
+
+template<>
+struct TStructOpsTypeTraits<FDMXOptionalInt32>
+	: public TStructOpsTypeTraitsBase2<FDMXOptionalInt32>
+{
+	enum
+	{
+		WithSerializer = true,
+		WithCopy = true,
+		WithIdenticalViaEquality = true
+	};
+};
+
+USTRUCT()
+struct DMXRUNTIME_API FDMXOptionalFloat
+{
+	GENERATED_BODY()
+
+	/** The Optional Value */
+	TOptional<float> Value;
+
+	/** Resets the optional */
+	void Reset() { Value.Reset(); }
+	
+	/** Returns true if the optional is set */
+	bool IsSet() const { return Value.IsSet(); }
+	
+	/** Returns the Value of the optional, only call when the value is set. */
+	const float& GetValue() const { return Value.GetValue(); }
+
+	const FDMXOptionalFloat& operator=(const FDMXOptionalFloat& Rhs)
+	{
+		Value = Rhs.Value;
+		return *this;
+	}
+	const FDMXOptionalFloat& operator=(const float& Transform)
+	{
+		Value = Transform;
+		return *this;
+	}
+
+	friend bool operator==(const FDMXOptionalFloat& A, const FDMXOptionalFloat& B)
+	{
+		const bool bBothSet = A.Value.IsSet() && B.Value.IsSet();
+		if (bBothSet)
+		{
+			return A.Value.GetValue() == B.Value.GetValue();
+		}
+		else
+		{
+			return A.Value.IsSet() == B.Value.IsSet();
+		}
+	}
+
+	friend bool operator!=(const FDMXOptionalFloat& A, const FDMXOptionalFloat& B)
+	{
+		return !(A == B);
+	}
+
+	bool Serialize(FArchive& Ar)
+	{
+		Ar << Value;
+		return true;
+	}
+
+	friend FArchive& operator<<(FArchive& Ar, FDMXOptionalFloat& Struct)
+	{
+		Ar << Struct.Value;
+		return Ar;
+	}
+};
+
+template<>
+struct TStructOpsTypeTraits<FDMXOptionalFloat>
+	: public TStructOpsTypeTraitsBase2<FDMXOptionalFloat>
+{
+	enum
+	{
+		WithSerializer = true,
+		WithCopy = true,
+		WithIdenticalViaEquality = true
+	};
+};
+
+USTRUCT()
+struct DMXRUNTIME_API FDMXOptionalGuid
+{
+	GENERATED_BODY()
+
+	/** The Optional Value */
+	TOptional<FGuid> Value;
+
+	/** Resets the optional */
+	void Reset() { Value.Reset(); }
+	
+	/** Returns true if the optional is set */
+	bool IsSet() const { return Value.IsSet(); }
+	
+	/** Returns the Value of the optional, only call when the value is set. */
+	const FGuid& GetValue() const { return Value.GetValue(); }
+
+	const FDMXOptionalGuid& operator=(const FDMXOptionalGuid& Rhs)
+	{
+		Value = Rhs.Value;
+		return *this;
+	}
+	const FDMXOptionalGuid& operator=(const FGuid& Transform)
+	{
+		Value = Transform;
+		return *this;
+	}
+
+	friend bool operator==(const FDMXOptionalGuid& A, const FDMXOptionalGuid& B)
+	{
+		const bool bBothSet = A.Value.IsSet() && B.Value.IsSet();
+		if (bBothSet)
+		{
+			return A.Value.GetValue() == B.Value.GetValue();
+		}
+		else
+		{
+			return A.Value.IsSet() == B.Value.IsSet();
+		}
+	}
+
+	friend bool operator!=(const FDMXOptionalGuid& A, const FDMXOptionalGuid& B)
+	{
+		return !(A == B);
+	}
+
+	bool Serialize(FArchive& Ar)
+	{
+		Ar << Value;
+		return true;
+	}
+
+	FArchive& operator<<(FArchive& Ar)
+	{
+		Ar << Value;
+		return Ar;
+	}
+
+	friend FArchive& operator<<(FArchive& Ar, FDMXOptionalGuid& Struct)
+	{
+		Ar << Struct.Value;
+		return Ar;
+	}
+};
+
+template<>
+struct TStructOpsTypeTraits<FDMXOptionalGuid>
+	: public TStructOpsTypeTraitsBase2<FDMXOptionalGuid>
+{
+	enum
+	{
+		WithSerializer = true,
+		WithCopy = true,
+		WithIdenticalViaEquality = true
+	};
+};
+
+USTRUCT()
+struct DMXRUNTIME_API FDMXOptionalTransform
+{
+	GENERATED_BODY()
+
+	/** The Optional Value */
+	TOptional<FTransform> Value;
+
+	/** Resets the optional */
+	void Reset() { Value.Reset(); }
+	
+	/** Returns true if the optional is set */
+	bool IsSet() const { return Value.IsSet(); }
+	
+	/** Returns the Value of the optional, only call when the value is set. */
+	const FTransform& GetValue() const { return Value.GetValue(); }
+
+	const FDMXOptionalTransform& operator=(const FDMXOptionalTransform& Rhs)
+	{
+		Value = Rhs.Value;
+		return *this;
+	}
+	const FDMXOptionalTransform& operator=(const FTransform& Transform)
+	{
+		Value = Transform;
+		return *this;
+	}
+
+	friend bool operator==(const FDMXOptionalTransform& A, const FDMXOptionalTransform& B)
+	{
+		const bool bBothSet = A.Value.IsSet() && B.Value.IsSet();
+		if (bBothSet)
+		{
+			return A.Value.GetValue().Equals(B.Value.GetValue(), 0.0);
+		}
+		else
+		{
+			return A.Value.IsSet() == B.Value.IsSet();
+		}
+	}
+
+	friend bool operator!=(const FDMXOptionalTransform& A, const FDMXOptionalTransform& B)
+	{
+		return !(A == B);
+	}
+
+	bool Serialize(FArchive& Ar)
+	{
+		Ar << Value;
+		return true;
+	}
+
+	FArchive& operator<<(FArchive& Ar)
+	{
+		Ar << Value;
+		return Ar;
+	}
+
+	friend FArchive& operator<<(FArchive& Ar, FDMXOptionalTransform& Struct)
+	{
+		Ar << Struct.Value;
+		return Ar;
+	}
+};
+	
+template<>
+struct TStructOpsTypeTraits<FDMXOptionalTransform>
+	: public TStructOpsTypeTraitsBase2<FDMXOptionalTransform>
+{
+	enum
+	{
+		WithSerializer = true,
+		WithCopy = true,
+		WithIdenticalViaEquality = true
+	};
+};
+
+USTRUCT()
+struct DMXRUNTIME_API FDMXOptionalColorCIE1931
+{
+	GENERATED_BODY()
+
+	/** The Optional Value */
+	TOptional<FDMXColorCIE1931xyY> Value;
+
+	/** Resets the optional */
+	void Reset() { Value.Reset(); }
+	
+	/** Returns true if the optional is set */
+	bool IsSet() const { return Value.IsSet(); }
+	
+	/** Returns the Value of the optional, only call when the value is set. */
+	const FDMXColorCIE1931xyY& GetValue() const { return Value.GetValue(); }
+
+	const FDMXOptionalColorCIE1931& operator=(const FDMXOptionalColorCIE1931& Rhs)
+	{
+		Value = Rhs.Value;
+		return *this;
+	}
+	const FDMXOptionalColorCIE1931& operator=(const FDMXColorCIE1931xyY& Transform)
+	{
+		Value = Transform;
+		return *this;
+	}
+
+	friend bool operator==(const FDMXOptionalColorCIE1931& A, const FDMXOptionalColorCIE1931& B)
+	{
+		const bool bBothSet = A.Value.IsSet() && B.Value.IsSet();
+		if (bBothSet)
+		{
+			return A.Value.GetValue() == B.Value.GetValue();
+		}
+		else
+		{
+			return A.Value.IsSet() == B.Value.IsSet();
+		}
+	}
+
+	friend bool operator!=(const FDMXOptionalColorCIE1931& A, const FDMXOptionalColorCIE1931& B)
+	{
+		return !(A == B);
+	}
+
+	bool Serialize(FArchive& Ar)
+	{
+		Ar << Value;
+		return true;
+	}
+
+	friend FArchive& operator<<(FArchive& Ar, FDMXOptionalColorCIE1931& Struct)
+	{
+		Ar << Struct.Value;
+		return Ar;
+	}
+};
+
+template<>
+struct TStructOpsTypeTraits<FDMXOptionalColorCIE1931>
+	: public TStructOpsTypeTraitsBase2<FDMXOptionalColorCIE1931>
+{
+	enum
+	{
+		WithSerializer = true,
+		WithCopy = true,
+		WithIdenticalViaEquality = true
+	};
+};
+
+USTRUCT()
+struct DMXRUNTIME_API FDMXOptionalString
+{
+	GENERATED_BODY()
+
+	/** The Optional Value */
+	TOptional<FString> Value;
+
+	/** Resets the optional */
+	void Reset() { Value.Reset(); }
+	
+	/** Returns true if the optional is set */
+	bool IsSet() const { return Value.IsSet(); }
+	
+	/** Returns the Value of the optional, only call when the value is set. */
+	const FString& GetValue() const { return Value.GetValue(); }
+
+	const FDMXOptionalString& operator=(const FDMXOptionalString& Rhs)
+	{
+		Value = Rhs.Value;
+		return *this;
+	}
+	const FDMXOptionalString& operator=(const FString& Transform)
+	{
+		Value = Transform;
+		return *this;
+	}
+
+	friend bool operator==(const FDMXOptionalString& A, const FDMXOptionalString& B)
+	{
+		const bool bBothSet = A.Value.IsSet() && B.Value.IsSet();
+		if (bBothSet)
+		{
+			return A.Value.GetValue() == B.Value.GetValue();
+		}
+		else
+		{
+			return A.Value.IsSet() == B.Value.IsSet();
+		}
+	}
+
+	friend bool operator!=(const FDMXOptionalString& A, const FDMXOptionalString& B)
+	{
+		return !(A == B);
+	}
+
+	bool Serialize(FArchive& Ar)
+	{
+		Ar << Value;
+		return true;
+	}
+
+	friend FArchive& operator<<(FArchive& Ar, FDMXOptionalString& Struct)
+	{
+		Ar << Struct.Value;
+		return Ar;
+	}
+};
+
+template<>
+struct TStructOpsTypeTraits<FDMXOptionalString>
+	: public TStructOpsTypeTraitsBase2<FDMXOptionalString>
+{
+	enum
+	{
+		WithSerializer = true,
+		WithCopy = true,
+		WithIdenticalViaEquality = true
+	};
+};
diff --git a/Engine/Plugins/VirtualProduction/DMX/DMXEngine/Source/DMXRuntime/Public/DMXRuntimeUtils.h b/Engine/Plugins/VirtualProduction/DMX/DMXEngine/Source/DMXRuntime/Public/DMXRuntimeUtils.h
--- a/Engine/Plugins/VirtualProduction/DMX/DMXEngine/Source/DMXRuntime/Public/DMXRuntimeUtils.h
+++ b/Engine/Plugins/VirtualProduction/DMX/DMXEngine/Source/DMXRuntime/Public/DMXRuntimeUtils.h
@@ -2,6 +2,7 @@
 
 #pragma once
 
+#include "DMXOptionalTypes.h"
 #include "Library/DMXEntityFixtureType.h"
 
 #include "CoreMinimal.h"
@@ -15,6 +16,12 @@ class UDMXLibrary;
 class DMXRUNTIME_API FDMXRuntimeUtils
 {
 public:
+	/** Parses a GDTF (or MVR) Transformation Matrix from a String. Result is optional and set if parsing succeeded. */
+	static FDMXOptionalTransform ParseGDTFMatrix(const FString& String);
+
+	/** Converts a UE transform to a GDTF style string */
+	static FString ConvertTransformToGDTF4x3MatrixString(FTransform Transform);
+
 	/**
 	 * Generates a unique name given a base one and a list of existing ones, by appending an index to
 	 * existing names. If InBaseName is an empty String, it returns "Default name".
diff --git a/Engine/Plugins/VirtualProduction/DMX/DMXEngine/Source/DMXRuntime/Public/DMXTypes.h b/Engine/Plugins/VirtualProduction/DMX/DMXEngine/Source/DMXRuntime/Public/DMXTypes.h
--- a/Engine/Plugins/VirtualProduction/DMX/DMXEngine/Source/DMXRuntime/Public/DMXTypes.h
+++ b/Engine/Plugins/VirtualProduction/DMX/DMXEngine/Source/DMXRuntime/Public/DMXTypes.h
@@ -12,7 +12,8 @@
 class UDMXLibrary;
 
 
-// Holds an array Attribute Names with their normalized Values (expand the property to see the map)
+
+/** Holds an array Attribute Names with their normalized Values (expand the property to see the map) */
 USTRUCT(BlueprintType, Category = "DMX")
 struct DMXRUNTIME_API FDMXNormalizedAttributeValueMap
 {
@@ -22,6 +23,103 @@ struct DMXRUNTIME_API FDMXNormalizedAttributeValueMap
 	TMap<FDMXAttributeName, float> Map;
 };
 
+
+/** xyY color representation in the CIE 1931 color space, as typically used in GDTF and MVR */
+USTRUCT(BlueprintType, Category = "DMX")
+struct DMXRUNTIME_API FDMXColorCIE1931xyY
+{
+	GENERATED_BODY()
+
+	FString ToString() const
+	{
+		return FString::FromInt(X) + TEXT(", ") + FString::FromInt(Y) + TEXT(", ") + FString::FromInt(YY);
+	}
+
+	/** x */
+	UPROPERTY(BlueprintReadWrite, EditAnywhere, Category = "DMX")
+	float X = 0.f;
+
+	/** y */
+	UPROPERTY(BlueprintReadWrite, EditAnywhere, Category = "DMX")
+	float Y = 0.f;
+
+	/** Y */
+	UPROPERTY(BlueprintReadWrite, EditAnywhere, Category = "DMX")
+	float YY = 0.f;
+
+	friend bool operator==(const FDMXColorCIE1931xyY& A, const FDMXColorCIE1931xyY& B)
+	{
+		return A.X == B.X && A.Y == B.Y && A.YY == B.YY;
+	}
+
+	friend bool operator!=(const FDMXColorCIE1931xyY& A, const FDMXColorCIE1931xyY& B)
+	{
+		return !(A == B);
+	}
+
+	friend FArchive& operator<<(FArchive& Ar, FDMXColorCIE1931xyY& ColorCIE1931)
+	{
+		Ar << ColorCIE1931.X;
+		Ar << ColorCIE1931.Y;
+		Ar << ColorCIE1931.YY;
+		return Ar;
+	}
+};
+
+USTRUCT()
+struct FDMXByteArray64
+{
+	GENERATED_BODY()
+
+	TArray64<uint8> ByteArray;
+
+	const FDMXByteArray64& operator=(const FDMXByteArray64& Rhs)
+	{
+		ByteArray = Rhs.ByteArray;
+		return *this;
+	}
+	const FDMXByteArray64& operator=(const TArray64<uint8>& Rhs)
+	{
+		ByteArray = Rhs;
+		return *this;
+	}
+
+	friend bool operator==(const FDMXByteArray64& A, const FDMXByteArray64& B)
+	{
+		return A.ByteArray == B.ByteArray;
+	}
+
+	friend bool operator!=(const FDMXByteArray64& A, const FDMXByteArray64& B)
+	{
+		return !(A == B);
+	}
+
+	bool Serialize(FArchive& Ar)
+	{
+		Ar << ByteArray;
+		return true;
+	}
+
+	friend FArchive& operator<<(FArchive& Ar, FDMXByteArray64& Struct)
+	{
+		Ar << Struct.ByteArray;
+		return Ar;
+	}
+};
+
+template<>
+struct TStructOpsTypeTraits<FDMXByteArray64>
+	: public TStructOpsTypeTraitsBase2<FDMXByteArray64>
+{
+	enum
+	{
+		WithSerializer = true,
+		WithCopy = true,
+		WithIdenticalViaEquality = true
+	};
+};
+
+
 USTRUCT(BlueprintType)
 struct FDMXRequestBase
 {
diff --git a/Engine/Plugins/VirtualProduction/DMX/DMXEngine/Source/DMXRuntime/Public/Library/DMXGDTFAssetImportData.h b/Engine/Plugins/VirtualProduction/DMX/DMXEngine/Source/DMXRuntime/Public/Library/DMXGDTFAssetImportData.h
--- a/Engine/Plugins/VirtualProduction/DMX/DMXEngine/Source/DMXRuntime/Public/Library/DMXGDTFAssetImportData.h
+++ b/Engine/Plugins/VirtualProduction/DMX/DMXEngine/Source/DMXRuntime/Public/Library/DMXGDTFAssetImportData.h
@@ -2,6 +2,8 @@
 
 #pragma once
 
+#include "DMXTypes.h"
+
 #include "CoreMinimal.h"
 #include "EditorFramework/AssetImportData.h"
 
@@ -24,8 +26,8 @@ public:
 	/** Creates new asset GDTF Asset Import Data from the File. Returns the GDTF Asset Import Data or nullptr if not a valid GDTF File */
 	void SetSourceFile(const FString& FilePathAndName);
 
-	/** Returns the raw zip file the asset was generated from */
-	FORCEINLINE const TArray<uint8>& GetRawZipFile() const { return RawZipFile; }
+	/** Returns the source data the asset was generated from */
+	FORCEINLINE const TArray64<uint8>& GetRawSourceData() const { return RawSourceData.ByteArray; }
 #endif // WITH_EDITOR 
 
 private:
@@ -35,6 +37,6 @@ private:
 #if WITH_EDITORONLY_DATA
 	/** The raw GDTF zip file as byte array */
 	UPROPERTY()
-	TArray<uint8> RawZipFile;
+	FDMXByteArray64 RawSourceData;
 #endif // WITH_EDITORONLY_DATA
 };
diff --git a/Engine/Plugins/VirtualProduction/DMX/DMXEngine/Source/DMXRuntime/Public/MVR/DMXMVRAssetImportData.h b/Engine/Plugins/VirtualProduction/DMX/DMXEngine/Source/DMXRuntime/Public/MVR/DMXMVRAssetImportData.h
--- a/Engine/Plugins/VirtualProduction/DMX/DMXEngine/Source/DMXRuntime/Public/MVR/DMXMVRAssetImportData.h
+++ b/Engine/Plugins/VirtualProduction/DMX/DMXEngine/Source/DMXRuntime/Public/MVR/DMXMVRAssetImportData.h
@@ -2,11 +2,14 @@
 
 #pragma once
 
+#include "DMXTypes.h"
+
 #include "CoreMinimal.h"
 #include "EditorFramework/AssetImportData.h"
 
 #include "DMXMVRAssetImportData.generated.h"
 
+class FDMXXmlFile;
 class UDMXMVRGeneralSceneDescription;
 
 
@@ -24,8 +27,8 @@ public:
 	/** Returns the Asset Import Data */
 	FORCEINLINE UAssetImportData* GetAssetImportData() { return AssetImportData;};
 
-	/** Returns the raw zip file the asset was generated from */
-	FORCEINLINE const TArray<uint8>& GetRawZipFile() const { return RawZipFile; }
+	/** Returns the raw source file as byte array, as it was imported */
+	FORCEINLINE const TArray64<uint8>& GetRawSourceData() const { return RawSourceData.ByteArray; }
 #endif // WITH_EDITOR 
 
 private:
@@ -34,8 +37,8 @@ private:
 	UPROPERTY()
 	UAssetImportData* AssetImportData;
 
-	/** The raw GDTF zip file as byte array */
+	/** The raw source file as byte array, as it was imported */
 	UPROPERTY()
-	TArray<uint8> RawZipFile;
+	FDMXByteArray64 RawSourceData;
 #endif // WITH_EDITORONLY_DATA
 };
diff --git a/Engine/Plugins/VirtualProduction/DMX/DMXEngine/Source/DMXRuntime/Public/MVR/DMXMVRFixture.h b/Engine/Plugins/VirtualProduction/DMX/DMXEngine/Source/DMXRuntime/Public/MVR/DMXMVRFixture.h
deleted file mode 100644
diff --git a/Engine/Plugins/VirtualProduction/DMX/DMXEngine/Source/DMXRuntime/Public/MVR/DMXMVRGeneralSceneDescription.h b/Engine/Plugins/VirtualProduction/DMX/DMXEngine/Source/DMXRuntime/Public/MVR/DMXMVRGeneralSceneDescription.h
--- a/Engine/Plugins/VirtualProduction/DMX/DMXEngine/Source/DMXRuntime/Public/MVR/DMXMVRGeneralSceneDescription.h
+++ b/Engine/Plugins/VirtualProduction/DMX/DMXEngine/Source/DMXRuntime/Public/MVR/DMXMVRGeneralSceneDescription.h
@@ -4,18 +4,17 @@
 
 #include "CoreMinimal.h"
 
-#include "MVR/DMXMVRFixture.h"
-
 #include "Misc/Optional.h"
 #include "Serialization/Archive.h"
 
 #include "DMXMVRGeneralSceneDescription.generated.h"
 
+class FXmlFile;
 class UDMXEntityFixturePatch;
 class UDMXLibrary;
 class UDMXMVRAssetImportData;
-
-class FXmlFile;
+class UDMXMVRRootNode;
+class UDMXMVRFixtureNode;
 
 
 /** MVR General Scene Description Object */
@@ -29,39 +28,32 @@ public:
 	/** Constructor */
 	UDMXMVRGeneralSceneDescription();
 
+	/** Gets the MVR Fixture Nodes in this General Scene Description */
+	void GetFixtureNodes(TArray<UDMXMVRFixtureNode*>& OutFixtureNodes) const;
+
+	/** Returns the Fixture Node corresponding to the UUID, or nullptr if it cannot be found */
+	UDMXMVRFixtureNode* FindFixtureNode(const FGuid& FixtureUUID) const;
+
 #if WITH_EDITOR
 	/** Creates an MVR General Scene Description from an Xml File */
 	static UDMXMVRGeneralSceneDescription* CreateFromXmlFile(TSharedRef<FXmlFile> GeneralSceneDescriptionXml, UObject* Outer, FName Name, EObjectFlags Flags = RF_NoFlags);
 
 	/** Creates an MVR General Scene Description from a DMX Library */
 	static UDMXMVRGeneralSceneDescription* CreateFromDMXLibrary(const UDMXLibrary& DMXLibrary, UObject* Outer, FName Name, EObjectFlags Flags = RF_NoFlags);
-#endif 
 
-	/** Returns a pointer to the MVR Fixture with corresponding UUID or nullptr if not found */
-	FDMXMVRFixture* FindMVRFixture(const FGuid& MVRFixtureUUID);
-
-	/** Adds an MVR Fixture to the General Scene Description */
-	void AddMVRFixture(FDMXMVRFixture& MVRFixture);
-
-#if WITH_EDITOR
 	/**
 	 * Writes the Library to the General Scene Description, effectively removing inexisting and adding
 	 * new MVR Fixtures, according to what MVR Fixture UUIDs the Fixture Patches of the Library contain.
 	 */
 	void WriteDMXLibraryToGeneralSceneDescription(const UDMXLibrary& DMXLibrary);
 
-	/** Returns the MVR Fixtures of the General Scene Description */
-	FORCEINLINE const TArray<FDMXMVRFixture>& GetMVRFixtures() const { return MVRFixtures; }
+	/** Creates an General Scene Description Xml File from this. */
+	TSharedPtr<FXmlFile> CreateXmlFile() const;
 
 	/** Returns MVR Asset Import Data for this asset */
 	FORCEINLINE UDMXMVRAssetImportData* GetMVRAssetImportData() const { return MVRAssetImportData; }
 #endif
 
-protected:
-	//~ Begin UObject interface
-	virtual void Serialize(FArchive& Ar) override;
-	//~ End UObject interface
-
 private:
 #if WITH_EDITOR
 	/** Writes the Fixture Patch to the General Scene Description, adding a new MVR Fixture if required */
@@ -74,8 +66,9 @@ private:
 	/** Returns the Unit Numbers currently in use, sorted from lowest to highest */
 	TArray<int32> GetUnitNumbersInUse(const UDMXLibrary& DMXLibrary);
 
-	/** The MVR Fixtures of the General Scene Description */
-	TArray<FDMXMVRFixture> MVRFixtures;
+	/** The Root Node of the General Scene Description */
+	UPROPERTY()
+	UDMXMVRRootNode* RootNode;
 
 #if WITH_EDITORONLY_DATA
 	/** Import Data for this asset */
diff --git a/Engine/Plugins/VirtualProduction/DMX/DMXEngine/Source/DMXRuntime/Public/MVR/DMXMVRUnzip.h b/Engine/Plugins/VirtualProduction/DMX/DMXEngine/Source/DMXRuntime/Public/MVR/DMXMVRUnzip.h
deleted file mode 100644
diff --git a/Engine/Plugins/VirtualProduction/DMX/DMXEngine/Source/DMXRuntime/Public/MVR/DMXMVRVersion.h b/Engine/Plugins/VirtualProduction/DMX/DMXEngine/Source/DMXRuntime/Public/MVR/DMXMVRVersion.h
new file mode 100644
--- /dev/null
+++ b/Engine/Plugins/VirtualProduction/DMX/DMXEngine/Source/DMXRuntime/Public/MVR/DMXMVRVersion.h
@@ -0,0 +1,25 @@
+// Copyright Epic Games, Inc. All Rights Reserved.
+
+#pragma once
+
+#include "CoreMinimal.h"
+
+
+/** Version of the MVR standard currently being used by the engine */
+struct DMXRUNTIME_API FDMXMVRVersion
+{
+	static const int32 MajorVersion = 1;
+
+	static const int32 MinorVersion = 4;
+
+	static const FString GetMajorVersionAsString();
+	static const FString GetMinorVersionAsString();
+
+	/** Returns true if the provided version is the latest version */
+	static bool IsLatest(int32 InMajorVersion, int32 InMinorVersion)
+	{
+		return
+			InMajorVersion == MajorVersion &&
+			InMinorVersion == MinorVersion;
+	}
+};
diff --git a/Engine/Plugins/VirtualProduction/DMX/DMXEngine/Source/DMXRuntime/Public/MVR/Types/DMXMVRChildListNode.h b/Engine/Plugins/VirtualProduction/DMX/DMXEngine/Source/DMXRuntime/Public/MVR/Types/DMXMVRChildListNode.h
new file mode 100644
--- /dev/null
+++ b/Engine/Plugins/VirtualProduction/DMX/DMXEngine/Source/DMXRuntime/Public/MVR/Types/DMXMVRChildListNode.h
@@ -0,0 +1,58 @@
+// Copyright Epic Games, Inc. All Rights Reserved.
+
+#pragma once
+
+#include "CoreMinimal.h"
+
+#include "DMXMVRChildListNode.generated.h"
+
+class UDMXMVRFixtureNode;
+class UDMXMVRParametricObjectNodeBase;
+
+class FXmlNode;
+
+
+/** This node defines a generic graphical object. */
+UCLASS()
+class UDMXMVRChildListNode
+	: public UObject
+{
+	GENERATED_BODY()
+
+public:
+	/** Initializes the Child List Node and its Children from a Child List Xml Node. Does not merge, clears previous entries. */
+	void InitializeFromChildListXmlNode(const FXmlNode& ChildListXmlNode);
+
+	/** Creates a DMX Xml Node as per MVR standard in parent, or logs warnings if no compliant Node can be created. */
+	void CreateXmlNodeInParent(FXmlNode& ParentNode) const;
+
+	/** Returns the Fixture Nodes contained in this Layer. Note, does not clear the array. */
+	void GetFixtureNodes(TArray<UDMXMVRFixtureNode*>& OutFixtureNodes) const;
+
+	/** Returns true if the Parametric Object Node is contained in this Child List */
+	bool Contains(UDMXMVRParametricObjectNodeBase* ParametricObjectNode);
+
+	/** Returns a pointer ot the Parametric Object Node that corresponds to the UUID, or nullptr if there is no corresponding Parametric Object Node */
+	UDMXMVRParametricObjectNodeBase** FindParametricObjectNodeByUUID(const FGuid& UUID);
+
+	/** Creates a Parametric Object Node as a Child in the Child List Node */
+	template<typename ParametricObjectNodeClass, typename TEnableIf<TIsDerivedFrom<ParametricObjectNodeClass, UDMXMVRParametricObjectNodeBase>::Value, bool>::Type = true>
+	ParametricObjectNodeClass* CreateParametricObject()
+	{
+		ParametricObjectNodeClass* NewParametricObject = NewObject<ParametricObjectNodeClass>(this);
+		ParametricObjectNodes.Add(NewParametricObject);
+
+		return NewParametricObject;
+	}
+
+	/** Removes a Parametric Object Child from the Child List. Returns true if the Parametric Object Node was contained and Removed */
+	bool RemoveParametricObject(UDMXMVRParametricObjectNodeBase* ParametricObjectNodeToRemove);
+
+	/** Returns the Parametric Object Node Children of this Child List Node */
+	FORCEINLINE const TArray<UDMXMVRParametricObjectNodeBase*>& GetParametricObjectNodes() const { return ParametricObjectNodes; };
+
+private:
+	/** A list of geometrical representation objects that are part of the object. */
+	UPROPERTY()
+	TArray<UDMXMVRParametricObjectNodeBase*> ParametricObjectNodes;
+};
diff --git a/Engine/Plugins/VirtualProduction/DMX/DMXEngine/Source/DMXRuntime/Public/MVR/Types/DMXMVRFixtureNode.h b/Engine/Plugins/VirtualProduction/DMX/DMXEngine/Source/DMXRuntime/Public/MVR/Types/DMXMVRFixtureNode.h
new file mode 100644
--- /dev/null
+++ b/Engine/Plugins/VirtualProduction/DMX/DMXEngine/Source/DMXRuntime/Public/MVR/Types/DMXMVRFixtureNode.h
@@ -0,0 +1,231 @@
+// Copyright Epic Games, Inc. All Rights Reserved.
+
+#pragma once
+
+#include "DMXMVRParametricObjectNodeBase.h"
+
+#include "DMXOptionalTypes.h"
+
+#include "CoreMinimal.h"
+#include "Misc/Optional.h"
+#include "Serialization/Archive.h"
+
+#include "DMXMVRFixtureNode.generated.h"
+
+class FXmlNode;
+
+
+/** The container for DMX Addresses for this fixture. */
+USTRUCT()
+struct FDMXMVRFixtureAddresses
+{
+	GENERATED_BODY()
+
+	FDMXMVRFixtureAddresses()
+		: Universe(1)
+		, Address(1)
+	{}
+
+	FDMXMVRFixtureAddresses(int32 InUniverse, int32 InAddress)
+		: Universe(InUniverse)
+		, Address(InAddress)
+	{}
+
+	/** Creates a DMX Xml Node as per MVR standard in parent, or logs warnings if no compliant Node can be created. */
+	void CreateXmlNodeInParent(FXmlNode& ParentNode) const;
+
+	UPROPERTY()
+	int32 Universe;
+
+	UPROPERTY()
+	int32 Address;
+};
+
+/** Defines a Mapping */
+USTRUCT()
+struct FDMXMVRFixtureMapping
+{
+	GENERATED_BODY()
+
+	/** Creates a DMX Xml Node as per MVR standard in parent, or logs warnings if no compliant Node can be created. */
+	void CreateXmlNodeInParent(FXmlNode& ParentNode) const;
+
+	/** The unique identifier of the MappingDefinition node that will be the source of the mapping. */
+	FGuid LinkedDef;
+
+	/** The offset in pixels in x direction from top left corner of the source that will be used for the mapped object. */
+	TOptional<int32> UX;
+
+	/** The offset in pixels in y direction from top left corner of the source that will be used for the mapped object. */
+	TOptional<int32> UY;
+
+	/** The size in pixels in x direction from top left of the starting point. */
+	TOptional<int32> OX;
+
+	/** The size in pixels in y direction from top left of the starting point. */
+	TOptional<int32> OY;
+
+	/** The rotation around the middle point of the defined rectangle in degree. Positive direction is counter cock wise. */
+	TOptional<float> RZ;
+};
+
+/** The Gobo used for the fixture. The image ressource must apply to the GDTF standard. */
+USTRUCT()
+struct FDMXMVRFixtureGobo
+{
+	GENERATED_BODY()
+
+	/** The node value is the Gobo used for the fixture. The image ressource must apply to the GDTF standard. Use a FileName to specify. */
+	UPROPERTY()
+	FString Value;
+
+	/** The roation of the Gobo in degree. */
+	UPROPERTY()
+	FDMXOptionalFloat Rotation;
+
+	friend FArchive& operator<<(FArchive& Ar, FDMXMVRFixtureGobo& Gobo)
+	{
+		Ar << Gobo.Value;
+		Ar << Gobo.Rotation;
+
+		return Ar;
+	}
+};
+
+USTRUCT()
+struct FDMXOptionalMVRFixtureGobo
+{
+	GENERATED_BODY()
+
+	/** Creates a DMX Xml Node as per MVR standard in parent, or logs warnings if no compliant Node can be created. */
+	void CreateXmlNodeInParent(FXmlNode& ParentNode) const;
+
+	TOptional<FDMXMVRFixtureGobo> Value;
+
+	bool Serialize(FArchive& Ar)
+	{
+		Ar << Value;
+		return true;
+	}
+
+	friend bool operator==(const FDMXOptionalMVRFixtureGobo& A, const FDMXOptionalMVRFixtureGobo& B)
+	{
+		const bool bBothSet = A.Value.IsSet() && B.Value.IsSet();
+		if (bBothSet)
+		{
+			return 
+				A.Value.GetValue().Value == B.Value.GetValue().Value &&
+				A.Value.GetValue().Rotation == B.Value.GetValue().Rotation;
+		}
+		else
+		{
+			return A.Value.IsSet() == B.Value.IsSet();
+		}
+	}
+
+	friend bool operator!=(const FDMXOptionalMVRFixtureGobo& A, const FDMXOptionalMVRFixtureGobo& B)
+	{
+		return !(A == B);
+	}
+};
+
+template<>
+struct TStructOpsTypeTraits<FDMXOptionalMVRFixtureGobo>
+	: public TStructOpsTypeTraitsBase2<FDMXOptionalMVRFixtureGobo>
+{
+	enum
+	{
+		WithSerializer = true,
+		WithIdenticalViaEquality = true
+	};
+};
+
+
+/** This node defines a light fixture object. */
+UCLASS()
+class UDMXMVRFixtureNode
+	: public UDMXMVRParametricObjectNodeBase
+{
+	GENERATED_BODY()
+
+public:
+	/** Initializes the Fixture Node and its Children from a Fixture Xml Node. Does not merge, clears previous entries. */
+	bool InitializeFromFixtureXmlNode(const FXmlNode& FixtureXmlNode);
+
+	//~Begin UDMXMVRParametricObjectNodeBase Interface
+	virtual void CreateXmlNodeInParent(FXmlNode& ParentNode) const override;
+	//~End UDMXMVRParametricObjectNodeBase Interface
+
+	/** Returns true if the data is sufficient to fullfill the mvr standard. If bLogInvalidReason is set to true, logs the reason for the fixture being invalid */
+	bool IsValid(bool bLogInvalidReason) const;
+
+	/** The name of the object. */
+	UPROPERTY()
+	FString Name;
+
+	/** The name of the file containing the GDTF information for this light fixture. */
+	UPROPERTY()
+	FString GDTFSpec;
+
+	/** The name of the used DMX mode. This has to match the name of a DMXMode in the GDTF file. */
+	UPROPERTY()
+	FString GDTFMode;
+
+	/** A focus point reference that this lighting fixture aims at if this reference exists. */
+	UPROPERTY()
+	FDMXOptionalGuid Focus;
+
+	/** Defines if a Object cast Shadow. */
+	UPROPERTY()
+	FDMXOptionalBool CastShadow;
+
+	/** A position reference that this lighting fixture belongs to if this reference exists. */
+	UPROPERTY()
+	FDMXOptionalGuid Position;
+
+	/** The Fixture Id of the lighting fixture. This is the short name of the fixture. */
+	UPROPERTY()
+	FString FixtureID;
+
+	/** The unit number of the lighting fixture in a position. */
+	UPROPERTY()
+	int32 UnitNumber = 0;
+	
+	/** Sets the Universe ID, in UE form (Index 0 is ID 1) */
+	FORCEINLINE void SetUniverseID(int32 UniverseID) { Addresses.Universe = UniverseID - 1; }
+
+	/** Returns the Universe ID, in UE form (Index 0 is ID 1) */
+	FORCEINLINE int32 GetUniverseID() const { return Addresses.Universe + 1; }
+
+	/** Sets the Address of the MVR Fixture */
+	FORCEINLINE void SetStartingChannel(int32 Address) { Addresses.Address = Address; }
+
+	/** Returns the Address of the MVR Fixture */
+	FORCEINLINE int32 GetStartingChannel() const { return Addresses.Address; }
+
+private:
+	/** The container for DMX Addresses for this fixture. */
+	UPROPERTY()
+	FDMXMVRFixtureAddresses Addresses;
+
+public:
+	/** A color assigned to a fixture. If it is not defined, there is no color for the fixture. */
+	UPROPERTY()
+	FDMXOptionalColorCIE1931 CIEColor;
+
+	/** The Fixture Type ID is a value that can be used as a short name of the Fixture Type. This does not have to be unique. The default value is 0. */
+	UPROPERTY()
+	FDMXOptionalInt32 FixtureTypeId;
+
+	/** The Custom ID is a value that can be used as a short name of the Fixture Instance. This does not have to be unique. The default value is 0. */
+	UPROPERTY()
+	FDMXOptionalInt32 CustomId;
+
+	/** The container for Mappings for this fixture. */
+	UPROPERTY()
+	TArray<FDMXMVRFixtureMapping> Mappings;
+
+	/** The Gobo used for the fixture. The image ressource must apply to the GDTF standard. */
+	UPROPERTY()
+	FDMXOptionalMVRFixtureGobo Gobo;
+};
diff --git a/Engine/Plugins/VirtualProduction/DMX/DMXEngine/Source/DMXRuntime/Public/MVR/Types/DMXMVRGroupObjectNode.h b/Engine/Plugins/VirtualProduction/DMX/DMXEngine/Source/DMXRuntime/Public/MVR/Types/DMXMVRGroupObjectNode.h
new file mode 100644
--- /dev/null
+++ b/Engine/Plugins/VirtualProduction/DMX/DMXEngine/Source/DMXRuntime/Public/MVR/Types/DMXMVRGroupObjectNode.h
@@ -0,0 +1,63 @@
+// Copyright Epic Games, Inc. All Rights Reserved.
+
+#pragma once
+
+#include "DMXMVRParametricObjectNodeBase.h"
+
+#include "DMXOptionalTypes.h"
+
+#include "CoreMinimal.h"
+
+#include "DMXMVRGroupObjectNode.generated.h"
+
+class UDMXMVRChildListNode;
+class UDMXMVRFixtureNode;
+class UDMXMVRParametricObjectNodeBase;
+
+class FXmlNode;
+
+
+/** This node defines logical group of objects. The child objects are located inside a local coordinate system. */
+UCLASS()
+class UDMXMVRGroupObjectNode
+	: public UDMXMVRParametricObjectNodeBase
+{
+	GENERATED_BODY()
+
+public:	
+	/** Constructor */
+	UDMXMVRGroupObjectNode();
+
+	/** Initializes the Child List Node and its Children from a Group Object Xml Node. Does not merge, clears previous entries. */
+	void InitializeFromGroupObjectXmlNode(const FXmlNode& GroupObjectXmlNode);
+
+	//~Begin UDMXMVRParametricObjectNodeBase Interface
+	virtual void CreateXmlNodeInParent(FXmlNode& ParentNode) const override;
+	//~End UDMXMVRParametricObjectNodeBase Interface
+
+	/** Returns the Fixture Nodes contained in this Layer. Note, does not clear the array. */
+	void GetFixtureNodes(TArray<UDMXMVRFixtureNode*>& OutFixtureNodes) const;
+
+	/** Returns true if the Parametric Object Node is contained in this Group Object Node */
+	bool Contains(UDMXMVRParametricObjectNodeBase* ParametricObjectNode);
+
+	/** Returns a pointer ot the Parametric Object Node that corresponds to the UUID, or nullptr if there is no corresponding Parametric Object Node */
+	UDMXMVRParametricObjectNodeBase** FindParametricObjectNodeByUUID(const FGuid& InUUID);
+
+	/** Removes the Parametric Object Node from Child List. Returns true if the Node was removed. */
+	bool RemoveParametricObjectNode(UDMXMVRParametricObjectNodeBase* ParametricObjectNode);
+
+	/** Returns a list of graphic objects that are part of the group. */
+	FORCEINLINE UDMXMVRChildListNode* GetChildListNode() const { return ChildListNode; }
+
+	/**
+	 * The name of the GroupObject 
+	 */
+	UPROPERTY()
+	FDMXOptionalString Name;
+
+private:
+	/** A list of graphic objects that are part of the group. */
+	UPROPERTY()
+	UDMXMVRChildListNode* ChildListNode;
+};
diff --git a/Engine/Plugins/VirtualProduction/DMX/DMXEngine/Source/DMXRuntime/Public/MVR/Types/DMXMVRLayerNode.h b/Engine/Plugins/VirtualProduction/DMX/DMXEngine/Source/DMXRuntime/Public/MVR/Types/DMXMVRLayerNode.h
new file mode 100644
--- /dev/null
+++ b/Engine/Plugins/VirtualProduction/DMX/DMXEngine/Source/DMXRuntime/Public/MVR/Types/DMXMVRLayerNode.h
@@ -0,0 +1,69 @@
+// Copyright Epic Games, Inc. All Rights Reserved.
+
+#pragma once
+
+#include "DMXOptionalTypes.h"
+
+#include "CoreMinimal.h"
+
+#include "DMXMVRLayerNode.generated.h"
+
+class UDMXMVRChildListNode;
+class UDMXMVRFixtureNode;
+class UDMXMVRParametricObjectNodeBase;
+
+class FXmlNode;
+
+
+/** A Layer node in the Layers Node of a Scene. */
+UCLASS()
+class UDMXMVRLayerNode
+	: public UObject
+{
+	GENERATED_BODY()
+
+public:
+	/** Constructor */
+	UDMXMVRLayerNode();
+
+	/** Initializes the Layer Node and its Children from a Layer Xml Node. Does not merge, clears previous entries. */
+	void InitializeFromLayerXmlNode(const FXmlNode& LayerXmlNode);
+
+	/** Creates a DMX Xml Node as per MVR standard in parent, or logs warnings if no compliant Node can be created. */
+	void CreateXmlNodeInParent(FXmlNode& ParentNode) const;
+
+	/** Returns the Fixture Nodes contained in this Layer. Note, does not clear the array. */
+	void GetFixtureNodes(TArray<UDMXMVRFixtureNode*>& OutFixtureNodes) const;
+
+	/** Returns true if the Parametric Object Node is contained in this layer */
+	bool Contains(UDMXMVRParametricObjectNodeBase* ParametricObjectNode);
+
+	/** Returns a pointer ot the Parametric Object Node that corresponds to the UUID, or nullptr if there is no corresponding Parametric Object Node */
+	UDMXMVRParametricObjectNodeBase** FindParametricObjectNodeByUUID(const FGuid& InUUID) const;
+
+	/** Returns a list of graphic objects that are part of the layer. */
+	FORCEINLINE UDMXMVRChildListNode* GetChildListNode() const { return ChildListNode; }
+
+	/** The unique identifier of the object. */
+	UPROPERTY()
+	FGuid UUID;
+
+	/** The name of the object */
+	UPROPERTY()
+	FDMXOptionalString Name;
+
+	/**
+	 * The transformation matrix that defines the location and orientation of this the layer inside its global coordinate space. 
+	 * This effectively defines local coordinate space for the objects inside. The Matrix of the Layer is only allowed to have an elevation, but no rotation.
+	 * 
+	 * UE specific: Note, we do not use an FVector since the standard clearly specifies a transformation matrix, even tho it also forbids any rotation.
+	 * While it makes the code more readable and future proof, the user of this struct is responsible to follow the standard directly.
+	 */
+	UPROPERTY()
+	FDMXOptionalTransform Matrix;
+
+private:
+	/** A list of graphic objects that are part of the group. */
+	UPROPERTY()
+	UDMXMVRChildListNode* ChildListNode;
+};
diff --git a/Engine/Plugins/VirtualProduction/DMX/DMXEngine/Source/DMXRuntime/Public/MVR/Types/DMXMVRLayersNode.h b/Engine/Plugins/VirtualProduction/DMX/DMXEngine/Source/DMXRuntime/Public/MVR/Types/DMXMVRLayersNode.h
new file mode 100644
--- /dev/null
+++ b/Engine/Plugins/VirtualProduction/DMX/DMXEngine/Source/DMXRuntime/Public/MVR/Types/DMXMVRLayersNode.h
@@ -0,0 +1,50 @@
+// Copyright Epic Games, Inc. All Rights Reserved.
+
+#pragma once
+
+#include "CoreMinimal.h"
+#include "XmlNode.h"
+
+#include "DMXMVRLayersNode.generated.h"
+
+class UDMXMVRFixtureNode;
+class UDMXMVRLayerNode;
+class UDMXMVRParametricObjectNodeBase;
+
+class FXmlNode;
+
+
+/** This node defines a list of layers inside the scene. The layer is a container of graphical objects defining a local coordinate system. */
+UCLASS()
+class DMXRUNTIME_API UDMXMVRLayersNode
+	: public UObject
+{
+	GENERATED_BODY()
+
+public:
+	/** Initializes the Layers Node and its Children from a Layers Xml Node. Does not merge, clears previous entries. */
+	void InitializeFromLayersXmlNode(const FXmlNode& LayersXmlNode);
+
+	/** Creates a DMX Xml Node as per MVR standard in parent, or logs warnings if no compliant Node can be created. */
+	void CreateXmlNodeInParent(FXmlNode& ParentNode) const;
+
+	/** Returns the Fixture Nodes contained in this Layer. Note, does not clear the array. */
+	void GetFixtureNodes(TArray<UDMXMVRFixtureNode*>& OutFixtureNodes) const;
+
+	/** Returns a pointer ot the Parametric Object Node that corresponds to the UUID, or nullptr if there is no corresponding Parametric Object Node */
+	UDMXMVRParametricObjectNodeBase** FindParametricObjectNodeByUUID(const FGuid& UUID) const;
+
+	/** Creates a new Layer. Returns the newly created Layer */
+	UDMXMVRLayerNode* CreateLayer();
+
+	/** Removes a Layer. Ensures the Layer is contained in the Layers of this Node */
+	void RemoveLayer(UDMXMVRLayerNode* LayerNodeToRemove);
+
+	/** Returns the Layers of this Node */
+	FORCEINLINE const TArray<UDMXMVRLayerNode*>& GetLayerNodes() const { return LayerNodes; };
+
+private:
+	/** A[n array of] layer representation[s]. */
+	UPROPERTY()
+	TArray<UDMXMVRLayerNode*> LayerNodes;
+};
diff --git a/Engine/Plugins/VirtualProduction/DMX/DMXEngine/Source/DMXRuntime/Public/MVR/Types/DMXMVRParametricObjectNodeBase.h b/Engine/Plugins/VirtualProduction/DMX/DMXEngine/Source/DMXRuntime/Public/MVR/Types/DMXMVRParametricObjectNodeBase.h
new file mode 100644
--- /dev/null
+++ b/Engine/Plugins/VirtualProduction/DMX/DMXEngine/Source/DMXRuntime/Public/MVR/Types/DMXMVRParametricObjectNodeBase.h
@@ -0,0 +1,55 @@
+// Copyright Epic Games, Inc. All Rights Reserved.
+
+#pragma once
+
+#include "DMXOptionalTypes.h"
+
+#include "CoreMinimal.h"
+#include "Misc/Guid.h"
+
+#include "DMXMVRParametricObjectNodeBase.generated.h"
+
+class UDMXMVRGroupObjectNode;
+class UDMXMVRLayerNode;
+
+class FXmlNode;
+
+
+
+/** UE specific: The base class for all Children in a Child List Node, whereas the standard refers to them as 'Parametric Objects'. */
+UCLASS(Abstract)
+class UDMXMVRParametricObjectNodeBase
+	: public UObject
+{
+	GENERATED_BODY()
+
+public:
+	/** Creates a DMX Xml Node as per MVR standard in parent, or logs warnings if no compliant Node can be created. */
+	virtual void CreateXmlNodeInParent(FXmlNode& ParentNode) const PURE_VIRTUAL(UDMXMVRParametricObjectNodeBase::CreateXmlNodeInParent, return;);
+
+	/** Returns the Layer this Node resides in */
+	UDMXMVRLayerNode* GetLayer();
+
+	/** Returns the Layer this Node resides in */
+	const UDMXMVRLayerNode* GetLayer() const;
+
+	/** Returns the Group this Node resides in, or null if not in a Group */
+	UDMXMVRGroupObjectNode* GetGroup();
+
+	/** Returns the Group this Node resides in, or null if not in a Group */
+	const UDMXMVRGroupObjectNode* GetGroup() const;
+
+	/** Gets the Transform in the absolute coordinate system */
+	FTransform GetTransformAbsolute() const;
+
+	/** Sets the Transform in the absolute coordinate system */
+	void SetTransformAbsolute(FTransform NewTransform);
+
+	/** The unique identifier of the object. */
+	UPROPERTY()
+	FGuid UUID;
+
+	/** The location of the object inside the parent coordinate system. */
+	UPROPERTY(VisibleAnywhere, Category = "DMX")
+	FDMXOptionalTransform Matrix;
+};
diff --git a/Engine/Plugins/VirtualProduction/DMX/DMXEngine/Source/DMXRuntime/Public/MVR/Types/DMXMVRRootNode.h b/Engine/Plugins/VirtualProduction/DMX/DMXEngine/Source/DMXRuntime/Public/MVR/Types/DMXMVRRootNode.h
new file mode 100644
--- /dev/null
+++ b/Engine/Plugins/VirtualProduction/DMX/DMXEngine/Source/DMXRuntime/Public/MVR/Types/DMXMVRRootNode.h
@@ -0,0 +1,67 @@
+// Copyright Epic Games, Inc. All Rights Reserved.
+
+#pragma once
+
+#include "CoreMinimal.h"
+
+#include "DMXMVRRootNode.generated.h"
+
+class UDMXMVRChildListNode;
+class UDMXMVRFixtureNode;
+class UDMXMVRLayerNode;
+class UDMXMVRParametricObjectNodeBase;
+class UDMXMVRSceneNode;
+class UDMXMVRUserDataNode;
+
+class FXmlFile;
+
+
+/** 
+ * The root node of an MVR General Scene Description as typically held in an MVR's GeneralSceneDescription.xml. 
+ * 
+ * UE specific: Note, while properties follow the standard closely, only children relevant to the Engine are implemented.
+ */
+UCLASS()
+class UDMXMVRRootNode
+	: public UObject
+{
+	GENERATED_BODY()
+
+public:
+	/** Constructor */
+	UDMXMVRRootNode();
+
+	/** 
+	 * Initializes the Object and its children from a General Scene Description Xml File. 
+	 * Does not merge, clears previous entries. 
+	 * Returns false if not a valid MVR General Scene Description Xml 
+	 */
+	bool InitializeFromGeneralSceneDescriptionXml(const TSharedRef<FXmlFile>& GeneralSceneDescriptionXml);
+
+	/** Creates a DMX Xml Node as per MVR standard, returns nullptr if the Xml File cannot be created. */
+	TSharedPtr<FXmlFile> CreateXmlFile();
+
+	/** Returns the Fixture Nodes contained in this Layer. Note, does not clear the array. */
+	void GetFixtureNodes(TArray<UDMXMVRFixtureNode*>& OutFixtureNodes) const;
+
+	/** Returns the first Child List Node in the Scene */
+	UDMXMVRChildListNode& GetOrCreateFirstChildListNode();
+
+	/** Returns a pointer ot the Parametric Object Node that corresponds to the UUID, or nullptr if there is no corresponding Parametric Object Node */
+	UDMXMVRParametricObjectNodeBase** FindParametricObjectNodeByUUID(const FGuid& UUID) const;
+
+	/** Removes the Parametric Object Node from the General Scene Description. Returns true if the Node was removed. */
+	bool RemoveParametricObjectNode(UDMXMVRParametricObjectNodeBase* ParametricObjectNode);
+
+	/** Returns the node that contains information about the mvr scene */
+	FORCEINLINE UDMXMVRSceneNode* GetSceneNode() const { return SceneNode; }
+
+private:
+	/** This node contains a collection of user data nodes defined and used by provider applications if required. */
+	UPROPERTY()
+	UDMXMVRUserDataNode* UserDataNode;
+
+	/** This node contains information about the scene. */
+	UPROPERTY()
+	UDMXMVRSceneNode* SceneNode;
+};
diff --git a/Engine/Plugins/VirtualProduction/DMX/DMXEngine/Source/DMXRuntime/Public/MVR/Types/DMXMVRSceneNode.h b/Engine/Plugins/VirtualProduction/DMX/DMXEngine/Source/DMXRuntime/Public/MVR/Types/DMXMVRSceneNode.h
new file mode 100644
--- /dev/null
+++ b/Engine/Plugins/VirtualProduction/DMX/DMXEngine/Source/DMXRuntime/Public/MVR/Types/DMXMVRSceneNode.h
@@ -0,0 +1,42 @@
+// Copyright Epic Games, Inc. All Rights Reserved.
+
+#pragma once
+
+#include "CoreMinimal.h"
+
+#include "DMXMVRSceneNode.generated.h"
+
+class UDMXMVRLayersNode;
+
+class FXmlNode;
+
+
+/**
+ * This node contains information about the scene.
+ *
+ * UE Specific: Note, while properties follow the standard closely, only properties relevant to the Engine are implemented.
+ */
+UCLASS()
+class UDMXMVRSceneNode
+	: public UObject
+{
+	GENERATED_BODY()
+
+public:
+	/** Constructor */
+	UDMXMVRSceneNode();
+
+	/** Initializes the Scene Node andd its Children from a Scene Xml Node. Does not merge, clears previous entries. */
+	void InitializeFromSceneXmlNode(const FXmlNode& SceneXmlNode);
+
+	/** Creates a DMX Xml Node as per MVR standard in parent, or logs warnings if no compliant Node can be created. */
+	void CreateXmlNodeInParent(FXmlNode& ParentNode) const;
+
+	/** Returns the node that defines a list of layers inside the scene. The layer is a container of graphical objects defining a local coordinate system. */
+	FORCEINLINE UDMXMVRLayersNode* GetLayersNode() const { return LayersNode; }
+
+private:
+	/** This node defines a list of layers inside the scene. The layer is a container of graphical objects defining a local coordinate system. */
+	UPROPERTY()
+	UDMXMVRLayersNode* LayersNode;
+};
diff --git a/Engine/Plugins/VirtualProduction/DMX/DMXEngine/Source/DMXRuntime/Public/MVR/Types/DMXMVRUnrealEngineDataNode.h b/Engine/Plugins/VirtualProduction/DMX/DMXEngine/Source/DMXRuntime/Public/MVR/Types/DMXMVRUnrealEngineDataNode.h
new file mode 100644
--- /dev/null
+++ b/Engine/Plugins/VirtualProduction/DMX/DMXEngine/Source/DMXRuntime/Public/MVR/Types/DMXMVRUnrealEngineDataNode.h
@@ -0,0 +1,28 @@
+// Copyright Epic Games, Inc. All Rights Reserved.
+
+#pragma once
+
+#include "CoreMinimal.h"
+
+#include "DMXMVRUnrealEngineDataNode.generated.h"
+
+class FXmlNode;
+
+
+/** 
+ * This node contains a collection of data specified by the provider application.
+ * UE sepecific: Used to add UE specific meta data
+ */
+UCLASS()
+class UDMXMVRUnrealEngineDataNode
+	: public UObject
+{
+	GENERATED_BODY()
+
+public:
+	/** Creates a DMX Xml Node as per MVR standard in parent, or logs warnings if no compliant Node can be created. */
+	void CreateXmlNodeInParent(FXmlNode& ParentNode) const;
+
+	/** The Provider Name when the Data node is provided by Unreal Engine */
+	static constexpr TCHAR ProviderNameUnrealEngine[] = TEXT("UnrealEngine");
+};
diff --git a/Engine/Plugins/VirtualProduction/DMX/DMXEngine/Source/DMXRuntime/Public/MVR/Types/DMXMVRUserDataNode.h b/Engine/Plugins/VirtualProduction/DMX/DMXEngine/Source/DMXRuntime/Public/MVR/Types/DMXMVRUserDataNode.h
new file mode 100644
--- /dev/null
+++ b/Engine/Plugins/VirtualProduction/DMX/DMXEngine/Source/DMXRuntime/Public/MVR/Types/DMXMVRUserDataNode.h
@@ -0,0 +1,33 @@
+// Copyright Epic Games, Inc. All Rights Reserved.
+
+#pragma once
+
+#include "CoreMinimal.h"
+
+#include "DMXMVRUserDataNode.generated.h"
+
+class UDMXMVRDataNode;
+class UDMXMVRUnrealEngineDataNode;
+
+class FXmlNode;
+
+
+/** This node contains a collection of user data nodes defined and used by provider applications if required. */
+UCLASS()
+class UDMXMVRUserDataNode
+	: public UObject
+{
+	GENERATED_BODY()
+
+public:
+	/** Constructor */
+	UDMXMVRUserDataNode();
+
+	/** Creates a DMX Xml Node as per MVR standard in parent, or logs warnings if no compliant Node can be created. */
+	void CreateXmlNodeInParent(FXmlNode& ParentNode) const;
+	
+private:
+	/** UE Specific: The UE's proprietary data node */
+	UPROPERTY()
+	UDMXMVRUnrealEngineDataNode* UnrealEngineDataNode;
+};
diff --git a/Engine/Plugins/VirtualProduction/DMX/DatasmithMVR/Source/DatasmithMVRTranslator/Private/DatasmithMVRNativeTranslator.cpp b/Engine/Plugins/VirtualProduction/DMX/DatasmithMVR/Source/DatasmithMVRTranslator/Private/DatasmithMVRNativeTranslator.cpp
--- a/Engine/Plugins/VirtualProduction/DMX/DatasmithMVR/Source/DatasmithMVRTranslator/Private/DatasmithMVRNativeTranslator.cpp
+++ b/Engine/Plugins/VirtualProduction/DMX/DatasmithMVR/Source/DatasmithMVRTranslator/Private/DatasmithMVRNativeTranslator.cpp
@@ -6,6 +6,7 @@
 #include "Factories/DMXLibraryFromMVRFactory.h"
 #include "Library/DMXEntityFixturePatch.h"
 #include "Library/DMXLibrary.h"
+#include "MVR/Types/DMXMVRFixtureNode.h"
 
 #include "Editor.h"
 #include "AssetRegistry/AssetRegistryModule.h"
@@ -180,7 +181,9 @@ void FDatasmithMVRNativeTranslator::ReplaceMVRActorsWithMVRSceneActor(TSharedRef
 		return;
 	}
 
-	for (const FDMXMVRFixture& MVRFixture : GeneralSceneDescription->GetMVRFixtures())
+	TArray<UDMXMVRFixtureNode*> FixtureNodes;
+	GeneralSceneDescription->GetFixtureNodes(FixtureNodes);
+	for (const UDMXMVRFixtureNode* MVRFixture : FixtureNodes)
 	{
 		for (int32 MetaDataIndex = 0; MetaDataIndex < NumMetaDatas; MetaDataIndex++)
 		{
@@ -214,7 +217,7 @@ void FDatasmithMVRNativeTranslator::ReplaceMVRActorsWithMVRSceneActor(TSharedRef
 				continue;
 			}
 
-			if (MetaDataMVRUUID == MVRFixture.UUID)
+			if (MetaDataMVRUUID == MVRFixture->UUID)
 			{
 				const TSharedPtr<IDatasmithElement> AssociatedElementToRemove = MetaDataElement->GetAssociatedElement();
 
