DMX - Fix CIS Check Copyright Notices, missing copyright boilerplate in DMXZipper.h, DMXZipper.cpp

diff --git a/Engine/Plugins/VirtualProduction/DMX/DMXEngine/Source/DMXEditor/Private/DMXZipper.cpp b/Engine/Plugins/VirtualProduction/DMX/DMXEngine/Source/DMXEditor/Private/DMXZipper.cpp
--- a/Engine/Plugins/VirtualProduction/DMX/DMXEngine/Source/DMXEditor/Private/DMXZipper.cpp
+++ b/Engine/Plugins/VirtualProduction/DMX/DMXEngine/Source/DMXEditor/Private/DMXZipper.cpp
@@ -1,4 +1,4 @@
-// Fill out your copyright notice in the Description page of Project Settings.
+// Copyright Epic Games, Inc. All Rights Reserved.
 
 #include "DMXZipper.h"
 
diff --git a/Engine/Plugins/VirtualProduction/DMX/DMXEngine/Source/DMXEditor/Private/DMXZipper.h b/Engine/Plugins/VirtualProduction/DMX/DMXEngine/Source/DMXEditor/Private/DMXZipper.h
--- a/Engine/Plugins/VirtualProduction/DMX/DMXEngine/Source/DMXEditor/Private/DMXZipper.h
+++ b/Engine/Plugins/VirtualProduction/DMX/DMXEngine/Source/DMXEditor/Private/DMXZipper.h
@@ -1,4 +1,4 @@
-// Fill out your copyright notice in the Description page of Project Settings.
+// Copyright Epic Games, Inc. All Rights Reserved.
 
 #pragma once
 
