DMX - Fix MVR layer did not export its attributes, MVR Fixture Addresses attribute 'break' now exports with the correct name

diff --git a/Engine/Plugins/VirtualProduction/DMX/DMXEngine/Source/DMXRuntime/Private/MVR/Types/DMXMVRFixtureNode.cpp b/Engine/Plugins/VirtualProduction/DMX/DMXEngine/Source/DMXRuntime/Private/MVR/Types/DMXMVRFixtureNode.cpp
--- a/Engine/Plugins/VirtualProduction/DMX/DMXEngine/Source/DMXRuntime/Private/MVR/Types/DMXMVRFixtureNode.cpp
+++ b/Engine/Plugins/VirtualProduction/DMX/DMXEngine/Source/DMXRuntime/Private/MVR/Types/DMXMVRFixtureNode.cpp
@@ -391,7 +391,7 @@ void FDMXMVRFixtureAddresses::CreateXmlNodeInParent(FXmlNode& ParentNode) const
 	constexpr TCHAR Tag[] = TEXT("Address");
 	const FString Content = FString::FromInt(AbsoluteChannel);
 
-	constexpr TCHAR BreakAttributeName[] = TEXT("rotation");
+	constexpr TCHAR BreakAttributeName[] = TEXT("break");
 	constexpr TCHAR BreakAttributeValue[] = TEXT("0");
 	const TArray<FXmlAttribute> Attributes =
 	{
diff --git a/Engine/Plugins/VirtualProduction/DMX/DMXEngine/Source/DMXRuntime/Private/MVR/Types/DMXMVRLayerNode.cpp b/Engine/Plugins/VirtualProduction/DMX/DMXEngine/Source/DMXRuntime/Private/MVR/Types/DMXMVRLayerNode.cpp
--- a/Engine/Plugins/VirtualProduction/DMX/DMXEngine/Source/DMXRuntime/Private/MVR/Types/DMXMVRLayerNode.cpp
+++ b/Engine/Plugins/VirtualProduction/DMX/DMXEngine/Source/DMXRuntime/Private/MVR/Types/DMXMVRLayerNode.cpp
@@ -56,35 +56,34 @@ void UDMXMVRLayerNode::CreateXmlNodeInParent(FXmlNode& ParentNode) const
 	checkf(ChildListNode, TEXT("Default Subobject 'ChildListNode' is invalid, this is not expected."));
 	checkf(ParentNode.GetTag() == TEXT("layers"), TEXT("Layer Nodes have to be created in a Layers node, but parent node is %s."), *ParentNode.GetTag());
 
+	// Add self
+	constexpr TCHAR Tag[] = TEXT("Layer");
+	ParentNode.AppendChildNode(Tag);
+	FXmlNode* LayerXmlNode = ParentNode.GetChildrenNodes().Last();
+	check(LayerXmlNode);
+
 	TArray<FXmlAttribute> Attributes;
 
-	// Name attribute
 	if (Name.IsSet())
 	{
-		constexpr TCHAR AttributeName[] = TEXT("name");
-		FXmlAttribute Attribute = FXmlAttribute(AttributeName, Name.GetValue());
+		constexpr TCHAR NameAttributeName[] = TEXT("name");
+		FXmlAttribute Attribute = FXmlAttribute(NameAttributeName, Name.GetValue());
 		Attributes.Add(Attribute);
 	}
 	const FString NonOptionalName = Name.IsSet() ? Name.GetValue() : TEXT("Unnamed"); // Hold a name for logging
 
-	// UUID attribute
+	if (!UUID.IsValid())
 	{
-		if (!UUID.IsValid())
-		{
-			UE_LOG(LogDMXRuntime, Warning, TEXT("Cannot export MVR Layer '%s'. Invalid UUID specified."), *NonOptionalName);
-			return;
-		}
-
-		constexpr TCHAR AttributeName[] = TEXT("uuid");
-		FXmlAttribute Attribute = FXmlAttribute(AttributeName, UUID.ToString(EGuidFormats::DigitsWithHyphens));
-		Attributes.Add(Attribute);
+		UE_LOG(LogDMXRuntime, Warning, TEXT("Cannot export MVR Layer '%s'. Invalid UUID specified."), *NonOptionalName);
+		return;
 	}
 
-	// Add self
-	constexpr TCHAR Tag[] = TEXT("Layer");
-	ParentNode.AppendChildNode(Tag);
-	FXmlNode* LayerXmlNode = ParentNode.GetChildrenNodes().Last();
-	check(LayerXmlNode);
+	constexpr TCHAR UUIDAttributeName[] = TEXT("uuid");
+	FXmlAttribute Attribute = FXmlAttribute(UUIDAttributeName, UUID.ToString(EGuidFormats::DigitsWithHyphens));
+	Attributes.Add(Attribute);
+
+	LayerXmlNode->SetAttributes(Attributes);
+
 
 	// Add Matrix child Node
 	if (Matrix.IsSet())
