DMX - Fix DMX component does not raise its tick event in editor

diff --git a/Engine/Plugins/VirtualProduction/DMX/DMXEngine/Source/DMXRuntime/Private/Game/DMXComponent.cpp b/Engine/Plugins/VirtualProduction/DMX/DMXEngine/Source/DMXRuntime/Private/Game/DMXComponent.cpp
--- a/Engine/Plugins/VirtualProduction/DMX/DMXEngine/Source/DMXRuntime/Private/Game/DMXComponent.cpp
+++ b/Engine/Plugins/VirtualProduction/DMX/DMXEngine/Source/DMXRuntime/Private/Game/DMXComponent.cpp
@@ -12,11 +12,11 @@ UDMXComponent::UDMXComponent()
 	: bReceiveDMXFromPatch(true)
 {
 	PrimaryComponentTick.bCanEverTick = true;
-	PrimaryComponentTick.bStartWithTickEnabled = false;
+	PrimaryComponentTick.bStartWithTickEnabled = true;
 	PrimaryComponentTick.TickGroup = ETickingGroup::TG_PrePhysics;
+	bTickInEditor = true;
 
 #if WITH_EDITOR
-	bTickInEditor = true;
 	if (!HasAnyFlags(RF_ClassDefaultObject | RF_ArchetypeObject))
 	{
 		UDMXProtocolSettings* DMXProtocolSettings = GetMutableDefault<UDMXProtocolSettings>();
@@ -60,11 +60,11 @@ void UDMXComponent::SetReceiveDMXFromPatch(bool bReceive)
 
 void UDMXComponent::OnFixturePatchReceivedDMX(UDMXEntityFixturePatch* FixturePatch, const FDMXNormalizedAttributeValueMap& NormalizedValuePerAttribute)
 {
-#if WITH_EDITOR
-	FEditorScriptExecutionGuard ScriptGuard;
-#endif
-	
-	OnFixturePatchReceived.Broadcast(FixturePatch, NormalizedValuePerAttribute);
+	if (OnFixturePatchReceived.IsBound())
+	{
+		FEditorScriptExecutionGuard ScriptGuard;
+		OnFixturePatchReceived.Broadcast(FixturePatch, NormalizedValuePerAttribute);
+	}
 }
 
 void UDMXComponent::SetupReceiveDMXBinding()
@@ -179,7 +179,11 @@ void UDMXComponent::TickComponent(float DeltaTime, enum ELevelTick TickType, FAc
 {
 	Super::TickComponent(DeltaTime, TickType, ThisTickFunction);
 
-	OnDMXComponentTick.Broadcast(DeltaTime);
+	if (OnDMXComponentTick.IsBound())
+	{
+		FEditorScriptExecutionGuard ScriptGuard;
+		OnDMXComponentTick.Broadcast(DeltaTime);
+	}
 }
 
 void UDMXComponent::DestroyComponent(bool bPromoteChildren)
