DMX - Fixture List now correctly sets an initial selection, Fixture List now displays a warning for patches that do not have a Fixture Type or Mode selected

diff --git a/Engine/Plugins/VirtualProduction/DMX/DMXEngine/Source/DMXEditor/Private/Widgets/FixturePatch/SDMXMVRFixtureList.cpp b/Engine/Plugins/VirtualProduction/DMX/DMXEngine/Source/DMXEditor/Private/Widgets/FixturePatch/SDMXMVRFixtureList.cpp
--- a/Engine/Plugins/VirtualProduction/DMX/DMXEngine/Source/DMXEditor/Private/Widgets/FixturePatch/SDMXMVRFixtureList.cpp
+++ b/Engine/Plugins/VirtualProduction/DMX/DMXEngine/Source/DMXEditor/Private/Widgets/FixturePatch/SDMXMVRFixtureList.cpp
@@ -41,43 +41,40 @@ public:
 	{
 		TMap<TSharedPtr<FDMXMVRFixtureListItem>, FText> AccumulatedConflicts;
 
+		TMap<TSharedPtr<FDMXMVRFixtureListItem>, FText> FixtureTypeIssues = GetFixtureTypeIssues();
+		AppendConflictTexts(FixtureTypeIssues, AccumulatedConflicts);
+
 		TMap<TSharedPtr<FDMXMVRFixtureListItem>, FText> UnitNumberConflicts = GetUnitNumberConflicts();
-		AccumulatedConflicts.Append(UnitNumberConflicts);
+		AppendConflictTexts(UnitNumberConflicts, AccumulatedConflicts);
 
 		TMap<TSharedPtr<FDMXMVRFixtureListItem>, FText> ChannelExcessConflicts = GetChannelExcessConflicts();
-		for (const TTuple<TSharedPtr<FDMXMVRFixtureListItem>, FText>& ChannelConflict : ChannelExcessConflicts)
-		{
-			if (AccumulatedConflicts.Contains(ChannelConflict.Key))
-			{
-				const FText LineTerminator = FText::FromString(LINE_TERMINATOR);
-				const FText AccumulatedErrorText = FText::Format(FText::FromString(TEXT("{0}{1}{2}{3}")), AccumulatedConflicts[ChannelConflict.Key], LineTerminator, LineTerminator, ChannelConflict.Value);
-				AccumulatedConflicts[ChannelConflict.Key] = AccumulatedErrorText;
-			}
-			else
-			{
-				AccumulatedConflicts.Add(ChannelConflict.Key, ChannelConflict.Value);
-			}
-		}
+		AppendConflictTexts(ChannelExcessConflicts, AccumulatedConflicts);
 
 		TMap<TSharedPtr<FDMXMVRFixtureListItem>, FText> ChannelOverlapConflicts = GetChannelOverlapConflicts();
-		for (const TTuple<TSharedPtr<FDMXMVRFixtureListItem>, FText>& ChannelConflict : ChannelOverlapConflicts)
+		AppendConflictTexts(ChannelOverlapConflicts, AccumulatedConflicts);
+
+		return AccumulatedConflicts;
+	}
+
+private:
+	void AppendConflictTexts(const TMap<TSharedPtr<FDMXMVRFixtureListItem>, FText>& InItemToConflictTextMap, TMap<TSharedPtr<FDMXMVRFixtureListItem>, FText>& InOutConflictTexts) const
+	{
+		for (const TTuple<TSharedPtr<FDMXMVRFixtureListItem>, FText>& ItemToConflictTextPair : InItemToConflictTextMap)
 		{
-			if (AccumulatedConflicts.Contains(ChannelConflict.Key))
+			if (InOutConflictTexts.Contains(ItemToConflictTextPair.Key))
 			{
 				const FText LineTerminator = FText::FromString(LINE_TERMINATOR);
-				const FText AccumulatedErrorText = FText::Format(FText::FromString(TEXT("{0}{1}{2}{3}")), AccumulatedConflicts[ChannelConflict.Key], LineTerminator, LineTerminator, ChannelConflict.Value);
-				AccumulatedConflicts[ChannelConflict.Key] = AccumulatedErrorText;
+				const FText AccumulatedErrorText = FText::Format(FText::FromString(TEXT("{0}{1}{2}{3}")), InOutConflictTexts[ItemToConflictTextPair.Key], LineTerminator, LineTerminator, ItemToConflictTextPair.Value);
+				InOutConflictTexts[ItemToConflictTextPair.Key] = AccumulatedErrorText;
 			}
 			else
 			{
-				AccumulatedConflicts.Add(ChannelConflict.Key, ChannelConflict.Value);
+				InOutConflictTexts.Add(ItemToConflictTextPair);
 			}
 		}
 
-		return AccumulatedConflicts;
 	}
 
-private:
 	/** The patch of an item */
 	struct FItemPatch
 	{
@@ -154,6 +151,32 @@ private:
 		TSharedPtr<FDMXMVRFixtureListItem> Item;
 	};
 
+	/** Returns a Map of Items to Channels that have Fixture Types with issues set */
+	TMap<TSharedPtr<FDMXMVRFixtureListItem>, FText> GetFixtureTypeIssues() const
+	{
+		TMap<TSharedPtr<FDMXMVRFixtureListItem>, FText> ItemToIssueMap;
+		for (const TSharedPtr<FDMXMVRFixtureListItem>& Item : Items)
+		{
+			if (!Item->GetFixtureType())
+			{
+				const FText IssueText = LOCTEXT("NoFixtureTypeIssue", "No Fixture Type selected.");
+				ItemToIssueMap.Add(Item, IssueText);
+			}
+			else if (Item->GetFixtureType()->Modes.IsEmpty())
+			{
+				const FText IssueText = LOCTEXT("NoFixtureTypeIssue", "Fixture Type has no Modes defined.");
+				ItemToIssueMap.Add(Item, IssueText);
+			}
+			else if (Item->GetFixtureType() && Item->GetModeIndex() == INDEX_NONE)
+			{
+				const FText IssueText = FText::Format(LOCTEXT("FixtureTypeWithoutModeIssue", "Fixture Type '{0}' does not specify any Modes."), FText::FromString(Item->GetFixtureType()->GetName()));
+				ItemToIssueMap.Add(Item, IssueText);
+			}
+		}
+
+		return ItemToIssueMap;
+	}
+
 	/** Returns a Map of Items to Channels exceeding the DMX address range Texts */
 	TMap<TSharedPtr<FDMXMVRFixtureListItem>, FText> GetChannelExcessConflicts() const
 	{
@@ -653,6 +676,15 @@ void SDMXMVRFixtureList::AdoptSelectionFromFixturePatchSharedData()
 		ListView->SetItemSelection(NewSelection, bSelected, ESelectInfo::OnMouseClick);
 		ListView->RequestScrollIntoView(NewSelection[0]);
 	}
+	else if(TSharedPtr<FDMXEditor> DMXEditor = WeakDMXEditor.Pin())
+	{
+		const UDMXLibrary* DMXLibrary = DMXEditor->GetDMXLibrary();
+		const TArray<UDMXEntityFixturePatch*> FixturePatches = DMXLibrary->GetEntitiesTypeCast<UDMXEntityFixturePatch>();
+		if (FixturePatches.Num() > 0)
+		{
+			FixturePatchSharedData->SelectFixturePatch(FixturePatches[0]);
+		}
+	}
 }
 
 void SDMXMVRFixtureList::AutoAssignFixturePatches()
