DMX - Add a workaround to create Matrix Fixture Types from GDTF, offer option to load GDTFs for assets created prior to 5.1, rollback deprecation of FixtureType's DMXImport member

diff --git a/Engine/Plugins/VirtualProduction/DMX/DMXEngine/Source/DMXEditor/Private/Customizations/DMXEntityFixtureTypeDetails.cpp b/Engine/Plugins/VirtualProduction/DMX/DMXEngine/Source/DMXEditor/Private/Customizations/DMXEntityFixtureTypeDetails.cpp
--- a/Engine/Plugins/VirtualProduction/DMX/DMXEngine/Source/DMXEditor/Private/Customizations/DMXEntityFixtureTypeDetails.cpp
+++ b/Engine/Plugins/VirtualProduction/DMX/DMXEngine/Source/DMXEditor/Private/Customizations/DMXEntityFixtureTypeDetails.cpp
@@ -2,6 +2,8 @@
 
 #include "Customizations/DMXEntityFixtureTypeDetails.h"
 
+#include "DMXEditorLog.h"
+#include "DMXInitializeFixtureTypeFromGDTFHelper.h"
 #include "Library/DMXImportGDTF.h"
 #include "Library/DMXEntityFixtureType.h"
 
@@ -24,26 +26,42 @@ void FDMXEntityFixtureTypeDetails::CustomizeDetails(IDetailLayoutBuilder& Detail
 
 	DetailBuilder.HideProperty(GET_MEMBER_NAME_CHECKED(UDMXEntityFixtureType, Modes));
 
-	GDTFHandle = DetailBuilder.GetProperty(GET_MEMBER_NAME_CHECKED(UDMXEntityFixtureType, GDTF));
+	GDTFHandle = DetailBuilder.GetProperty(GET_MEMBER_NAME_CHECKED(UDMXEntityFixtureType, DMXImport));
 	GDTFHandle->SetOnPropertyValueChanged(FSimpleDelegate::CreateSP(this, &FDMXEntityFixtureTypeDetails::OnDMXImportChanged));
 }
 
 void FDMXEntityFixtureTypeDetails::OnDMXImportChanged()
 {
 	UObject* DMXImportObject = nullptr;
-	if (GDTFHandle->GetValue(DMXImportObject) == FPropertyAccess::Success)
+	if (GDTFHandle->GetValue(DMXImportObject) != FPropertyAccess::Success)
 	{
-		if (UDMXImportGDTF* DMXImportGDTF = Cast<UDMXImportGDTF>(DMXImportObject))
+		return;
+	}
+
+	UDMXImportGDTF* DMXImportGDTF = Cast<UDMXImportGDTF>(DMXImportObject);
+	const TArray<TWeakObjectPtr<UObject>>& SelectedObjects = PropertyUtilities->GetSelectedObjects();
+	for (TWeakObjectPtr<UObject> WeakFixtureTypeObject : SelectedObjects)
+	{
+		if (UDMXEntityFixtureType* FixtureType = Cast<UDMXEntityFixtureType>(WeakFixtureTypeObject.Get()))
 		{
-			const TArray<TWeakObjectPtr<UObject>>& SelectedObjects = PropertyUtilities->GetSelectedObjects();
+			FixtureType->PreEditChange(nullptr);
+			FixtureType->Modes.Reset();
+			FixtureType->PostEditChange();
+			
+			if (!DMXImportGDTF)
+			{
+				continue;
+			}
 
-			for (TWeakObjectPtr<UObject> WeakFixtureTypeObject : SelectedObjects)
+			// Try to use the work around that supports creation of matrices, otherwise setup the fixture type with the old implementation
+			FixtureType->PreEditChange(nullptr);
+			const bool bAdvancedImportSuccess = FDMXInitializeFixtureTypeFromGDTFHelper::GenerateModesFromGDTF(*FixtureType, *DMXImportGDTF);
+			if (!bAdvancedImportSuccess)
 			{
-				if (UDMXEntityFixtureType* FixtureType = Cast<UDMXEntityFixtureType>(WeakFixtureTypeObject.Get()))
-				{
-					FixtureType->SetGDTF(DMXImportGDTF);
-				}
+				UE_LOG(LogDMXEditor, Warning, TEXT("Failed to initialize Fixture Type '%s', falling back to legacy method that doesn't support matrix fixtures."), *FixtureType->GetName());
+				FixtureType->SetModesFromDMXImport(DMXImportGDTF);
 			}
+			FixtureType->PostEditChange();
 		}
 	}
 }
diff --git a/Engine/Plugins/VirtualProduction/DMX/DMXEngine/Source/DMXEditor/Private/DMXEditor.cpp b/Engine/Plugins/VirtualProduction/DMX/DMXEngine/Source/DMXEditor/Private/DMXEditor.cpp
--- a/Engine/Plugins/VirtualProduction/DMX/DMXEngine/Source/DMXEditor/Private/DMXEditor.cpp
+++ b/Engine/Plugins/VirtualProduction/DMX/DMXEngine/Source/DMXEditor/Private/DMXEditor.cpp
@@ -163,8 +163,9 @@ void FDMXEditor::ExportDMXLibrary() const
 			return;
 		}
 
-		const bool bExportSuccess = FDMXMVRExporter::Export(DMXLibrary, SaveFilenames[0]);
-		if (bExportSuccess)
+		FText ErrorReason;
+		FDMXMVRExporter::Export(DMXLibrary, SaveFilenames[0], ErrorReason);
+		if (ErrorReason.IsEmpty())
 		{
 			const FString SavePath = FPaths::GetPath(SaveFilenames[0]);
 			DMXEditorSettings->LastMVRExportPath = SavePath;
@@ -176,7 +177,7 @@ void FDMXEditor::ExportDMXLibrary() const
 		}
 		else
 		{
-			FNotificationInfo NotificationInfo(LOCTEXT("ExportDMXLibraryAsMVRFailureNotification", "Failed to export MVR"));
+			FNotificationInfo NotificationInfo(ErrorReason);
 			NotificationInfo.ExpireDuration = 10.f;
 			NotificationInfo.Image = FAppStyle::GetBrush("Icons.Warning");
 
diff --git a/Engine/Plugins/VirtualProduction/DMX/DMXEngine/Source/DMXEditor/Private/DMXInitializeFixtureTypeFromGDTFHelper.cpp b/Engine/Plugins/VirtualProduction/DMX/DMXEngine/Source/DMXEditor/Private/DMXInitializeFixtureTypeFromGDTFHelper.cpp
new file mode 100644
--- /dev/null
+++ b/Engine/Plugins/VirtualProduction/DMX/DMXEngine/Source/DMXEditor/Private/DMXInitializeFixtureTypeFromGDTFHelper.cpp
@@ -0,0 +1,397 @@
+// Copyright Epic Games, Inc. All Rights Reserved.
+
+#include "DMXInitializeFixtureTypeFromGDTFHelper.h"
+
+#include "DMXEditorLog.h"
+#include "DMXProtocolSettings.h"
+#include "DMXZipper.h"
+#include "Library/DMXEntityFixtureType.h"
+#include "Library/DMXGDTFAssetImportData.h"
+#include "Library/DMXImportGDTF.h"
+
+#include "XmlFile.h"
+#include "XmlNode.h"
+#include "Algo/Count.h"
+#include "Algo/Find.h"
+
+
+#define LOCTEXT_NAMESPACE "DMXInitializeFixtureTypeFromGDTFHelper"
+
+namespace UE::DMX::DMXInitializeFixtureTypeFromGDTFHelper::Private
+{
+	/**
+	 * Depending on the GDTF Spec, some XML Attribute Names might start with or without DMX.
+	 *
+	 * @param ParentNode		The Parent Node to search in
+	 * @param NodeName			The Attribute Name with DMX tag. This is the correct form, e.g. 'DMXMode', not 'Mode', 'DMXChannels', not 'Channels'.
+	 * @return					Returns a pointer to the child node, or nullptr if the child cannot be found
+	 */
+	const FXmlNode* FindChildNodeEvenIfDMXSubstringIsMissing(const FXmlNode& ParentNode, const FString& AttributeNameWithDMXTag)
+	{
+		if (const FXmlNode* ChildNodePtrWithSubstring = ParentNode.FindChildNode(AttributeNameWithDMXTag))
+		{
+			return ChildNodePtrWithSubstring;
+		}
+		else if (const FXmlNode* ChildNodePtrWithoutSubstring = ParentNode.FindChildNode(AttributeNameWithDMXTag.RightChop(3)))
+		{
+			return ChildNodePtrWithoutSubstring;
+		}
+
+		return nullptr;
+	}
+
+	/**
+	 * Depending on the GDTF Spec, some XML Attribute Names might start with or without DMX.
+	 *
+	 * @param ParentNode		The Parent Node to search in
+	 * @param NodeName			The Attribute Name with DMX tag. This is the correct form, e.g. 'DMXMode', not 'Mode', 'DMXChannels', not 'Channels'.
+	 * @return					Returns the Attribute as String
+	 */
+	FString FindAttributeEvenIfDMXSubstringIsMissing(const FXmlNode& ParentNode, const FString& AttributeNameWithDMXTag)
+	{
+		FString Attribute = ParentNode.GetAttribute(AttributeNameWithDMXTag);
+		if (Attribute.IsEmpty())
+		{
+			Attribute = ParentNode.GetAttribute(AttributeNameWithDMXTag.RightChop(3));
+		}
+
+		return Attribute;
+	}
+};
+
+bool FDMXInitializeFixtureTypeFromGDTFHelper::GenerateModesFromGDTF(UDMXEntityFixtureType& InOutFixtureType, const UDMXImportGDTF& InGDTF)
+{
+	FDMXInitializeFixtureTypeFromGDTFHelper Instance;
+	return Instance.GenerateModesFromGDTFInternal(InOutFixtureType, InGDTF);
+}
+
+bool FDMXInitializeFixtureTypeFromGDTFHelper::GenerateModesFromGDTFInternal(UDMXEntityFixtureType& InOutFixtureType, const UDMXImportGDTF& InGDTF) const
+{
+	UDMXGDTFAssetImportData* GDTFAssetImportData = InGDTF.GetGDTFAssetImportData();
+	if (!ensureMsgf(GDTFAssetImportData, TEXT("Found GDTF Asset that has no GDTF asset import data subobject.")))
+	{
+		return false;
+	}
+
+	const TSharedRef<FDMXZipper> Zip = MakeShared<FDMXZipper>();
+	if (!Zip->LoadFromData(GDTFAssetImportData->GetRawSourceData()))
+	{
+		return false;
+	}
+
+	const FDMXZipper::FDMXScopedUnzipToTempFile ScopedUnzippedDescriptionXml(Zip, TEXT("Description.xml"));
+	if (ScopedUnzippedDescriptionXml.TempFilePathAndName.IsEmpty())
+	{
+		return false;
+	}
+
+	const TSharedRef<FXmlFile> XmlFile = MakeShared<FXmlFile>();
+	if (!XmlFile->LoadFile(ScopedUnzippedDescriptionXml.TempFilePathAndName))
+	{
+		return false;
+	}
+
+	// Largely a copy of GDTF Factory, but avoiding the asset
+	const FXmlNode* RootNode = XmlFile->GetRootNode();
+	if (!RootNode)
+	{
+		return false;
+	}
+
+	constexpr TCHAR FixtureTypeTag[] = TEXT("FixtureType");
+	const FXmlNode* FixtureTypeNode = RootNode->FindChildNode(FixtureTypeTag);
+	if (!FixtureTypeNode)
+	{
+		return false;
+	}
+
+	// Iterate modes and build the DMXFixtureMode array from here
+	const FXmlNode* const ModesNode = UE::DMX::DMXInitializeFixtureTypeFromGDTFHelper::Private::FindChildNodeEvenIfDMXSubstringIsMissing(*FixtureTypeNode, TEXT("DMXModes"));
+	if (!ModesNode)
+	{
+		return false;
+	}
+
+	for (const FXmlNode* DMXModeNode : ModesNode->GetChildrenNodes())
+	{
+		FDMXFixtureMode Mode;
+		if (GenerateMode(*FixtureTypeNode, *DMXModeNode, Mode))
+		{
+			InOutFixtureType.Modes.Add(Mode);
+		}
+	}
+
+	for (int32 ModeIndex = 0; ModeIndex < InOutFixtureType.Modes.Num(); ModeIndex++)
+	{
+		InOutFixtureType.UpdateChannelSpan(ModeIndex);
+	}
+	InOutFixtureType.GetOnFixtureTypeChanged().Broadcast(&InOutFixtureType);
+
+	CleanupAttributes(InOutFixtureType);
+
+	return true;
+}
+
+bool FDMXInitializeFixtureTypeFromGDTFHelper::GenerateMode(const FXmlNode& InFixtureTypeNode, const FXmlNode& InDMXModeNode, FDMXFixtureMode& OutMode) const
+{
+	if (!ensureMsgf(InFixtureTypeNode.GetTag() == TEXT("fixturetype"), TEXT("Trying to read Fixture Type node, but node is tagged '%s'."), *InFixtureTypeNode.GetTag()))
+	{
+		return false;
+	}
+	if (!ensureMsgf(InDMXModeNode.GetTag() == TEXT("dmxmode") || InDMXModeNode.GetTag() == TEXT("mode"), TEXT("Trying to read Channel node, but node is tagged '%s'."), *InDMXModeNode.GetTag()))
+	{
+		return false;
+	}
+
+	// Get geometry nodes specific to this mode (can be empty)
+	TArray<const FXmlNode*> GeometryNodes;
+	constexpr TCHAR GeometriesTag[] = TEXT("Geometries");
+	if (const FXmlNode* GeometriesNode = InFixtureTypeNode.FindChildNode(GeometriesTag))
+	{
+		constexpr TCHAR GeometryTag[] = TEXT("Geometry");
+		const FString GeometryName = InDMXModeNode.GetAttribute(GeometryTag);
+
+		// Root Geometry of this Mode
+		const FXmlNode* const* GeometryNodePtr = Algo::FindByPredicate(GeometriesNode->GetChildrenNodes(), [&GeometryName, GeometryTag](const FXmlNode* GeometryNode)
+			{
+				constexpr TCHAR NameTag[] = TEXT("Name");
+				if (GeometryNode->GetAttribute(NameTag) == GeometryName)
+				{
+					return true;
+				}
+
+				return false;
+			});
+
+		if (GeometryNodePtr)
+		{
+			GeometryNodes = GetChildrenRecursive(**GeometryNodePtr);
+		}
+	}
+
+	// Depending on the GDTF spec in use modes may be stored in the "Modes" or "DMXModes" node.
+	const FXmlNode* const ChannelsNode = UE::DMX::DMXInitializeFixtureTypeFromGDTFHelper::Private::FindChildNodeEvenIfDMXSubstringIsMissing(InDMXModeNode, TEXT("DMXChannels"));
+	if (!ChannelsNode)
+	{
+		return false;
+	}
+
+	FDMXFixtureMode Mode;
+	Mode.FixtureMatrixConfig.CellAttributes.Reset();
+
+	constexpr TCHAR NameTag[] = TEXT("Name");
+	Mode.ModeName = InDMXModeNode.GetAttribute(NameTag); // Not using the class's default attribute
+
+	int32 NextFunctionStartingChannel = 1;
+	for (const FXmlNode* DMXChannelNode : ChannelsNode->GetChildrenNodes())
+	{
+		FDMXFixtureFunction Function;
+		if (!GetChannelPropertiesWithoutStartingChannel(*DMXChannelNode, Function))
+		{
+			continue;
+		}
+
+		UE_LOG(LogTemp, Warning, TEXT("Channel: %i"), NextFunctionStartingChannel);
+
+		constexpr TCHAR GeometryTag[] = TEXT("Geometry");
+		const FString GeometryName = DMXChannelNode->GetAttribute(GeometryTag);
+
+		// Count the number of geometries this channel references
+		const int32 NumReferencesInGeometry = Algo::CountIf(GeometryNodes, [&GeometryName, GeometryTag](const FXmlNode* GeometryNode)
+			{
+				if (GeometryNode->GetAttribute(GeometryTag) == GeometryName)
+				{
+					return true;
+				}
+
+				return false;
+			});
+
+		if (NumReferencesInGeometry < 2)
+		{
+			// A single or no referenced Geometry - Create a normal function.
+			Function.Channel = FMath::Max(NextFunctionStartingChannel, Mode.FixtureMatrixConfig.GetLastChannel() + 1);
+			Mode.Functions.Add(Function);
+		}
+		else
+		{
+			// Try to add a matrix function instead of a common function
+			if (Mode.bFixtureMatrixEnabled &&
+				Mode.Functions.Num() > 0 &&
+				Mode.FixtureMatrixConfig.GetLastChannel() > Mode.Functions.Last().GetLastChannel())
+			{
+				UE_LOG(LogDMXEditor, Warning, TEXT("Mode '%s' contains many matrices, but this version of Unreal Engine only supports one matrix. Skipping import of mode."), *Mode.ModeName)
+					Mode.Functions.Reset();
+				Mode.bFixtureMatrixEnabled = false;
+				Mode.FixtureMatrixConfig.CellAttributes.Reset();
+				Mode.ModeName = FString::Printf(TEXT("n/a '%s' [not supported in this Engine Version]"), *Mode.ModeName);
+				break;
+			}
+
+			FDMXFixtureCellAttribute MatrixAttribute;
+			MatrixAttribute.Attribute = Function.Attribute;
+			MatrixAttribute.bUseLSBMode = Function.bUseLSBMode;
+			MatrixAttribute.DataType = Function.DataType;
+			MatrixAttribute.DefaultValue = Function.DefaultValue;
+
+			Mode.FixtureMatrixConfig.CellAttributes.Add(MatrixAttribute);
+			Mode.FixtureMatrixConfig.YCells = NumReferencesInGeometry;
+			if (!Mode.bFixtureMatrixEnabled)
+			{
+				Mode.FixtureMatrixConfig.FirstCellChannel = NextFunctionStartingChannel;
+				Mode.bFixtureMatrixEnabled = true;
+			}
+		}
+
+		NextFunctionStartingChannel += Function.GetNumChannels();
+	}
+
+	OutMode = Mode;
+	return true;
+}
+
+bool FDMXInitializeFixtureTypeFromGDTFHelper::GetChannelPropertiesWithoutStartingChannel(const FXmlNode& DMXChannelNode, FDMXFixtureFunction& OutFunction) const
+{
+	if (!ensureMsgf(DMXChannelNode.GetTag() == TEXT("dmxchannel") || DMXChannelNode.GetTag() == TEXT("channel"), TEXT("Trying to read Channel node, but node is tagged '%s'."), *DMXChannelNode.GetTag()))
+	{
+		return false;
+	}
+
+	// As per legacy consider only the first function, although there may be many.
+	constexpr TCHAR LogicalChannelTag[] = TEXT("LogicalChannel");
+	const FXmlNode* LogicalChannelNode = DMXChannelNode.FindChildNode(LogicalChannelTag);
+	if (!LogicalChannelNode)
+	{
+		return false;
+	}
+
+	constexpr TCHAR ChannelFunctionTag[] = TEXT("ChannelFunction");
+	const FXmlNode* ChannelFunctionNode = LogicalChannelNode->FindChildNode(ChannelFunctionTag);
+	if (!ChannelFunctionNode)
+	{
+		return false;
+	}
+
+	// Name
+	constexpr TCHAR NameTag[] = TEXT("Name");
+	const FString Name = ChannelFunctionNode->GetAttribute(NameTag);
+
+	OutFunction.FunctionName = Name.IsEmpty() ? TEXT("<empty>") : Name;
+	OutFunction.Attribute.Name = Name.IsEmpty() ? NAME_None : FName(*Name);
+	
+	// Default value
+	constexpr TCHAR DefaultTag[] = TEXT("Default");
+	int32 DefaultValue;
+	if (LexTryParseString(DefaultValue, *ChannelFunctionNode->GetAttribute(DefaultTag)))
+	{
+		OutFunction.DefaultValue = DefaultValue;
+	}
+
+	// Word size and endianness
+	constexpr TCHAR OffsetTag[] = TEXT("Offset");
+	TArray<int32> OffsetArray;
+	TArray<FString> OffsetStrArray;
+	DMXChannelNode.GetAttribute(OffsetTag).ParseIntoArray(OffsetStrArray, TEXT(","));
+
+	for (int32 OffsetIndex = 0; OffsetIndex < OffsetStrArray.Num(); ++OffsetIndex)
+	{
+		int32 OffsetValue;
+		LexTryParseString(OffsetValue, *OffsetStrArray[OffsetIndex]);
+		OffsetArray.Add(OffsetValue);
+	}
+
+	if (OffsetArray.Num() > 0)
+	{
+		// Compute number of used addresses
+		int32 AddressMin = DMX_MAX_ADDRESS;
+		int32 AddressMax = 0;
+		for (const int32& Address : OffsetArray)
+		{
+			AddressMin = FMath::Min(AddressMin, Address);
+			AddressMax = FMath::Max(AddressMax, Address);
+		}
+		const int32 NumUsedAddresses = FMath::Clamp(AddressMax - AddressMin + 1, 1, DMX_MAX_FUNCTION_SIZE);
+
+		OutFunction.DataType = static_cast<EDMXFixtureSignalFormat>(NumUsedAddresses - 1);
+
+		// Offsets represent the channels in MSB order. If they are in reverse order, it means this Function uses LSB format.
+		if (OffsetArray.Num() > 1)
+		{
+			OutFunction.bUseLSBMode = OffsetArray[0] > OffsetArray[1];
+		}
+		else
+		{
+			OutFunction.bUseLSBMode = false;
+		}
+	}
+	else
+	{
+		OutFunction.DataType = EDMXFixtureSignalFormat::E8Bit;
+	}
+
+	return true;
+}
+
+void FDMXInitializeFixtureTypeFromGDTFHelper::CleanupAttributes(UDMXEntityFixtureType& InOutFixtureType) const
+{
+	// Get Protocol Setting's default attributes
+	const UDMXProtocolSettings* ProtocolSettings = GetDefault<UDMXProtocolSettings>();
+	TMap<FName, TArray<FString> > AttributeNameToKeywordsMap;
+	for (const FDMXAttribute& Attribute : ProtocolSettings->Attributes)
+	{
+		TArray<FString> Keywords = Attribute.GetKeywords();
+
+		AttributeNameToKeywordsMap.Emplace(Attribute.Name, Keywords);
+	}
+
+	TArray<FName> AssignedAttributeNames;
+	for (FDMXFixtureMode& Mode : InOutFixtureType.Modes)
+	{
+		for (FDMXFixtureFunction& Function : Mode.Functions)
+		{
+			const TTuple<FName, TArray<FString>>* AttributeNameToKeywordPairPtr = Algo::FindByPredicate(AttributeNameToKeywordsMap, [Function](const TTuple<FName, TArray<FString>> AttributeNameToKeywordPair)
+				{
+					if (Function.Attribute.Name == AttributeNameToKeywordPair.Key ||
+						AttributeNameToKeywordPair.Value.Contains(Function.Attribute.Name.ToString()))
+					{
+						return true;
+					}
+
+					return false;
+				});
+
+			if (AttributeNameToKeywordPairPtr && !AssignedAttributeNames.Contains(AttributeNameToKeywordPairPtr->Key))
+			{
+				Function.Attribute.Name = AttributeNameToKeywordPairPtr->Key;
+				AssignedAttributeNames.Add(AttributeNameToKeywordPairPtr->Key);
+			}
+			else
+			{
+				if (AssignedAttributeNames.Contains(Function.Attribute.Name))
+				{
+					UE_LOG(LogDMXEditor, Warning, TEXT("Attribute '%s' occurs twice in Fixture Type '%s'. Reseting second occurence to None"), *Function.Attribute.Name.ToString(), *InOutFixtureType.GetName());
+				}
+				else
+				{
+					// Attribute not specified in project settings, but unique
+					AssignedAttributeNames.Add(Function.Attribute.Name);
+				}
+			}
+		}
+	}
+}
+
+TArray<const FXmlNode*> FDMXInitializeFixtureTypeFromGDTFHelper::GetChildrenRecursive(const FXmlNode& ParentNode) const
+{
+	TArray<const FXmlNode*> Result;
+
+	for (const FXmlNode* Child : ParentNode.GetChildrenNodes())
+	{
+		Result.Add(Child);
+		Result.Append(GetChildrenRecursive(*Child));
+	}
+
+	return Result;
+}
+
+#undef LOCTEXT_NAMESPACE
diff --git a/Engine/Plugins/VirtualProduction/DMX/DMXEngine/Source/DMXEditor/Private/DMXInitializeFixtureTypeFromGDTFHelper.h b/Engine/Plugins/VirtualProduction/DMX/DMXEngine/Source/DMXEditor/Private/DMXInitializeFixtureTypeFromGDTFHelper.h
new file mode 100644
--- /dev/null
+++ b/Engine/Plugins/VirtualProduction/DMX/DMXEngine/Source/DMXEditor/Private/DMXInitializeFixtureTypeFromGDTFHelper.h
@@ -0,0 +1,45 @@
+// Copyright Epic Games, Inc. All Rights Reserved.
+
+#pragma once
+
+#include "CoreMinimal.h"
+
+class UDMXEntityFixtureType;
+struct FDMXFixtureFunction;
+struct FDMXFixtureMode;
+class FDMXZipper;
+class UDMXImportGDTF;
+
+class FXmlNode;
+
+
+/**
+ * This is a 5.1 private editor only workaround to initialize Fixture Types from GDTF assets.
+ * 
+ * Required since DMXImportGDTF does not hold required data, hence UDMXEntityFixtureType::SetModesFromDMXImport cannot create matrix fixture types.
+ *
+ * Opeted for this workaround since it would require to rework the entirety of the GDTF implementation to hold its data correctly. 
+ * This exceeds planed features of 5.1 and is a planned feature for 5.2. 
+ */
+class FDMXInitializeFixtureTypeFromGDTFHelper
+{
+public:
+	/** Imports the modes directly from the asset import data */
+	static bool GenerateModesFromGDTF(UDMXEntityFixtureType& FixtureType, const UDMXImportGDTF& InGDTF);
+
+private:
+	/** Imports the modes directly from the asset import data. Non-static implementation */
+	bool GenerateModesFromGDTFInternal(UDMXEntityFixtureType& InOutFixtureType, const UDMXImportGDTF& InGDTF) const;
+
+	/** Generates a Mode from specified Nodes. Returns true on success. */
+	bool GenerateMode(const FXmlNode& InFixtureTypeNode, const FXmlNode& InDMXModeNode, FDMXFixtureMode& OutMode) const;
+
+	/** Gets the properties of the channel node, as a fixture function. Does not set the starting channel. */
+	bool GetChannelPropertiesWithoutStartingChannel(const FXmlNode& DMXChannelNode, FDMXFixtureFunction& OutFunction) const;
+
+	/** Makes sure that attributes are replaced with attributes specified in protocol settings where possible */
+	void CleanupAttributes(UDMXEntityFixtureType& InOutFixtureType) const;
+
+	/** Returns all nodes that are decendants of the node */
+	TArray<const FXmlNode*> GetChildrenRecursive(const FXmlNode& ParentNode) const;
+};
diff --git a/Engine/Plugins/VirtualProduction/DMX/DMXEngine/Source/DMXEditor/Private/Exporters/DMXMVRExporter.cpp b/Engine/Plugins/VirtualProduction/DMX/DMXEngine/Source/DMXEditor/Private/Exporters/DMXMVRExporter.cpp
--- a/Engine/Plugins/VirtualProduction/DMX/DMXEngine/Source/DMXEditor/Private/Exporters/DMXMVRExporter.cpp
+++ b/Engine/Plugins/VirtualProduction/DMX/DMXEngine/Source/DMXEditor/Private/Exporters/DMXMVRExporter.cpp
@@ -3,6 +3,8 @@
 #include "Exporters/DMXMVRExporter.h"
 
 #include "DMXEditorLog.h"
+#include "DMXEditorSettings.h"
+#include "DMXInitializeFixtureTypeFromGDTFHelper.h"
 #include "DMXMVRXmlMergeUtility.h"
 #include "DMXZipper.h"
 #include "Library/DMXEntityFixturePatch.h"
@@ -13,41 +15,62 @@
 #include "MVR/DMXMVRAssetImportData.h"
 #include "MVR/DMXMVRGeneralSceneDescription.h"
 
+#include "DesktopPlatformModule.h"
+#include "IDesktopPlatform.h"
 #include "XmlFile.h"
+#include "Framework/Notifications/NotificationManager.h"
 #include "Misc/FileHelper.h"
+#include "Misc/MessageDialog.h"
 #include "Misc/Paths.h"
+#include "Widgets/Notifications/SNotificationList.h"
 
 
-bool FDMXMVRExporter::Export(UDMXLibrary* DMXLibrary, const FString& FilePathAndName)
+#define LOCTEXT_NAMESPACE "DMXMVRExporter"
+
+void FDMXMVRExporter::Export(UDMXLibrary* DMXLibrary, const FString& FilePathAndName, FText& OutErrorReason)
 {
+
 	FDMXMVRExporter Instance;
-	return Instance.ExportInternal(DMXLibrary, FilePathAndName);
+	Instance.ExportInternal(DMXLibrary, FilePathAndName, OutErrorReason);
 }
 
-bool FDMXMVRExporter::ExportInternal(UDMXLibrary* DMXLibrary, const FString& FilePathAndName)
+void FDMXMVRExporter::ExportInternal(UDMXLibrary* DMXLibrary, const FString& FilePathAndName, FText& OutErrorReason)
 {
+	OutErrorReason = FText::GetEmpty();
+
 	if (!ensureAlwaysMsgf(DMXLibrary, TEXT("Trying to export DMX Library '%s' as MVR file, but the DMX Library is invalid."), *DMXLibrary->GetName()))
 	{
-		return false;
+		OutErrorReason = FText::Format(LOCTEXT("MVRExportDMXLibraryInvalidReason", "DMX Library is invalid. Cannot export {0}."), FText::FromString(FilePathAndName));
+		return;
 	}
 
 	DMXLibrary->UpdateGeneralSceneDescription();
 	const UDMXMVRGeneralSceneDescription* GeneralSceneDescription = DMXLibrary->GetLazyGeneralSceneDescription();
 	if (!ensureAlwaysMsgf(GeneralSceneDescription, TEXT("Trying to export DMX Library '%s' as MVR file, but its General Scene Description is invalid."), *DMXLibrary->GetName()))
 	{
-		return false;
+		OutErrorReason = FText::Format(LOCTEXT("MVRExportGeneralSceneDescriptionInvalidReason", "DMX Library is invalid. Cannot export {0}."), FText::FromString(FilePathAndName));
+		return;
 	}
 
 	const TSharedRef<FDMXZipper> Zip = MakeShared<FDMXZipper>();
 	if (!ZipGeneralSceneDescription(Zip, GeneralSceneDescription))
 	{
-		return false;
+		OutErrorReason = FText::Format(LOCTEXT("MVRExportCreateZipFailedReason", "Failed to write General Scene Description. Cannot export {0}."), FText::FromString(FilePathAndName));
+		return;
 	}
 
-	ZipGDTFs(Zip, DMXLibrary);
-	ZipThirdPartyData(Zip, GeneralSceneDescription);
+	if (!ZipGDTFs(Zip, DMXLibrary))
+	{
+		OutErrorReason = FText::Format(LOCTEXT("MVRExportZipGDTFsFailedReason", "Some GDTFs are missing import data and were not considered. Exported MVR to {0}."), FText::FromString(FilePathAndName));
+		// Allow continuation of export
+	}
 
-	return Zip->SaveToFile(FilePathAndName);
+	ZipThirdPartyData(Zip, GeneralSceneDescription);
+	if (!Zip->SaveToFile(FilePathAndName))
+	{
+		OutErrorReason = FText::Format(LOCTEXT("MVRExportCreateZipFailedReason", "File is not writable or locked by another process. Cannot export {0}."), FText::FromString(FilePathAndName));
+		return;
+	}
 }
 
 bool FDMXMVRExporter::ZipGeneralSceneDescription(const TSharedRef<FDMXZipper>& Zip, const UDMXMVRGeneralSceneDescription* GeneralSceneDescription)
@@ -86,46 +109,62 @@ bool FDMXMVRExporter::ZipGeneralSceneDescription(const TSharedRef<FDMXZipper>& Z
 	return true;
 }
 
-void FDMXMVRExporter::ZipGDTFs(const TSharedRef<FDMXZipper>& Zip, UDMXLibrary* DMXLibrary)
+bool FDMXMVRExporter::ZipGDTFs(const TSharedRef<FDMXZipper>& Zip, UDMXLibrary* DMXLibrary)
 {
 	check(DMXLibrary);
 
 	// Gather GDTFs
-	TArray<const UDMXGDTFAssetImportData*> GDTFAssetImportDataToExport;
-	const TArray<const UDMXEntityFixturePatch*> FixturePatches = DMXLibrary->GetEntitiesTypeCast<const UDMXEntityFixturePatch>();
-	for (const UDMXEntityFixturePatch* FixturePatch : FixturePatches)
+	TArray<UDMXEntityFixtureType*> FixtureTypesToExport;
+	const TArray<UDMXEntityFixturePatch*> FixturePatches = DMXLibrary->GetEntitiesTypeCast<UDMXEntityFixturePatch>();
+	for (UDMXEntityFixturePatch* FixturePatch : FixturePatches)
 	{
-		const UDMXEntityFixtureType* FixtureType = FixturePatch->GetFixtureType();
-		if (!FixtureType)
+		UDMXEntityFixtureType* FixtureType = FixturePatch->GetFixtureType();
+		if (!FixtureType || FixtureType->Modes.IsEmpty())
 		{
 			continue;
 		}
-		else if (!FixtureType->GDTF)
+
+		FixtureTypesToExport.AddUnique(FixtureType);
+	}
+
+	// Zip GDTFs
+	bool bAllZippedSuccessfully = true;
+	for (UDMXEntityFixtureType* FixtureType : FixtureTypesToExport)
+	{
+		if (!FixtureType->DMXImport)
 		{
-			UE_LOG(LogDMXEditor, Warning, TEXT("Cannot export Fixture Patch '%s' to MVR. Its Fixture Type '%s' has no GDTF set."), *FixturePatch->Name, *FixtureType->Name);
+			UE_LOG(LogDMXEditor, Warning, TEXT("Cannot export Fixture Type '%s' to MVR, but the Fixture Type has no GDTF set."), *FixtureType->Name);
 			continue;
 		}
 
-		const UDMXGDTFAssetImportData* GDTFAssetImportData = FixtureType->GDTF->GetGDTFAssetImportData();
-		if (!GDTFAssetImportData || GDTFAssetImportData->GetRawSourceData().IsEmpty())
+		UDMXImportGDTF* DMXImportGDTF = Cast<UDMXImportGDTF>(FixtureType->DMXImport);
+		if (!DMXImportGDTF)
 		{
-			UE_LOG(LogDMXEditor, Warning, TEXT("Cannot export Fixture Patch '%s' to MVR. Its Fixture Type '%' is from a version prior to 5.1, and does not hold required GDTF data."), *FixturePatch->Name, *FixtureType->Name);
+			UE_LOG(LogDMXEditor, Warning, TEXT("Cannot export Fixture Type '%s' to MVR, but the Fixture Type has a DMX Import Type which is not GDTF."), *FixtureType->Name);
 			continue;
 		}
 
-		GDTFAssetImportDataToExport.AddUnique(GDTFAssetImportData);
-	}
+		UDMXGDTFAssetImportData* GDTFAssetImportData = DMXImportGDTF->GetGDTFAssetImportData();
+		if (!ensureAlwaysMsgf(GDTFAssetImportData, TEXT("Missing default GDTF Asset Import Data subobject in GDTF.")))
+		{
+			continue;
+		}
 
-	// Zip GDTFs
-	TArray<FString> TempGDTFFilePathAndNames;
-	for (const UDMXGDTFAssetImportData* GDTFAssetImportData : GDTFAssetImportDataToExport)
-	{
-		if (GDTFAssetImportData->GetRawSourceData().Num() > 0)
+		const TArray64<uint8>& RawSourceData = RefreshSourceDataAndFixtureType(*FixtureType, *GDTFAssetImportData);
+		if (RawSourceData.Num() > 0)
 		{
 			const FString GDTFFilename = FPaths::GetCleanFilename(GDTFAssetImportData->GetSourceFilePathAndName());
 			Zip->AddFile(GDTFFilename, GDTFAssetImportData->GetRawSourceData());
 		}
+		else
+		{
+			bAllZippedSuccessfully = false;
+			const UDMXImportGDTF* GDTF = Cast<UDMXImportGDTF>(GDTFAssetImportData->GetOuter());
+			UE_LOG(LogDMXEditor, Error, TEXT("Cannot export '%s' to MVR File. The asset is missing source data."), GDTF ? *GDTF->GetName() : TEXT("Invalid GDTF Asset"));
+		}
 	}
+
+	return bAllZippedSuccessfully;
 }
 
 void FDMXMVRExporter::ZipThirdPartyData(const TSharedRef<FDMXZipper>& Zip, const UDMXMVRGeneralSceneDescription* GeneralSceneDescription)
@@ -196,3 +235,84 @@ const TSharedPtr<FXmlFile> FDMXMVRExporter::CreateSourceGeneralSceneDescriptionX
 
 	return XmlFile;
 }
+
+
+const TArray64<uint8>& FDMXMVRExporter::RefreshSourceDataAndFixtureType(UDMXEntityFixtureType& FixtureType, UDMXGDTFAssetImportData& InOutGDTFAssetImportData) const
+{
+	const UDMXImportGDTF* GDTF = Cast<UDMXImportGDTF>(InOutGDTFAssetImportData.GetOuter());
+	if (!InOutGDTFAssetImportData.GetRawSourceData().IsEmpty() || !GDTF)
+	{
+		return InOutGDTFAssetImportData.GetRawSourceData();
+	}
+
+	static bool bAskForReloadAgain = true;
+	static EAppReturnType::Type MessageDialogResult = EAppReturnType::Yes;
+
+	if (bAskForReloadAgain)
+	{
+		static const FText MessageTitle = LOCTEXT("NoGDTFSourceAvailableTitle", "Trying to use old GDTF asset.");
+		static const FText Message = FText::Format(LOCTEXT("NoGDTFSourceAvailableMessage", "Insufficient data to export '{0}' to MVR file. The GDTF asset was created prior to UE5.1. Do you want to reload the source GDTF?"), FText::FromString(GDTF->GetName()));
+		MessageDialogResult = FMessageDialog::Open(EAppMsgType::YesNoYesAllNoAll, Message, &MessageTitle);
+
+		if (MessageDialogResult == EAppReturnType::YesAll || MessageDialogResult == EAppReturnType::NoAll)
+		{
+			bAskForReloadAgain = false;
+		}
+	}
+
+	if (MessageDialogResult == EAppReturnType::YesAll || MessageDialogResult == EAppReturnType::Yes)
+	{
+		IDesktopPlatform* const DesktopPlatform = FDesktopPlatformModule::Get();
+
+		if (DesktopPlatform)
+		{
+			UDMXEditorSettings* EditorSettings = GetMutableDefault<UDMXEditorSettings>();
+			if (!ensureAlwaysMsgf(EditorSettings, TEXT("Unexpected cannot access DMX Editor Settings CDO.")))
+			{
+				return InOutGDTFAssetImportData.GetRawSourceData();
+			}
+
+			TArray<FString> Filenames;
+			if (!InOutGDTFAssetImportData.GetSourceData().SourceFiles.IsEmpty() &&
+				FPaths::FileExists(InOutGDTFAssetImportData.GetSourceFilePathAndName()))
+			{
+				Filenames.Add(InOutGDTFAssetImportData.GetSourceFilePathAndName());
+			}
+			else
+			{
+				DesktopPlatform->OpenFileDialog(
+					nullptr,
+					FText::Format(LOCTEXT("OpenGDTFTitle", "Choose a GDTF file for '%s'."), FText::FromString(GDTF->GetName())).ToString(),
+					EditorSettings->LastGDTFImportPath,
+					TEXT(""),
+					TEXT("General Scene Description (*.gdtf)|*.gdtf"),
+					EFileDialogFlags::None,
+					Filenames
+				);
+			}
+
+			if (Filenames.Num() > 0)
+			{
+				EditorSettings->LastGDTFImportPath = FPaths::GetPath(Filenames[0]);
+
+				InOutGDTFAssetImportData.PreEditChange(nullptr);
+				InOutGDTFAssetImportData.SetSourceFile(Filenames[0]);
+				InOutGDTFAssetImportData.PostEditChange();
+
+				FDMXInitializeFixtureTypeFromGDTFHelper::GenerateModesFromGDTF(FixtureType, *GDTF);
+
+				if (InOutGDTFAssetImportData.GetRawSourceData().Num() == 0)
+				{
+					FNotificationInfo NotificationInfo(FText::Format(LOCTEXT("ReloadGDTFFailure", "Failed to update GDTF '{0}' from '{1}."), FText::FromString(GDTF->GetName()), FText::FromString(Filenames[0])));
+					NotificationInfo.ExpireDuration = 10.f;
+
+					FSlateNotificationManager::Get().AddNotification(NotificationInfo);
+				}
+			}
+		}
+	}
+
+	return InOutGDTFAssetImportData.GetRawSourceData();
+}
+
+#undef LOCTEXT_NAMESPACE
diff --git a/Engine/Plugins/VirtualProduction/DMX/DMXEngine/Source/DMXEditor/Private/Exporters/DMXMVRExporter.h b/Engine/Plugins/VirtualProduction/DMX/DMXEngine/Source/DMXEditor/Private/Exporters/DMXMVRExporter.h
--- a/Engine/Plugins/VirtualProduction/DMX/DMXEngine/Source/DMXEditor/Private/Exporters/DMXMVRExporter.h
+++ b/Engine/Plugins/VirtualProduction/DMX/DMXEngine/Source/DMXEditor/Private/Exporters/DMXMVRExporter.h
@@ -5,6 +5,8 @@
 #include "CoreMinimal.h"
 
 class FDMXZipper;
+class UDMXEntityFixtureType;
+class UDMXGDTFAssetImportData;
 class UDMXLibrary;
 class UDMXMVRAssetImportData;
 class UDMXMVRGeneralSceneDescription;
@@ -17,21 +19,27 @@ class FDMXMVRExporter
 {
 public:
 	/** Exports the DMX Library as MVR File */
-	static bool Export(UDMXLibrary* DMXLibrary, const FString& FilePathAndName);
+	static void Export(UDMXLibrary* DMXLibrary, const FString& FilePathAndName, FText& OutErrorReason);
 
 private:
-	/** Exports the DMX Library as MVR File */
-	UE_NODISCARD bool ExportInternal(UDMXLibrary* DMXLibrary, const FString& FilePathAndName);
+	/** Exports the DMX Library as MVR File. If OutErrorReason is not empty there were issues with the export. */
+	void ExportInternal(UDMXLibrary* DMXLibrary, const FString& FilePathAndName, FText& OutErrorReason);
 
 	/** Zips the GeneralSceneDescription.xml */
 	UE_NODISCARD bool ZipGeneralSceneDescription(const TSharedRef<FDMXZipper>& Zip, const UDMXMVRGeneralSceneDescription* GeneralSceneDescription);
 
 	/** Zips the GDTFs from the Library */
-	void ZipGDTFs(const TSharedRef<FDMXZipper>& Zip, UDMXLibrary* DMXLibrary);
+	UE_NODISCARD bool ZipGDTFs(const TSharedRef<FDMXZipper>& Zip, UDMXLibrary* DMXLibrary);
 
 	/** Zips 3rd Party Data from the MVR Asset Import Data */
 	void ZipThirdPartyData(const TSharedRef<FDMXZipper>& Zip, const UDMXMVRGeneralSceneDescription* GeneralSceneDescription);
 
 	/** Creates an General Scene Description Xml File from the MVR Source, as it was imported */
 	const TSharedPtr<FXmlFile> CreateSourceGeneralSceneDescriptionXmlFile(const UDMXMVRGeneralSceneDescription* GeneralSceneDescription) const;
+
+	/** 
+	 * Gets raw source data or creates (possibly empty) source data where the source data is not present. 
+	 * Prior 5.1 there was no source data stored. Offers a dialog to load missing data 
+	 */
+	const TArray64<uint8>& RefreshSourceDataAndFixtureType(UDMXEntityFixtureType& FixtureType, UDMXGDTFAssetImportData& InOutGDTFAssetImportData) const;
 };
diff --git a/Engine/Plugins/VirtualProduction/DMX/DMXEngine/Source/DMXEditor/Private/Factories/DMXLibraryFromMVRFactory.cpp b/Engine/Plugins/VirtualProduction/DMX/DMXEngine/Source/DMXEditor/Private/Factories/DMXLibraryFromMVRFactory.cpp
--- a/Engine/Plugins/VirtualProduction/DMX/DMXEngine/Source/DMXEditor/Private/Factories/DMXLibraryFromMVRFactory.cpp
+++ b/Engine/Plugins/VirtualProduction/DMX/DMXEngine/Source/DMXEditor/Private/Factories/DMXLibraryFromMVRFactory.cpp
@@ -254,7 +254,7 @@ void UDMXLibraryFromMVRFactory::InitDMXLibrary(UDMXLibrary* DMXLibrary, const TA
 		FixtureTypeConstructionParams.ParentDMXLibrary = DMXLibrary;
 
 		UDMXEntityFixtureType* FixtureType = UDMXEntityFixtureType::CreateFixtureTypeInLibrary(FixtureTypeConstructionParams, FPaths::GetBaseFilename(GDTFFilename));
-		FixtureType->SetGDTF(GDTF);
+		FixtureType->SetModesFromDMXImport(GDTF);
 
 		GDTFSpecToFixtureTypeMap.Add(GDTFFilename, FixtureType);
 	}
diff --git a/Engine/Plugins/VirtualProduction/DMX/DMXEngine/Source/DMXEditor/Private/Widgets/FixtureType/SDMXFixtureFunctionEditor.cpp b/Engine/Plugins/VirtualProduction/DMX/DMXEngine/Source/DMXEditor/Private/Widgets/FixtureType/SDMXFixtureFunctionEditor.cpp
--- a/Engine/Plugins/VirtualProduction/DMX/DMXEngine/Source/DMXEditor/Private/Widgets/FixtureType/SDMXFixtureFunctionEditor.cpp
+++ b/Engine/Plugins/VirtualProduction/DMX/DMXEngine/Source/DMXEditor/Private/Widgets/FixtureType/SDMXFixtureFunctionEditor.cpp
@@ -89,8 +89,12 @@ void SDMXFixtureFunctionEditor::NotifyPostChange(const FPropertyChangedEvent& Pr
 				int32 ResultingStartingChannel;
 				FixtureType->SetFunctionStartingChannel(ModeIndex, FunctionIndex, FunctionBeingEditedPtr->Channel, ResultingStartingChannel);
 			}
-			else if (PropertyName == GET_MEMBER_NAME_CHECKED(FDMXFixtureFunction, DataType) ||
-				PropertyName == GET_MEMBER_NAME_CHECKED(FDMXFixtureFunction, DefaultValue))
+			else if (PropertyName == GET_MEMBER_NAME_CHECKED(FDMXFixtureFunction, DataType))
+			{
+				FixtureType->AlignFunctionChannels(ModeIndex);
+				FixtureType->ClampFunctionDefautValueByDataType(ModeIndex, FunctionIndex);
+			}
+			else if (PropertyName == GET_MEMBER_NAME_CHECKED(FDMXFixtureFunction, DefaultValue))
 			{
 				FixtureType->ClampFunctionDefautValueByDataType(ModeIndex, FunctionIndex);
 			}
diff --git a/Engine/Plugins/VirtualProduction/DMX/DMXEngine/Source/DMXEditor/Public/DMXEditorSettings.h b/Engine/Plugins/VirtualProduction/DMX/DMXEngine/Source/DMXEditor/Public/DMXEditorSettings.h
--- a/Engine/Plugins/VirtualProduction/DMX/DMXEngine/Source/DMXEditor/Public/DMXEditorSettings.h
+++ b/Engine/Plugins/VirtualProduction/DMX/DMXEngine/Source/DMXEditor/Public/DMXEditorSettings.h
@@ -125,6 +125,11 @@ class UDMXEditorSettings : public UObject
 {
 	GENERATED_BODY()
 
+	// GDTF
+public:
+	UPROPERTY(Config)
+	FString LastGDTFImportPath;
+
 
 	// MVR
 public:
diff --git a/Engine/Plugins/VirtualProduction/DMX/DMXEngine/Source/DMXRuntime/Private/Library/DMXEntityFixtureType.cpp b/Engine/Plugins/VirtualProduction/DMX/DMXEngine/Source/DMXRuntime/Private/Library/DMXEntityFixtureType.cpp
--- a/Engine/Plugins/VirtualProduction/DMX/DMXEngine/Source/DMXRuntime/Private/Library/DMXEntityFixtureType.cpp
+++ b/Engine/Plugins/VirtualProduction/DMX/DMXEngine/Source/DMXRuntime/Private/Library/DMXEntityFixtureType.cpp
@@ -241,26 +241,6 @@ void UDMXEntityFixtureType::Serialize(FArchive& Ar)
 	}
 }
 
-void UDMXEntityFixtureType::PostLoad()
-{
-	Super::PostLoad();
-
-#if WITH_EDITOR
-	// Upgrade the GDTF type from the base DMXImport to the actual DMXImportGDTF type which was always used.
-	// Name the new variable 'GDTF' to make the purpose of the variable more clear, namely holding a GDTF asset.
-	PRAGMA_DISABLE_DEPRECATION_WARNINGS
-	if(DMXImport_DEPRECATED)
-	{
-		if (UDMXImportGDTF* OldGDTF = Cast<UDMXImportGDTF>(DMXImport_DEPRECATED))
-		{
-			GDTF = OldGDTF;
-			DMXImport_DEPRECATED = nullptr;
-		}
-	}
-	PRAGMA_ENABLE_DEPRECATION_WARNINGS
-#endif
-}
-
 #if WITH_EDITOR
 bool UDMXEntityFixtureType::Modify(bool bAlwaysMarkDirty)
 {
@@ -320,18 +300,19 @@ void UDMXEntityFixtureType::PostEditUndo()
 }
 #endif // WITH_EDITOR
 
-void UDMXEntityFixtureType::SetGDTF(UDMXImportGDTF* GDTFAsset)
+#if WITH_EDITOR
+void UDMXEntityFixtureType::SetModesFromDMXImport(UDMXImport* DMXImportAsset)
 {
-	if (!IsValid(GDTFAsset))
+	if (!IsValid(DMXImportAsset))
 	{
 		return;
 	}
 
-	GDTF = GDTFAsset;
-	if (UDMXImportGDTFDMXModes* GDTFDMXModes = Cast<UDMXImportGDTFDMXModes>(GDTF->DMXModes))
+	DMXImport = DMXImportAsset;
+	if (UDMXImportGDTFDMXModes* GDTFDMXModes = Cast<UDMXImportGDTFDMXModes>(DMXImport->DMXModes))
 	{
 		// Clear existing modes
-		Modes.Empty(GDTF->DMXModes != nullptr ? GDTFDMXModes->DMXModes.Num() : 0);
+		Modes.Empty(DMXImport->DMXModes != nullptr ? GDTFDMXModes->DMXModes.Num() : 0);
 
 		// Used to map Functions to Attributes
 		const UDMXProtocolSettings* ProtocolSettings = GetDefault<UDMXProtocolSettings>();
@@ -450,17 +431,7 @@ void UDMXEntityFixtureType::SetGDTF(UDMXImportGDTF* GDTFAsset)
 
 	OnFixtureTypeChangedDelegate.Broadcast(this);
 }
-
-void UDMXEntityFixtureType::SetModesFromDMXImport(UDMXImport* DMXImportAsset)
-{
-	UDMXImportGDTF* GDTFAsset = Cast<UDMXImportGDTF>(DMXImportAsset);
-	if (!ensureAlwaysMsgf(GDTFAsset, TEXT("DMXImportAsset %s is not a DMXImportGDTF type, this was never supported. Cannot SetModesFromDMXImport"), *DMXImportAsset->GetName()))
-	{
-		return;
-	}
-
-	SetGDTF(GDTFAsset);
-}
+#endif // WITH_EDITOR
 
 FDMXOnFixtureTypeChangedDelegate& UDMXEntityFixtureType::GetOnFixtureTypeChanged()
 {
@@ -566,21 +537,7 @@ void UDMXEntityFixtureType::SetFixtureMatrixEnabled(int32 ModeIndex, bool bEnabl
 		{
 			Mode.bFixtureMatrixEnabled = bEnableMatrix;
 
-			// Align functions and matrix
-			int32 NextFreeChannel = 1;
-			bool bHandledMatrix = !bEnableMatrix;
-			for (FDMXFixtureFunction& Function : Mode.Functions)
-			{
-				if (!bHandledMatrix && Mode.FixtureMatrixConfig.FirstCellChannel <= NextFreeChannel)
-				{
-					Mode.FixtureMatrixConfig.FirstCellChannel = NextFreeChannel;
-					NextFreeChannel = NextFreeChannel + Mode.FixtureMatrixConfig.GetNumChannels();
-					bHandledMatrix = true;
-				}
-
-				Function.Channel = NextFreeChannel;
-				NextFreeChannel = Function.GetLastChannel() + 1;
-			}
+			AlignFunctionChannels(ModeIndex);
 		}
 	}
 }
@@ -616,6 +573,31 @@ void UDMXEntityFixtureType::UpdateChannelSpan(int32 ModeIndex)
 	}
 }
 
+void UDMXEntityFixtureType::AlignFunctionChannels(int32 ModeIndex)
+{
+	if (ensureMsgf(Modes.IsValidIndex(ModeIndex), TEXT("Invalid Mode Index when aligning the Channels of all Functions in a Mode.")))
+	{
+		FDMXFixtureMode& Mode = Modes[ModeIndex];
+
+		// Align functions and matrix
+		int32 NextFreeChannel = 1;
+		bool bHandledMatrix = !Mode.bFixtureMatrixEnabled;
+		for (FDMXFixtureFunction& Function : Mode.Functions)
+		{
+			if (!bHandledMatrix && 
+				(Mode.FixtureMatrixConfig.FirstCellChannel <= NextFreeChannel || Mode.FixtureMatrixConfig.FirstCellChannel <= Function.Channel))
+			{
+				Mode.FixtureMatrixConfig.FirstCellChannel = NextFreeChannel;
+				NextFreeChannel = NextFreeChannel + Mode.FixtureMatrixConfig.GetNumChannels();
+				bHandledMatrix = true;
+			}
+
+			Function.Channel = NextFreeChannel;
+			NextFreeChannel = Function.GetLastChannel() + 1;
+		}
+	}
+}
+
 int32 UDMXEntityFixtureType::AddFunction(int32 InModeIndex)
 {
 	if (ensureMsgf(Modes.IsValidIndex(InModeIndex), TEXT("Invalid Mode Index when setting Mode Name.")))
diff --git a/Engine/Plugins/VirtualProduction/DMX/DMXEngine/Source/DMXRuntime/Private/MVR/DMXMVRGeneralSceneDescription.cpp b/Engine/Plugins/VirtualProduction/DMX/DMXEngine/Source/DMXRuntime/Private/MVR/DMXMVRGeneralSceneDescription.cpp
--- a/Engine/Plugins/VirtualProduction/DMX/DMXEngine/Source/DMXRuntime/Private/MVR/DMXMVRGeneralSceneDescription.cpp
+++ b/Engine/Plugins/VirtualProduction/DMX/DMXEngine/Source/DMXRuntime/Private/MVR/DMXMVRGeneralSceneDescription.cpp
@@ -158,10 +158,15 @@ void UDMXMVRGeneralSceneDescription::WriteFixturePatchToGeneralSceneDescription(
 	MVRFixtureNode->SetUniverseID(FixturePatch.GetUniverseID());
 	MVRFixtureNode->SetStartingChannel(FixturePatch.GetStartingChannel());
 
+	UDMXEntityFixtureType* FixtureType = FixturePatch.GetFixtureType();
+	const int32 ModeIndex = FixturePatch.GetActiveModeIndex();
 	bool bSetGDTFSpec = false;
-	if (UDMXEntityFixtureType* FixtureType = FixturePatch.GetFixtureType())
+	if (FixtureType &&
+		FixtureType->Modes.IsValidIndex(ModeIndex))
 	{
-		if (UDMXImportGDTF* GDTF = FixtureType->GDTF)
+		MVRFixtureNode->GDTFMode = FixtureType->Modes[ModeIndex].ModeName;
+
+		if (UDMXImportGDTF* GDTF = Cast<UDMXImportGDTF>(FixtureType->DMXImport))
 		{
 			const FString SourceFilename = [GDTF]()
 			{
@@ -172,13 +177,8 @@ void UDMXMVRGeneralSceneDescription::WriteFixturePatchToGeneralSceneDescription(
 				return FString();
 			}();
 			MVRFixtureNode->GDTFSpec = FPaths::GetCleanFilename(SourceFilename);
-			bSetGDTFSpec = true;
-		}
 
-		const int32 ModeIndex = FixturePatch.GetActiveModeIndex();
-		if (ensureAlwaysMsgf(FixtureType->Modes.IsValidIndex(ModeIndex), TEXT("Trying to write Active Mode of to General Scene Description, but the Mode Index is invalid")))
-		{
-			MVRFixtureNode->GDTFMode = FixtureType->Modes[ModeIndex].ModeName;
+			bSetGDTFSpec = true;
 		}
 	}
 
diff --git a/Engine/Plugins/VirtualProduction/DMX/DMXEngine/Source/DMXRuntime/Private/MVR/DMXMVRSceneActor.cpp b/Engine/Plugins/VirtualProduction/DMX/DMXEngine/Source/DMXRuntime/Private/MVR/DMXMVRSceneActor.cpp
--- a/Engine/Plugins/VirtualProduction/DMX/DMXEngine/Source/DMXRuntime/Private/MVR/DMXMVRSceneActor.cpp
+++ b/Engine/Plugins/VirtualProduction/DMX/DMXEngine/Source/DMXRuntime/Private/MVR/DMXMVRSceneActor.cpp
@@ -162,11 +162,11 @@ void ADMXMVRSceneActor::SetDMXLibrary(UDMXLibrary* NewDMXLibrary)
 	for (UDMXEntityFixturePatch* FixturePatch : FixturePatches)
 	{
 		UDMXEntityFixtureType* FixtureType = FixturePatch->GetFixtureType();
-		if (FixtureType && FixtureType->GDTF)
+		if (FixtureType && FixtureType->DMXImport)
 		{
 			const bool bGDTFToDefaultActorClassPairAlreadyCreated = GDTFToDefaultActorClasses.ContainsByPredicate([FixtureType](const FDMXMVRSceneGDTFToActorClassPair& GDTFToActorClassPair)
 				{
-					return GDTFToActorClassPair.GDTF == FixtureType->GDTF;
+					return GDTFToActorClassPair.GDTF == FixtureType->DMXImport;
 				});
 			if (bGDTFToDefaultActorClassPairAlreadyCreated)
 			{
@@ -175,7 +175,7 @@ void ADMXMVRSceneActor::SetDMXLibrary(UDMXLibrary* NewDMXLibrary)
 
 			UClass* ActorClass = MVRFixtureActorLibrary->FindMostAppropriateActorClassForPatch(FixturePatch);
 			FDMXMVRSceneGDTFToActorClassPair GDTFToActorClassPair;
-			GDTFToActorClassPair.GDTF = FixtureType->GDTF;
+			GDTFToActorClassPair.GDTF = FixtureType->DMXImport;
 			GDTFToActorClassPair.ActorClass = ActorClass;
 			GDTFToDefaultActorClasses.Add(GDTFToActorClassPair);
 		}
@@ -301,7 +301,7 @@ void ADMXMVRSceneActor::HandleDefaultActorClassForGDTFChanged()
 		UDMXEntityFixturePatch* FixturePatch = GetFixturePatch(Actor);
 		if (FixturePatch && 
 			FixturePatch->GetFixtureType() && 
-			FixturePatch->GetFixtureType()->GDTF == GDTFToDefaultActorClasses[IndexOfChangedElement].GDTF)
+			FixturePatch->GetFixtureType()->DMXImport == GDTFToDefaultActorClasses[IndexOfChangedElement].GDTF)
 		{
 			ReplaceMVRActor(Actor, Class);
 		}
diff --git a/Engine/Plugins/VirtualProduction/DMX/DMXEngine/Source/DMXRuntime/Public/Library/DMXEntityFixtureType.h b/Engine/Plugins/VirtualProduction/DMX/DMXEngine/Source/DMXRuntime/Public/Library/DMXEntityFixtureType.h
--- a/Engine/Plugins/VirtualProduction/DMX/DMXEngine/Source/DMXRuntime/Public/Library/DMXEntityFixtureType.h
+++ b/Engine/Plugins/VirtualProduction/DMX/DMXEngine/Source/DMXRuntime/Public/Library/DMXEntityFixtureType.h
@@ -303,7 +303,6 @@ public:
 
 	//~ Begin UObject interface
 	virtual void Serialize(FArchive& Ar) override;
-	virtual void PostLoad() override;
 #if WITH_EDITOR
 	virtual bool Modify(bool bAlwaysMarkDirty = true) override;
 	virtual void PostEditChangeProperty(FPropertyChangedEvent& PropertyChangedEvent) override;
@@ -312,29 +311,18 @@ public:
 #endif
 	//~ End UObject interface
 
-
 public:
-	/** Sets the GDTF for this Fixture Type. Note, this clears the fixture type so that the Fixture Type reflects the GDTF. Merging is not supported. */
+#if WITH_EDITOR
 	UFUNCTION(BlueprintCallable, Category = "Fixture Settings")
-	void SetGDTF(UDMXImportGDTF* DMXImportAsset);
-
-	UE_DEPRECATED(5.1, "Deprecated in favor of more clearly named member function SetGDTFAsset. Only supports the functional UDMXImportGDTF now.")
-	UFUNCTION(BlueprintCallable, Category = "Fixture Settings", Meta = (DeprecatedFunction, DeprecationMessage = "Deprecated 5.1. Please use the more clearly named SetGDTFAsset instead."))
 	void SetModesFromDMXImport(UDMXImport* DMXImportAsset);
+#endif // WITH_EDITOR
 
 	/** Returns a delegate that is and should be broadcast whenever a Fixture Type changed */
 	static FDMXOnFixtureTypeChangedDelegate& GetOnFixtureTypeChanged();
 
-#if WITH_EDITORONLY_DATA
-	/** DEPRECATED 5.0 */
-	UE_DEPRECATED(5.1, "Deprecated in favor of the GDTF property, which is of the GDTF type instead of just a base class.")
-	UPROPERTY(Meta = (DeprecatedProperty, DeprecationMessage = "Deprecated in favor of the GDTF property, which is of the GDTF type instead of just a base class."))
-	UDMXImport* DMXImport_DEPRECATED;
-#endif 
-
 	/** The GDTF file from which the Fixture Type was setup */
 	UPROPERTY(EditAnywhere, BlueprintReadOnly, Category = "Fixture Settings")
-	UDMXImportGDTF* GDTF;
+	UDMXImport* DMXImport;
 
 	/** The Category of the Fixture, useful for Filtering */
 	UPROPERTY(EditAnywhere, BlueprintReadOnly, Category = "Fixture Settings", meta = (DisplayName = "DMX Category"))
@@ -402,6 +390,8 @@ public:
 	 */
 	void UpdateChannelSpan(int32 ModeIndex);
 
+	/** Aligns alls channels of the functions in the Mode to be consecutive */
+	void AlignFunctionChannels(int32 InModeIndex);
 
 	// Fixture Function related
 public:
