DMX - Allow Attributes to be arbitrarily named

diff --git a/Engine/Plugins/VirtualProduction/DMX/DMXEngine/Source/DMXBlueprintGraph/Private/K2Node_CastPatchToType.cpp b/Engine/Plugins/VirtualProduction/DMX/DMXEngine/Source/DMXBlueprintGraph/Private/K2Node_CastPatchToType.cpp
--- a/Engine/Plugins/VirtualProduction/DMX/DMXEngine/Source/DMXBlueprintGraph/Private/K2Node_CastPatchToType.cpp
+++ b/Engine/Plugins/VirtualProduction/DMX/DMXEngine/Source/DMXBlueprintGraph/Private/K2Node_CastPatchToType.cpp
@@ -564,7 +564,7 @@ void UDEPRECATED_K2Node_CastPatchToType::ExposeAttributes()
 				}
 
 				FString EnumString = StaticEnum<EDMXFixtureSignalFormat>()->GetDisplayNameTextByIndex((int64)Function.DataType).ToString();
-				FString PinFunctionName = *FString::Printf(TEXT("%s_%s"), *Function.Attribute.GetName().ToString(), *EnumString);
+				FString PinFunctionName = *FString::Printf(TEXT("%s_%s"), *Function.Attribute.Name.ToString(), *EnumString);
 
 				{
 					FEdGraphPinType PinType;
diff --git a/Engine/Plugins/VirtualProduction/DMX/DMXEngine/Source/DMXBlueprintGraph/Private/K2Node_GetDMXAttributeValues.cpp b/Engine/Plugins/VirtualProduction/DMX/DMXEngine/Source/DMXBlueprintGraph/Private/K2Node_GetDMXAttributeValues.cpp
--- a/Engine/Plugins/VirtualProduction/DMX/DMXEngine/Source/DMXBlueprintGraph/Private/K2Node_GetDMXAttributeValues.cpp
+++ b/Engine/Plugins/VirtualProduction/DMX/DMXEngine/Source/DMXBlueprintGraph/Private/K2Node_GetDMXAttributeValues.cpp
@@ -526,7 +526,7 @@ FName UK2Node_GetDMXAttributeValues::GetPinName(const FDMXFixtureFunction& Funct
 	FString EnumString;
 	EnumString = StaticEnum<EDMXFixtureSignalFormat>()->GetDisplayNameTextByIndex((int64)Function.DataType).ToString();
 
-	return *FString::Printf(TEXT("%s_%s"), *Function.Attribute.GetName().ToString(), *EnumString);
+	return *FString::Printf(TEXT("%s_%s"), *Function.Attribute.Name.ToString(), *EnumString);
 }
 
 const FDMXFixtureMode* UK2Node_GetDMXAttributeValues::GetActiveFixtureMode() const
diff --git a/Engine/Plugins/VirtualProduction/DMX/DMXEngine/Source/DMXBlueprintGraph/Private/Widgets/SDynamicNameListGraphPin.h b/Engine/Plugins/VirtualProduction/DMX/DMXEngine/Source/DMXBlueprintGraph/Private/Widgets/SDynamicNameListGraphPin.h
--- a/Engine/Plugins/VirtualProduction/DMX/DMXEngine/Source/DMXBlueprintGraph/Private/Widgets/SDynamicNameListGraphPin.h
+++ b/Engine/Plugins/VirtualProduction/DMX/DMXEngine/Source/DMXBlueprintGraph/Private/Widgets/SDynamicNameListGraphPin.h
@@ -30,11 +30,7 @@ public:
 		return SNew(SNameListPicker)
 			.HasMultipleValues(false)
 			.Value(this, &SDynamicNameListGraphPin::GetValue)
-			.OnValueChanged(this, &SDynamicNameListGraphPin::SetValue)
-			.UpdateOptionsDelegate(&TStructType::OnValuesChanged)
 			.OptionsSource(MakeAttributeLambda(&TStructType::GetPossibleValues))
-			.IsValid(this, &SDynamicNameListGraphPin::IsValueValid)
-			.bCanBeNone(TStructType::bCanBeNone)
 			.bDisplayWarningIcon(true)
 			.Visibility(this, &SGraphPin::GetDefaultValueVisibility);
 	}
@@ -50,7 +46,7 @@ private:
 			TStructType::StaticStruct()->ImportText(*GraphPinObj->GetDefaultAsString(), &NameItem, nullptr, EPropertyPortFlags::PPF_None, GLog, TStructType::StaticStruct()->GetName());
 		}
 
-		return NameItem.GetName();
+		return NameItem.Name;
 	}
 
 	void SetValue(FName NewValue)
@@ -60,9 +56,4 @@ private:
 		TStructType::StaticStruct()->ExportText(ValueString, &NewNameItem, nullptr, nullptr, EPropertyPortFlags::PPF_None, nullptr);
 		GraphPinObj->GetSchema()->TrySetDefaultValue(*GraphPinObj, ValueString);
 	}
-
-	bool IsValueValid() const
-	{
-		return TStructType::IsValid(GetValue());
-	}
 };
diff --git a/Engine/Plugins/VirtualProduction/DMX/DMXEngine/Source/DMXEditor/Private/Customizations/DMXNameListCustomization.h b/Engine/Plugins/VirtualProduction/DMX/DMXEngine/Source/DMXEditor/Private/Customizations/DMXNameListCustomization.h
--- a/Engine/Plugins/VirtualProduction/DMX/DMXEngine/Source/DMXEditor/Private/Customizations/DMXNameListCustomization.h
+++ b/Engine/Plugins/VirtualProduction/DMX/DMXEngine/Source/DMXEditor/Private/Customizations/DMXNameListCustomization.h
@@ -52,11 +52,8 @@ public:
 				SNew(SNameListPicker)
 				.Font(CustomizationUtils.GetRegularFont())
 				.HasMultipleValues(this, &TNameListType::HasMultipleValues)
-				.OptionsSource(MakeAttributeLambda(&TStructType::GetPossibleValues))
-				.UpdateOptionsDelegate(&TStructType::OnValuesChanged)
-				.IsValid(this, &TNameListType::HideWarningIcon)
+				.OptionsSource(FDMXAttributeName::GetPossibleValues())
 				.Value(this, &TNameListType::GetValue)
-				.bCanBeNone(TStructType::bCanBeNone)
 				.bDisplayWarningIcon(true)
 				.OnValueChanged(this, &TNameListType::SetValue)
 			]
@@ -77,7 +74,7 @@ private:
 			if (RawPtr != nullptr)
 			{
 				// The types we use with this customization must have a cast constructor to FName
-				return reinterpret_cast<const TStructType*>(RawPtr)->GetName();
+				return reinterpret_cast<const TStructType*>(RawPtr)->Name;
 			}
 		}
 
@@ -139,22 +136,6 @@ private:
 		return false;
 	}
 
-	bool HideWarningIcon() const
-	{
-		if (HasMultipleValues())
-		{
-			return true;
-		}
-
-		const FName CurrentValue = GetValue();
-		if (CurrentValue.IsEqual(FDMXNameListItem::None))
-		{
-			return true;
-		}
-
-		return TStructType::IsValid(GetValue());
-	}
-
 private:
 	TSharedPtr<IPropertyHandle> StructPropertyHandle;
 };
diff --git a/Engine/Plugins/VirtualProduction/DMX/DMXEngine/Source/DMXEditor/Private/Sequencer/DMXAsyncDMXRecorder.cpp b/Engine/Plugins/VirtualProduction/DMX/DMXEngine/Source/DMXEditor/Private/Sequencer/DMXAsyncDMXRecorder.cpp
--- a/Engine/Plugins/VirtualProduction/DMX/DMXEngine/Source/DMXEditor/Private/Sequencer/DMXAsyncDMXRecorder.cpp
+++ b/Engine/Plugins/VirtualProduction/DMX/DMXEngine/Source/DMXEditor/Private/Sequencer/DMXAsyncDMXRecorder.cpp
@@ -40,7 +40,7 @@ namespace
 				}
 
 				const FDMXFixtureCellAttribute* CellAttributePtr = MatrixProperties.CellAttributes.FindByPredicate([&MutableChannelData](const FDMXFixtureCellAttribute& CellAttribute) {
-					return CellAttribute.Attribute.GetName() == MutableChannelData.GetAttributeName();
+					return CellAttribute.Attribute.Name == MutableChannelData.GetAttributeName();
 					});
 
 				if (CellAttributePtr)
@@ -76,7 +76,7 @@ namespace
 					int32 AbsoluteStartingChannel = FixturePatch->GetStartingChannel();
 
 					const FDMXFixtureFunction* FixtureFunctionPtr = ModePtr->Functions.FindByPredicate([&MutableChannelData](const FDMXFixtureFunction& Function) {
-						return Function.Attribute.GetName() == MutableChannelData.GetAttributeName();
+						return Function.Attribute.Name == MutableChannelData.GetAttributeName();
 						});
 
 					if (FixtureFunctionPtr)
diff --git a/Engine/Plugins/VirtualProduction/DMX/DMXEngine/Source/DMXEditor/Private/Widgets/FixtureType/DMXFixtureTypeFunctionsEditorFunctionItem.cpp b/Engine/Plugins/VirtualProduction/DMX/DMXEngine/Source/DMXEditor/Private/Widgets/FixtureType/DMXFixtureTypeFunctionsEditorFunctionItem.cpp
--- a/Engine/Plugins/VirtualProduction/DMX/DMXEngine/Source/DMXEditor/Private/Widgets/FixtureType/DMXFixtureTypeFunctionsEditorFunctionItem.cpp
+++ b/Engine/Plugins/VirtualProduction/DMX/DMXEngine/Source/DMXEditor/Private/Widgets/FixtureType/DMXFixtureTypeFunctionsEditorFunctionItem.cpp
@@ -107,7 +107,7 @@ FDMXAttributeName FDMXFixtureTypeFunctionsEditorFunctionItem::GetAttributeName()
 		return Function.Attribute;
 	}
 
-	return FDMXAttributeName::None;
+	return FDMXAttributeName(NAME_None);
 }
 
 void FDMXFixtureTypeFunctionsEditorFunctionItem::SetAttributeName(const FDMXAttributeName& AttributeName) const
diff --git a/Engine/Plugins/VirtualProduction/DMX/DMXEngine/Source/DMXEditor/Private/Widgets/FixtureType/SDMXFixtureTypeFunctionsEditor.cpp b/Engine/Plugins/VirtualProduction/DMX/DMXEngine/Source/DMXEditor/Private/Widgets/FixtureType/SDMXFixtureTypeFunctionsEditor.cpp
--- a/Engine/Plugins/VirtualProduction/DMX/DMXEngine/Source/DMXEditor/Private/Widgets/FixtureType/SDMXFixtureTypeFunctionsEditor.cpp
+++ b/Engine/Plugins/VirtualProduction/DMX/DMXEngine/Source/DMXEditor/Private/Widgets/FixtureType/SDMXFixtureTypeFunctionsEditor.cpp
@@ -12,6 +12,8 @@
 #include "Widgets/FixtureType/SDMXFixtureTypeFunctionsEditorFunctionRow.h"
 #include "Widgets/FixtureType/SDMXFixtureTypeFunctionsEditorMatrixRow.h"
 
+#include "Algo/Copy.h"
+#include "Algo/Find.h"
 #include "Styling/AppStyle.h"
 #include "ScopedTransaction.h"
 #include "Framework/Commands/GenericCommands.h"
@@ -133,7 +135,7 @@ void SDMXFixtureTypeFunctionsEditor::RebuildList()
 				TArray<TSharedPtr<FDMXFixtureTypeFunctionsEditorItemBase>> SelectedItems;
 				for (const int32 SelectedFunctionIndex : SelectedFunctionIndices)
 				{
-					const TSharedPtr<FDMXFixtureTypeFunctionsEditorItemBase>* SelectedItemPtr = ListSource.FindByPredicate([SelectedFunctionIndex](const TSharedPtr<FDMXFixtureTypeFunctionsEditorItemBase>& Item)
+					const TSharedPtr<FDMXFixtureTypeFunctionsEditorItemBase>* SelectedItemPtr = Algo::FindByPredicate(ListSource, [SelectedFunctionIndex](const TSharedPtr<FDMXFixtureTypeFunctionsEditorItemBase>& Item)
 						{
 							if (Item->GetType() == FDMXFixtureTypeFunctionsEditorItemBase::EItemType::Function)
 							{
@@ -152,7 +154,7 @@ void SDMXFixtureTypeFunctionsEditor::RebuildList()
 
 				if (FixtureTypeSharedData->IsFixtureMatrixSelected())
 				{
-					const TSharedPtr<FDMXFixtureTypeFunctionsEditorItemBase>* SelectedMatrixPtr = ListSource.FindByPredicate([](const TSharedPtr<FDMXFixtureTypeFunctionsEditorItemBase>& Item)
+					const TSharedPtr<FDMXFixtureTypeFunctionsEditorItemBase>* SelectedMatrixPtr = Algo::FindByPredicate(ListSource, [](const TSharedPtr<FDMXFixtureTypeFunctionsEditorItemBase>& Item)
 						{
 							return Item->GetType() == FDMXFixtureTypeFunctionsEditorItemBase::EItemType::Matrix;
 						});
@@ -331,10 +333,10 @@ void SDMXFixtureTypeFunctionsEditor::UpdateItemsStatus()
 	}
 	for (const TTuple<FDMXAttributeName, TSet<TSharedPtr<FDMXFixtureTypeFunctionsEditorItemBase>>>& OccupiedAttribueToItemsPair : OccupiedAttributeToItemsMap)
 	{
-		if (OccupiedAttribueToItemsPair.Value.Num() > 1 && FDMXAttributeName::IsValid(OccupiedAttribueToItemsPair.Key.GetName()))
+		if (OccupiedAttribueToItemsPair.Value.Num() > 1 && OccupiedAttribueToItemsPair.Key.Name != NAME_None)
 		{
 			// Create the warning status
-			FText NewWarningStatus = FText::Format(LOCTEXT("ConflictingAttributesWarning", "Ambiguous Attribute '{0}' used in:"), FText::FromName(OccupiedAttribueToItemsPair.Key.GetName()));
+			FText NewWarningStatus = FText::Format(LOCTEXT("ConflictingAttributesWarning", "Ambiguous Attribute '{0}' used in:"), FText::FromName(OccupiedAttribueToItemsPair.Key.Name));
 
 			for (const TSharedPtr<FDMXFixtureTypeFunctionsEditorItemBase>& ConflictingItem : OccupiedAttribueToItemsPair.Value)
 			{
@@ -475,8 +477,33 @@ void SDMXFixtureTypeFunctionsEditor::OnListSelectionChanged(TSharedPtr<FDMXFixtu
 {
 	if (SelectInfo != ESelectInfo::Direct)
 	{
-		const bool bMatrixSelected = IsMatrixSelected();
-		FixtureTypeSharedData->SetFunctionAndMatrixSelection(GetSelectionAsFunctionIndices(), bMatrixSelected);
+		const TArray<int32> SelectedFunctionIndicies = GetSelectionAsFunctionIndices();
+
+		if (SelectedFunctionIndicies.IsEmpty())
+		{
+			// Never clear selection
+			TArray<TSharedPtr<FDMXFixtureTypeFunctionsEditorItemBase>> ItemsToSelect;
+			Algo::CopyIf(ListSource, ItemsToSelect, [&SelectedFunctionIndicies](const TSharedPtr<FDMXFixtureTypeFunctionsEditorItemBase>& Item)
+				{
+					if (Item->GetType() == FDMXFixtureTypeFunctionsEditorItemBase::EItemType::Function)
+					{
+						const TSharedPtr<FDMXFixtureTypeFunctionsEditorFunctionItem> FunctionItem = StaticCastSharedPtr<FDMXFixtureTypeFunctionsEditorFunctionItem>(Item);
+						return SelectedFunctionIndicies.Contains(FunctionItem->GetFunctionIndex());
+					}
+
+					return false;
+				});
+
+			if (ItemsToSelect.Num() > 0)
+			{
+				ListView->SetItemSelection(ItemsToSelect, true);
+			}
+		}
+		else
+		{
+			const bool bMatrixSelected = IsMatrixSelected();
+			FixtureTypeSharedData->SetFunctionAndMatrixSelection(SelectedFunctionIndicies, bMatrixSelected);
+		}
 	}
 	else
 	{
@@ -490,7 +517,7 @@ void SDMXFixtureTypeFunctionsEditor::OnSharedDataFunctionOrMatrixSelectionChange
 	TArray<TSharedPtr<FDMXFixtureTypeFunctionsEditorItemBase>> NewSelection;
 	for (int32 SelectedFunctionIndex : SelectedFunctionIndices)
 	{
-		const TSharedPtr<FDMXFixtureTypeFunctionsEditorItemBase>* SelectedItemPtr = ListSource.FindByPredicate([SelectedFunctionIndex](const TSharedPtr<FDMXFixtureTypeFunctionsEditorItemBase>& Item)
+		const TSharedPtr<FDMXFixtureTypeFunctionsEditorItemBase>* SelectedItemPtr = Algo::FindByPredicate(ListSource, [SelectedFunctionIndex](const TSharedPtr<FDMXFixtureTypeFunctionsEditorItemBase>& Item)
 			{
 				if (Item->GetType() == FDMXFixtureTypeFunctionsEditorItemBase::EItemType::Function)
 				{
@@ -508,7 +535,7 @@ void SDMXFixtureTypeFunctionsEditor::OnSharedDataFunctionOrMatrixSelectionChange
 
 	if (FixtureTypeSharedData->IsFixtureMatrixSelected())
 	{
-		const TSharedPtr<FDMXFixtureTypeFunctionsEditorItemBase>* SelectedItemPtr = ListSource.FindByPredicate([](const TSharedPtr<FDMXFixtureTypeFunctionsEditorItemBase>& Item)
+		const TSharedPtr<FDMXFixtureTypeFunctionsEditorItemBase>* SelectedItemPtr = Algo::FindByPredicate(ListSource, [](const TSharedPtr<FDMXFixtureTypeFunctionsEditorItemBase>& Item)
 			{
 				return Item->GetType() == FDMXFixtureTypeFunctionsEditorItemBase::EItemType::Matrix;
 			});
@@ -794,7 +821,7 @@ void SDMXFixtureTypeFunctionsEditor::OnRenameItem()
 	if (SelectedFunctionItems.Num() == 1)
 	{
 		const TSharedPtr<FDMXFixtureTypeFunctionsEditorItemBase>& SelectedItem = SelectedFunctionItems[0];
-		const TSharedPtr<SDMXFixtureTypeFunctionsEditorFunctionRow>* SelectedRowPtr = FunctionRows.FindByPredicate([SelectedItem](const TSharedPtr<SDMXFixtureTypeFunctionsEditorFunctionRow>& FunctionRow)
+		const TSharedPtr<SDMXFixtureTypeFunctionsEditorFunctionRow>* SelectedRowPtr = Algo::FindByPredicate(FunctionRows, [SelectedItem](const TSharedPtr<SDMXFixtureTypeFunctionsEditorFunctionRow>& FunctionRow)
 			{
 				return FunctionRow->GetFunctionItem() == SelectedItem;
 			});
diff --git a/Engine/Plugins/VirtualProduction/DMX/DMXEngine/Source/DMXEditor/Private/Widgets/FixtureType/SDMXFixtureTypeFunctionsEditorFunctionRow.cpp b/Engine/Plugins/VirtualProduction/DMX/DMXEngine/Source/DMXEditor/Private/Widgets/FixtureType/SDMXFixtureTypeFunctionsEditorFunctionRow.cpp
--- a/Engine/Plugins/VirtualProduction/DMX/DMXEngine/Source/DMXEditor/Private/Widgets/FixtureType/SDMXFixtureTypeFunctionsEditorFunctionRow.cpp
+++ b/Engine/Plugins/VirtualProduction/DMX/DMXEngine/Source/DMXEditor/Private/Widgets/FixtureType/SDMXFixtureTypeFunctionsEditorFunctionRow.cpp
@@ -3,6 +3,7 @@
 #include "SDMXFixtureTypeFunctionsEditorFunctionRow.h"
 
 #include "DMXAttribute.h"
+#include "DMXEditor.h"
 #include "DMXEditorStyle.h"
 #include "DMXFixtureTypeSharedData.h"
 #include "DragDrop/DMXFixtureFunctionDragDropOp.h"
@@ -172,21 +173,9 @@ TSharedRef<SWidget> SDMXFixtureTypeFunctionsEditorFunctionRow::GenerateWidgetFor
 			[
 				SNew(SNameListPicker)
 				.OptionsSource(MakeAttributeLambda(&FDMXAttributeName::GetPossibleValues))
-				.UpdateOptionsDelegate(&FDMXAttributeName::OnValuesChanged)
-				.IsValid_Lambda([this]()
-					{	
-						const FName CurrentValue = GetAttributeName();
-						if (CurrentValue.IsEqual(FDMXNameListItem::None))
-						{
-							return true;
-						}
-
-						return FunctionItem->HasValidAttribute();
-					})
 				.Value(this, &SDMXFixtureTypeFunctionsEditorFunctionRow::GetAttributeName)
-				.bCanBeNone(FDMXAttributeName::bCanBeNone)
 				.bDisplayWarningIcon(true)
-				.OnValueChanged(this, &SDMXFixtureTypeFunctionsEditorFunctionRow::SetAttributeName)
+				.OnValueChanged(this, &SDMXFixtureTypeFunctionsEditorFunctionRow::OnUserChangedAttributeName)
 			];
 	}
 
@@ -337,15 +326,25 @@ void SDMXFixtureTypeFunctionsEditorFunctionRow::OnFunctionNameCommitted(const FT
 
 FName SDMXFixtureTypeFunctionsEditorFunctionRow::GetAttributeName() const
 {
-	return FunctionItem->GetAttributeName().GetName();
+	return FunctionItem->GetAttributeName().Name;
 }
 
-void SDMXFixtureTypeFunctionsEditorFunctionRow::SetAttributeName(FName NewValue)
+void SDMXFixtureTypeFunctionsEditorFunctionRow::OnUserChangedAttributeName(FName NewValue)
 {
 	FDMXAttributeName NewAttributeName;
 	NewAttributeName.SetFromName(NewValue);
 
 	FunctionItem->SetAttributeName(NewAttributeName);
+
+	// Since edit mode can be entered without selecting the item, select it when a value is set
+	if (TSharedPtr<FDMXEditor> DMXEditor = FunctionItem->GetDMXEditor())
+	{
+		if (TSharedPtr<FDMXFixtureTypeSharedData> SharedData = DMXEditor->GetFixtureTypeSharedData())
+		{
+			constexpr bool bMatrixSelected = false;
+			SharedData->SetFunctionAndMatrixSelection({ FunctionItem->GetFunctionIndex() }, bMatrixSelected);
+		}
+	}
 }
 
 #undef LOCTEXT_NAMESPACE
diff --git a/Engine/Plugins/VirtualProduction/DMX/DMXEngine/Source/DMXEditor/Private/Widgets/FixtureType/SDMXFixtureTypeFunctionsEditorFunctionRow.h b/Engine/Plugins/VirtualProduction/DMX/DMXEngine/Source/DMXEditor/Private/Widgets/FixtureType/SDMXFixtureTypeFunctionsEditorFunctionRow.h
--- a/Engine/Plugins/VirtualProduction/DMX/DMXEngine/Source/DMXEditor/Private/Widgets/FixtureType/SDMXFixtureTypeFunctionsEditorFunctionRow.h
+++ b/Engine/Plugins/VirtualProduction/DMX/DMXEngine/Source/DMXEditor/Private/Widgets/FixtureType/SDMXFixtureTypeFunctionsEditorFunctionRow.h
@@ -78,8 +78,8 @@ private:
 	/** Gets the Attribute Name of the Function */
 	FName GetAttributeName() const;
 
-	/** Sets the Attribute Name for the Function */
-	void SetAttributeName(FName NewValue);
+	/** Called when the Attribute Name changed */
+	void OnUserChangedAttributeName(FName NewValue);
 
 	/** True while this is being dragged in a drag drop op */
 	bool bIsBeingDragged = false;
diff --git a/Engine/Plugins/VirtualProduction/DMX/DMXEngine/Source/DMXEditor/Private/Widgets/FixtureType/SDMXFixtureTypeMatrixFunctionsEditor.cpp b/Engine/Plugins/VirtualProduction/DMX/DMXEngine/Source/DMXEditor/Private/Widgets/FixtureType/SDMXFixtureTypeMatrixFunctionsEditor.cpp
--- a/Engine/Plugins/VirtualProduction/DMX/DMXEngine/Source/DMXEditor/Private/Widgets/FixtureType/SDMXFixtureTypeMatrixFunctionsEditor.cpp
+++ b/Engine/Plugins/VirtualProduction/DMX/DMXEngine/Source/DMXEditor/Private/Widgets/FixtureType/SDMXFixtureTypeMatrixFunctionsEditor.cpp
@@ -147,20 +147,20 @@ void SDMXFixtureTypeMatrixFunctionsEditor::RefreshList()
 				MatrixFunctionItem->ChannelNumberText = FText::AsNumber(ChannelNumber);
 				CellAttributesListSource.Add(MatrixFunctionItem);
 
-				const FDMXAttributeName& Attribute = CellAttribute.Attribute.GetName();
+				const FDMXAttributeName& Attribute = CellAttribute.Attribute.Name;
 				AttributeToConflictingItemsMap.FindOrAdd(Attribute).Add(MatrixFunctionItem);
 
-				ChannelOffset = CellAttribute.GetNumChannels() - 1;
+				ChannelOffset += CellAttribute.GetNumChannels() - 1;
 			}
 
 			// Create item status
 			for (const TTuple<FDMXAttributeName, TSet<TSharedPtr<FDMXFixtureTypeMatrixFunctionsEditorItem>>>& AttributeToConflictingItemsPair : AttributeToConflictingItemsMap)
 			{
-				if (AttributeToConflictingItemsPair.Value.Num() > 1 && FDMXAttributeName::IsValid(AttributeToConflictingItemsPair.Key.GetName()))
+				if (AttributeToConflictingItemsPair.Value.Num() > 1 && AttributeToConflictingItemsPair.Key.Name != NAME_None)
 				{
 					for (const TSharedPtr<FDMXFixtureTypeMatrixFunctionsEditorItem>& ConflictingItem : AttributeToConflictingItemsPair.Value)
 					{
-						ConflictingItem->WarningStatus = FText::Format(LOCTEXT("AmbiguousCellAttribute", "Attribute {0} used more than once."), FText::FromName(AttributeToConflictingItemsPair.Key.GetName()));
+						ConflictingItem->WarningStatus = FText::Format(LOCTEXT("AmbiguousCellAttribute", "Attribute {0} used more than once."), FText::FromName(AttributeToConflictingItemsPair.Key.Name));
 					}
 				}
 			}
diff --git a/Engine/Plugins/VirtualProduction/DMX/DMXEngine/Source/DMXEditor/Private/Widgets/FixtureType/SDMXFixtureTypeMatrixFunctionsEditorMatrixRow.cpp b/Engine/Plugins/VirtualProduction/DMX/DMXEngine/Source/DMXEditor/Private/Widgets/FixtureType/SDMXFixtureTypeMatrixFunctionsEditorMatrixRow.cpp
--- a/Engine/Plugins/VirtualProduction/DMX/DMXEngine/Source/DMXEditor/Private/Widgets/FixtureType/SDMXFixtureTypeMatrixFunctionsEditorMatrixRow.cpp
+++ b/Engine/Plugins/VirtualProduction/DMX/DMXEngine/Source/DMXEditor/Private/Widgets/FixtureType/SDMXFixtureTypeMatrixFunctionsEditorMatrixRow.cpp
@@ -84,10 +84,7 @@ TSharedRef<SWidget> SDMXFixtureTypeMatrixFunctionsEditorMatrixRow::GenerateWidge
 		return
 			SNew(SNameListPicker)
 			.OptionsSource(MakeAttributeLambda(&FDMXAttributeName::GetPossibleValues))
-			.UpdateOptionsDelegate(&FDMXAttributeName::OnValuesChanged)
-			.IsValid(this, &SDMXFixtureTypeMatrixFunctionsEditorMatrixRow::HasInvalidCellAttributeName)
 			.Value(this, &SDMXFixtureTypeMatrixFunctionsEditorMatrixRow::GetCellAttributeName)
-			.bCanBeNone(FDMXAttributeName::bCanBeNone)
 			.bDisplayWarningIcon(true)
 			.OnValueChanged(this, &SDMXFixtureTypeMatrixFunctionsEditorMatrixRow::SetCellAttributeName);
 	}
@@ -120,20 +117,9 @@ FReply SDMXFixtureTypeMatrixFunctionsEditorMatrixRow::OnDeleteCellAttributeClick
 	return FReply::Handled();
 }
 
-bool SDMXFixtureTypeMatrixFunctionsEditorMatrixRow::HasInvalidCellAttributeName() const
-{
-	const FName CurrentValue = GetCellAttributeName();
-	if (CurrentValue.IsEqual(FDMXNameListItem::None))
-	{
-		return true;
-	}
-
-	return FDMXAttributeName::IsValid(CurrentValue);
-}
-
 FName SDMXFixtureTypeMatrixFunctionsEditorMatrixRow::GetCellAttributeName() const
 {
-	return CellAttributeItem->GetCellAttributeName().GetName();
+	return CellAttributeItem->GetCellAttributeName().Name;
 }
 
 void SDMXFixtureTypeMatrixFunctionsEditorMatrixRow::SetCellAttributeName(FName NewValue)
diff --git a/Engine/Plugins/VirtualProduction/DMX/DMXEngine/Source/DMXEditor/Private/Widgets/FixtureType/SDMXFixtureTypeMatrixFunctionsEditorMatrixRow.h b/Engine/Plugins/VirtualProduction/DMX/DMXEngine/Source/DMXEditor/Private/Widgets/FixtureType/SDMXFixtureTypeMatrixFunctionsEditorMatrixRow.h
--- a/Engine/Plugins/VirtualProduction/DMX/DMXEngine/Source/DMXEditor/Private/Widgets/FixtureType/SDMXFixtureTypeMatrixFunctionsEditorMatrixRow.h
+++ b/Engine/Plugins/VirtualProduction/DMX/DMXEngine/Source/DMXEditor/Private/Widgets/FixtureType/SDMXFixtureTypeMatrixFunctionsEditorMatrixRow.h
@@ -38,9 +38,6 @@ private:
 	/** Called when the delete Cell Attribute Button was clicked */
 	FReply OnDeleteCellAttributeClicked();
 
-	/** Returns true if the Attribute Name exists and is not None */
-	bool HasInvalidCellAttributeName() const;
-
 	/** Returns the Name of the Cell Attribute */
 	FName GetCellAttributeName() const;
 
diff --git a/Engine/Plugins/VirtualProduction/DMX/DMXEngine/Source/DMXEditor/Private/Widgets/SNameListPicker.cpp b/Engine/Plugins/VirtualProduction/DMX/DMXEngine/Source/DMXEditor/Private/Widgets/SNameListPicker.cpp
--- a/Engine/Plugins/VirtualProduction/DMX/DMXEngine/Source/DMXEditor/Private/Widgets/SNameListPicker.cpp
+++ b/Engine/Plugins/VirtualProduction/DMX/DMXEngine/Source/DMXEditor/Private/Widgets/SNameListPicker.cpp
@@ -4,11 +4,13 @@
 
 #include "DMXEditorLog.h"
 #include "DMXNameListItem.h"
+#include "DMXProtocolSettings.h"
 
 #include "Styling/AppStyle.h"
 #include "SListViewSelectorDropdownMenu.h"
 #include "Widgets/Images/SImage.h"
 #include "Widgets/Input/SComboBox.h"
+#include "Widgets/Input/SEditableTextBox.h"
 #include "Widgets/Input/SSearchBox.h"
 #include "Widgets/Views/SListView.h"
 #include "Widgets/Views/STableRow.h"
@@ -20,22 +22,24 @@ const FText SNameListPicker::NoneLabel = LOCTEXT("NoneLabel", "<Select a Value>"
 void SNameListPicker::Construct(const FArguments& InArgs)
 {
 	ValueAttribute = InArgs._Value;
-	OnValueChangedDelegate = InArgs._OnValueChanged;
 	HasMultipleValuesAttribute = InArgs._HasMultipleValues;
-
-	bCanBeNone = InArgs._bCanBeNone;
 	bDisplayWarningIcon = InArgs._bDisplayWarningIcon;
 	OptionsSourceAttr = InArgs._OptionsSource;
-	UpdateOptionsSource();
-	IsValidAttr = InArgs._IsValid;
+	MaxVisibleItems = InArgs._MaxVisibleItems;
+	OnValueChangedDelegate = InArgs._OnValueChanged;
 
-	UpdateOptionsDelegate = InArgs._UpdateOptionsDelegate;
-	if (UpdateOptionsDelegate)
+	UpdateOptionsSource();
+	UDMXProtocolSettings* ProtocolSettings = GetMutableDefault<UDMXProtocolSettings>();
+	if (!ProtocolSettings)
 	{
-		UpdateOptionsHandle = UpdateOptionsDelegate->Add(FSimpleDelegate::CreateSP(this, &SNameListPicker::UpdateOptionsSource));
+		return;
 	}
+	ProtocolSettings->GetOnDefaultAttributesChanged().AddSP(this, &SNameListPicker::UpdateOptionsSource);
 
-	MaxVisibleItems = InArgs._MaxVisibleItems;
+	EdititableTextBox = SNew(SEditableTextBox)
+		.Font(FAppStyle::GetFontStyle(TEXT("PropertyWindow.NormalFont")))
+		.Text(this, &SNameListPicker::GetCurrentNameLabel)
+		.OnTextCommitted(this, &SNameListPicker::OnTextCommitted);
 
 	// List of selectable names for the dropdown menu
 	SAssignNew(OptionsListView, SListView< TSharedPtr<FName> >)
@@ -85,27 +89,7 @@ void SNameListPicker::Construct(const FArguments& InArgs)
 	SAssignNew(PickerComboButton, SComboButton)
 		.ButtonContent()
 		[
-			SNew(SHorizontalBox)
-
-			+ SHorizontalBox::Slot()
-			.AutoWidth()
-			.HAlign(HAlign_Left)
-			.Padding(0.0f, 0.0f, 5.0f, 0.0f)
-			[
-				SNew(SImage)
-				.Image(FAppStyle::GetBrush("Icons.Warning"))
-				.ToolTipText(LOCTEXT("WarningToolTip", "Value was removed. Please, select another one."))
-				.Visibility(this, &SNameListPicker::GetWarningVisibility)
-			]
-		
-			+ SHorizontalBox::Slot()
-			.AutoWidth()
-			.HAlign(HAlign_Left)
-			.Padding(0.0f)
-			[
-				SNew(STextBlock)
-				.Text(this, &SNameListPicker::GetCurrentNameLabel)
-			]
+			EdititableTextBox.ToSharedRef()
 		]
 		.MenuContent()
 		[
@@ -124,17 +108,14 @@ void SNameListPicker::Construct(const FArguments& InArgs)
 
 void SNameListPicker::UpdateOptionsSource()
 {
-	// Number of options with or without the <None> option
-	const int32 NumOptions = OptionsSourceAttr.Get().Num() + (bCanBeNone ? 1 : 0);
+	// Number of options with the <None> option
+	const int32 NumOptions = OptionsSourceAttr.Get().Num() + 1;
 
 	OptionsSource.Reset();
 	OptionsSource.Reserve(NumOptions);
 
-	// If we can have <None>, it's the first option
-	if (bCanBeNone)
-	{
-		OptionsSource.Add(MakeShared<FName>(FDMXNameListItem::None));
-	}
+	// <None> is the first option
+	OptionsSource.Add(MakeShared<FName>(FDMXNameListItem::None));
 
 	for (const FName& Name : OptionsSourceAttr.Get())
 	{
@@ -142,15 +123,6 @@ void SNameListPicker::UpdateOptionsSource()
 	}
 }
 
-SNameListPicker::~SNameListPicker()
-{
-	if (UpdateOptionsDelegate)
-	{
-		UpdateOptionsDelegate->Remove(UpdateOptionsHandle);
-		UpdateOptionsDelegate = nullptr;
-	}
-}
-
 TSharedRef<ITableRow> SNameListPicker::GenerateNameItemWidget(TSharedPtr<FName> InItem, const TSharedRef<STableViewBase>& OwnerTable) const
 {
 	TSharedPtr<STextBlock> RowTextBlock = nullptr; 
@@ -263,21 +235,6 @@ void SNameListPicker::OnMenuOpened()
 	}
 }
 
-EVisibility SNameListPicker::GetWarningVisibility() const
-{
-	if (!bDisplayWarningIcon || HasMultipleValuesAttribute.Get())
-	{
-		return EVisibility::Collapsed;
-	}
-
-	if (!IsValidAttr.Get())
-	{
-		return EVisibility::Visible;
-	}
-
-	return EVisibility::Collapsed;
-}
-
 EVisibility SNameListPicker::GetSearchBoxVisibility() const
 {
 	return OptionsSource.Num() > MaxVisibleItems ? EVisibility::Visible : EVisibility::Collapsed;
@@ -356,4 +313,24 @@ FText SNameListPicker::GetCurrentNameLabel() const
 	return FText::FromName(CurrentName);
 }
 
+void SNameListPicker::OnTextCommitted(const FText& InNewText, ETextCommit::Type InTextCommit)
+{
+	if (InNewText.IsEmpty())
+	{
+		return;
+	}
+
+	const FName Value = *InNewText.ToString().TrimStartAndEnd();
+	EdititableTextBox->SetText(FText::FromName(Value));
+
+	if (OnValueChangedDelegate.IsBound())
+	{
+		OnValueChangedDelegate.Execute(Value);
+	}
+	else if (!ValueAttribute.IsBound())
+	{
+		ValueAttribute = Value;
+	}
+}
+
 #undef LOCTEXT_NAMESPACE
diff --git a/Engine/Plugins/VirtualProduction/DMX/DMXEngine/Source/DMXEditor/Public/Widgets/SNameListPicker.h b/Engine/Plugins/VirtualProduction/DMX/DMXEngine/Source/DMXEditor/Public/Widgets/SNameListPicker.h
--- a/Engine/Plugins/VirtualProduction/DMX/DMXEngine/Source/DMXEditor/Public/Widgets/SNameListPicker.h
+++ b/Engine/Plugins/VirtualProduction/DMX/DMXEngine/Source/DMXEditor/Public/Widgets/SNameListPicker.h
@@ -10,13 +10,13 @@
 
 class ITableRow;
 class SComboButton;
-template <typename OptionType>
-class SListView;
-template <typename OptionType>
-class SListViewSelectorDropdownMenu;
+class SEditableTextBox;
+template <typename OptionType> class SListView;
+template <typename OptionType> class SListViewSelectorDropdownMenu;
 class SSearchBox;
 class STableViewBase;
 
+
 /**  A widget which allows the user to pick a name of a specified list of names. */
 class DMXEDITOR_API SNameListPicker
 	: public SCompoundWidget
@@ -29,7 +29,6 @@ public:
 	SLATE_BEGIN_ARGS(SNameListPicker)
 		: _ComboButtonStyle(&FCoreStyle::Get().GetWidgetStyle< FComboButtonStyle >("ComboButton"))
 		, _ButtonStyle(nullptr)
-		, _UpdateOptionsDelegate(nullptr)
 		, _bDisplayWarningIcon(false)
 		, _bCanBeNone(false)
 		, _ForegroundColor(FCoreStyle::Get().GetSlateColor("InvertedForeground"))
@@ -45,14 +44,9 @@ public:
 		/** The visual style of the button (overrides ComboButtonStyle) */
 		SLATE_STYLE_ARGUMENT(FButtonStyle, ButtonStyle)
 
-		SLATE_ARGUMENT(FSimpleMulticastDelegate*, UpdateOptionsDelegate)
-
 		/** List of possible names */
 		SLATE_ATTRIBUTE(TArray<FName>, OptionsSource)
 
-		/** Checks if the selected value is no longer available */
-		SLATE_ATTRIBUTE(bool, IsValid)
-
 		/** Display warning icon for a selected invalid value? */
 		SLATE_ARGUMENT(bool, bDisplayWarningIcon)
 
@@ -85,11 +79,7 @@ public:
 	/**  Slate widget construction method */
 	void Construct(const FArguments& InArgs);
 
-	virtual ~SNameListPicker();
-
 private:
-	EVisibility GetWarningVisibility() const;
-
 	EVisibility GetSearchBoxVisibility() const;
 	void OnSearchBoxTextChanged(const FText& InSearchText);
 	void OnSearchBoxTextCommitted(const FText& NewText, ETextCommit::Type CommitInfo);
@@ -116,22 +106,22 @@ private:
 	 */
 	void OnMenuOpened();
 
+	/** Called when text was commited */
+	void OnTextCommitted(const FText& InNewText, ETextCommit::Type InTextCommit);
+
 private:
 	TSharedPtr< SComboButton > PickerComboButton;
 	TSharedPtr< SListView< TSharedPtr<FName> > > OptionsListView;
 	TSharedPtr<SSearchBox> SearchBox;
 	int32 MaxVisibleItems;
 	TSharedPtr< SListViewSelectorDropdownMenu< TSharedPtr<FName> > > NamesListDropdown;
+	TSharedPtr<SEditableTextBox> EdititableTextBox;
 
 	TAttribute<TArray<FName>> OptionsSourceAttr;
 	TArray<TSharedPtr<FName>> OptionsSource;
 	TArray<TSharedPtr<FName>> FilteredOptions;
 
-	FSimpleMulticastDelegate* UpdateOptionsDelegate;
-	FDelegateHandle UpdateOptionsHandle;
-
 	TAttribute<FName> ValueAttribute;
-	TAttribute<bool> IsValidAttr;
 	FOnValueChanged OnValueChangedDelegate;
 	TAttribute<bool> HasMultipleValuesAttribute;
 	bool bCanBeNone;
diff --git a/Engine/Plugins/VirtualProduction/DMX/DMXEngine/Source/DMXRuntime/Private/Library/DMXEntityFixturePatch.cpp b/Engine/Plugins/VirtualProduction/DMX/DMXEngine/Source/DMXRuntime/Private/Library/DMXEntityFixturePatch.cpp
--- a/Engine/Plugins/VirtualProduction/DMX/DMXEngine/Source/DMXRuntime/Private/Library/DMXEntityFixturePatch.cpp
+++ b/Engine/Plugins/VirtualProduction/DMX/DMXEngine/Source/DMXRuntime/Private/Library/DMXEntityFixturePatch.cpp
@@ -580,7 +580,7 @@ TArray<FDMXAttributeName> UDMXEntityFixturePatch::GetAllAttributesInActiveMode()
 	{
 		if (Function.GetLastChannel() <= Cache.GetChannelSpan())
 		{
-			AttributeNames.Add(Function.Attribute.GetName());
+			AttributeNames.Add(Function.Attribute.Name);
 		}
 	}
 
@@ -596,7 +596,7 @@ TMap<FDMXAttributeName, FDMXFixtureFunction> UDMXEntityFixturePatch::GetAttribut
 	{
 		if (Function.GetLastChannel() <= Cache.GetChannelSpan())
 		{
-			FunctionMap.Add(Function.Attribute.GetName(), Function);
+			FunctionMap.Add(Function.Attribute.Name, Function);
 		}
 	}
 
@@ -732,7 +732,7 @@ bool UDMXEntityFixturePatch::IsMapValid(const TMap<FDMXAttributeName, int32>& Fu
 	for (const TPair<FDMXAttributeName, int32>& Elem : FunctionMap)
 	{
 		const bool bContainsAttributeName = Cache.GetFunctions().ContainsByPredicate([Elem](const FDMXFixtureFunction& Function) {
-			return Function.Attribute.GetName() == Elem.Key;
+			return Function.Attribute.Name == Elem.Key;
 		});
 
 		if (!bContainsAttributeName)
@@ -747,7 +747,7 @@ bool UDMXEntityFixturePatch::IsMapValid(const TMap<FDMXAttributeName, int32>& Fu
 bool UDMXEntityFixturePatch::ContainsAttribute(FDMXAttributeName FunctionAttribute) const
 {
 	const bool bContainsAttributeName = Cache.GetFunctions().ContainsByPredicate([&FunctionAttribute](const FDMXFixtureFunction& Function) {
-		return Function.Attribute.GetName() == FunctionAttribute;
+		return Function.Attribute.Name == FunctionAttribute;
 	});
 
 	return bContainsAttributeName;
@@ -785,7 +785,7 @@ TArray<UDMXEntityController*> UDMXEntityFixturePatch::GetRelevantControllers() c
 const FDMXFixtureFunction* UDMXEntityFixturePatch::GetAttributeFunction(const FDMXAttributeName& Attribute) const
 {
 	const FDMXFixtureFunction* FixtureFunctionPtr = Cache.GetFunctions().FindByPredicate([Attribute](const FDMXFixtureFunction& Function) {
-		return Function.Attribute.GetName() == Attribute;
+		return Function.Attribute.Name == Attribute;
 	});
 
 	if (FixtureFunctionPtr && 
@@ -1045,7 +1045,7 @@ bool UDMXEntityFixturePatch::GetMatrixCellChannelsRelative(const FIntPoint& Cell
 		int32 AttributeOffset = 0;
 		for (const FDMXFixtureCellAttribute& CellAttribute : Cache.GetCellAttributes())
 		{
-			AttributeChannelMap.Add(CellAttribute.Attribute.GetName(), RelativeMatrixStartingChannel + AttributeOffset);
+			AttributeChannelMap.Add(CellAttribute.Attribute.Name, RelativeMatrixStartingChannel + AttributeOffset);
 			AttributeOffset += CellAttribute.GetNumChannels();
 		}
 	}
@@ -1064,7 +1064,7 @@ bool UDMXEntityFixturePatch::GetMatrixCellChannelsAbsolute(const FIntPoint& Cell
 		int32 AttributeOffset = 0;
 		for (const FDMXFixtureCellAttribute& CellAttribute : Cache.GetCellAttributes())
 		{
-			AttributeChannelMap.Add(CellAttribute.Attribute.GetName(), AbsoluteMatrixStartingChannel + AttributeOffset);
+			AttributeChannelMap.Add(CellAttribute.Attribute.Name, AbsoluteMatrixStartingChannel + AttributeOffset);
 			AttributeOffset += CellAttribute.GetNumChannels();
 		}
 		
diff --git a/Engine/Plugins/VirtualProduction/DMX/DMXEngine/Source/DMXRuntime/Private/Library/DMXEntityFixturePatchCache.cpp b/Engine/Plugins/VirtualProduction/DMX/DMXEngine/Source/DMXRuntime/Private/Library/DMXEntityFixturePatchCache.cpp
--- a/Engine/Plugins/VirtualProduction/DMX/DMXEngine/Source/DMXRuntime/Private/Library/DMXEntityFixturePatchCache.cpp
+++ b/Engine/Plugins/VirtualProduction/DMX/DMXEngine/Source/DMXRuntime/Private/Library/DMXEntityFixturePatchCache.cpp
@@ -174,10 +174,10 @@ bool FDMXEntityFixturePatchCache::InputDMXSignal(const FDMXSignalSharedPtr& DMXS
 					if (CachedDMXValues.IsValidIndex(FirstAttributeDataIndex) && CachedDMXValues.IsValidIndex(LastAttributeDataIndex))
 					{
 						const uint32 IntValue = UDMXEntityFixtureType::BytesToInt(CellAttribute.DataType, CellAttribute.bUseLSBMode, CachedDMXValues.GetData() + FirstAttributeDataIndex);
-						CellAttributeToRawValueMap.Add(CellAttribute.Attribute.GetName(), IntValue);
+						CellAttributeToRawValueMap.Add(CellAttribute.Attribute.Name, IntValue);
 
 						const float NormalizedValue = (float)IntValue / (float)FDMXConversions::GetSignalFormatMaxValue(CellAttribute.DataType);
-						CellAttributeTNormalizedValueMap.Map.Add(CellAttribute.Attribute.GetName(), NormalizedValue);
+						CellAttributeTNormalizedValueMap.Map.Add(CellAttribute.Attribute.Name, NormalizedValue);
 
 						// Increment attribute offset for the next attribute
 						CurrentAttributeOffset += AttributeDataSize;
@@ -366,7 +366,7 @@ TArray<FDMXAttributeName> FDMXEntityFixturePatchCache::GetAttributeNames() const
 	{
 		for (const FDMXFixtureFunction& Function : Mode.Functions)
 		{
-			AttributeNames.Add(Function.Attribute.GetName());
+			AttributeNames.Add(Function.Attribute.Name);
 		}
 	}
 
@@ -382,7 +382,7 @@ TArray<FDMXAttributeName> FDMXEntityFixturePatchCache::GetCellAttributeNames() c
 
 		for (const FDMXFixtureCellAttribute& CellAttribute : FixtureMatrix.CellAttributes)
 		{
-			MatrixAttributeNames.Add(CellAttribute.Attribute.GetName());
+			MatrixAttributeNames.Add(CellAttribute.Attribute.Name);
 		}
 	}
 
diff --git a/Engine/Plugins/VirtualProduction/DMX/DMXEngine/Source/DMXRuntime/Private/Sequencer/MovieSceneDMXLibrarySection.cpp b/Engine/Plugins/VirtualProduction/DMX/DMXEngine/Source/DMXRuntime/Private/Sequencer/MovieSceneDMXLibrarySection.cpp
--- a/Engine/Plugins/VirtualProduction/DMX/DMXEngine/Source/DMXRuntime/Private/Sequencer/MovieSceneDMXLibrarySection.cpp
+++ b/Engine/Plugins/VirtualProduction/DMX/DMXEngine/Source/DMXRuntime/Private/Sequencer/MovieSceneDMXLibrarySection.cpp
@@ -73,7 +73,7 @@ void FDMXFixturePatchChannel::UpdateNumberOfChannels(bool bResetDefaultValues /*
 		// Add channels for functions that are in mode range and have an Attribute set
 
 		if (Function.GetLastChannel() <= Patch->GetChannelSpan() &&
-			FDMXAttributeName::IsValid(Function.Attribute.GetName()))
+			!Function.Attribute.Name.IsNone())
 		{
 			bool bNewChannel = false;
 			if (!FunctionChannels.IsValidIndex(IdxFunctionChannel))
@@ -82,7 +82,7 @@ void FDMXFixturePatchChannel::UpdateNumberOfChannels(bool bResetDefaultValues /*
 				bNewChannel = true;
 			}					
 
-			FunctionChannels[IdxFunctionChannel].AttributeName = Function.Attribute.GetName();
+			FunctionChannels[IdxFunctionChannel].AttributeName = Function.Attribute.Name;
 
 			if (bResetDefaultValues || bNewChannel)
 			{
@@ -116,7 +116,7 @@ void FDMXFixturePatchChannel::UpdateNumberOfChannels(bool bResetDefaultValues /*
 						bNewChannel = true;
 					}
 
-					FunctionChannels[IdxFunctionChannel].AttributeName = CellAttribute.Attribute.GetName();
+					FunctionChannels[IdxFunctionChannel].AttributeName = CellAttribute.Attribute.Name;
 					FunctionChannels[IdxFunctionChannel].CellCoordinate = CellCoordinates;
 
 					if (bResetDefaultValues || bNewChannel)
@@ -210,7 +210,7 @@ FDMXCachedFunctionChannelInfo::FDMXCachedFunctionChannelInfo(const TArray<FDMXFi
 
 		const FDMXFixtureCellAttribute& CellAttribute = *CellAttributePtr;
 
-		AttributeName = CellAttribute.Attribute.GetName();
+		AttributeName = CellAttribute.Attribute.Name;
 		StartingChannel = AttributeNameChannelMap[FunctionChannel.AttributeName];
 		SignalFormat = CellAttribute.DataType;
 		bLSBMode = CellAttribute.bUseLSBMode;
@@ -232,7 +232,7 @@ FDMXCachedFunctionChannelInfo::FDMXCachedFunctionChannelInfo(const TArray<FDMXFi
 
 		const FDMXFixtureFunction& Function = *FunctionPtr;
 
-		AttributeName = Function.Attribute.GetName();
+		AttributeName = Function.Attribute.Name;
 
 		const int32 PatchChannelOffset = FixturePatch->GetStartingChannel() - 1;
 		StartingChannel = Function.Channel + PatchChannelOffset;
@@ -293,7 +293,7 @@ void FDMXCachedFunctionChannelInfo::GetMatrixCellChannelsAbsoluteNoSorting(UDMXE
 
 		for (const TTuple<const FDMXFixtureCellAttribute*, int32>& AttributeToRelativeChannelOffsetKvp : AttributeToRelativeChannelOffsetMap)
 		{
-			const FName FunctionAttributeName = AttributeToRelativeChannelOffsetKvp.Key->Attribute.GetName();
+			const FName FunctionAttributeName = AttributeToRelativeChannelOffsetKvp.Key->Attribute.Name;
 			const int32 AbsoluteChannel = AbsoluteCellStartingChannel + AttributeToRelativeChannelOffsetKvp.Value;
 
 			check(AbsoluteChannel > 0 && AbsoluteChannel <= DMX_UNIVERSE_SIZE);
diff --git a/Engine/Plugins/VirtualProduction/DMX/DMXPixelMapping/Source/DMXPixelMappingRuntime/Private/DMXPixelMappingRuntimeUtils.cpp b/Engine/Plugins/VirtualProduction/DMX/DMXPixelMapping/Source/DMXPixelMappingRuntime/Private/DMXPixelMappingRuntimeUtils.cpp
--- a/Engine/Plugins/VirtualProduction/DMX/DMXPixelMapping/Source/DMXPixelMappingRuntime/Private/DMXPixelMappingRuntimeUtils.cpp
+++ b/Engine/Plugins/VirtualProduction/DMX/DMXPixelMapping/Source/DMXPixelMappingRuntime/Private/DMXPixelMappingRuntimeUtils.cpp
@@ -25,7 +25,7 @@ void FDMXPixelMappingRuntimeUtils::ConvertNormalizedAttributeValueToChannelValue
 			TMap<int32, uint8> ChannelToValueMap;
 
 			const FDMXFixtureFunction* FunctionPtr = ModePtr->Functions.FindByPredicate([InAttributeName](const FDMXFixtureFunction& Function) {
-				return Function.Attribute.GetName() == InAttributeName;
+				return Function.Attribute.Name == InAttributeName;
 				});
 
 			if (FunctionPtr)
diff --git a/Engine/Plugins/VirtualProduction/DMX/DMXProtocol/Source/DMXProtocol/Private/DMXAttribute.cpp b/Engine/Plugins/VirtualProduction/DMX/DMXProtocol/Source/DMXProtocol/Private/DMXAttribute.cpp
--- a/Engine/Plugins/VirtualProduction/DMX/DMXProtocol/Source/DMXProtocol/Private/DMXAttribute.cpp
+++ b/Engine/Plugins/VirtualProduction/DMX/DMXProtocol/Source/DMXProtocol/Private/DMXAttribute.cpp
@@ -6,37 +6,22 @@
 
 #include "Modules/ModuleManager.h"
 
-IMPLEMENT_DMX_NAMELISTITEM_STATICVARS(FDMXAttributeName)
 
-IMPLEMENT_DMX_NAMELISTITEM_GetAllValues(FDMXAttributeName)
+TArray<FName> FDMXAttributeName::GetPossibleValues()
 {
-	TArray<FName> PossibleValues;
+	TArray<FName> DefaultValues;
 
 	if (const UDMXProtocolSettings* DMXSettings = GetDefault<UDMXProtocolSettings>())
 	{
-		PossibleValues.Reserve(DMXSettings->Attributes.Num());
+		DefaultValues.Reserve(DMXSettings->Attributes.Num());
 
 		for (const FDMXAttribute& Attribute : DMXSettings->Attributes)
 		{
-			PossibleValues.Emplace(Attribute.Name);
+			DefaultValues.Emplace(Attribute.Name);
 		}
 	}
 
-	return PossibleValues;
-}
-
-IMPLEMENT_DMX_NAMELISTITEM_IsValid(FDMXAttributeName)
-{
-	const UDMXProtocolSettings* DMXSettings = GetDefault<UDMXProtocolSettings>();
-	for (const FDMXAttribute& SettingsAttribute : DMXSettings->Attributes)
-	{
-		if (InName == SettingsAttribute.Name)
-		{
-			return true;
-		}
-	}
-
-	return false;
+	return DefaultValues;
 }
 
 FDMXAttributeName::FDMXAttributeName()
@@ -64,17 +49,7 @@ FDMXAttributeName::FDMXAttributeName(const FDMXAttribute& InAttribute)
 
 FDMXAttributeName::FDMXAttributeName(const FName& NameAttribute)
 {
-	const UDMXProtocolSettings* DMXSettings = GetDefault<UDMXProtocolSettings>();
-	for (const FDMXAttribute& SettingsAttribute : DMXSettings->Attributes)
-	{
-		if (SettingsAttribute.Name.IsEqual(NameAttribute))
-		{
-			Name = SettingsAttribute.Name;
-			return;
-		}
-	}
-
-	Name = FDMXNameListItem::None;	
+	Name = NameAttribute;
 }
 
 void FDMXAttributeName::SetFromName(const FName& InName)
@@ -105,12 +80,12 @@ const FDMXAttribute& FDMXAttributeName::GetAttribute() const
 
 FString UDMXAttributeNameConversions::Conv_DMXAttributeToString(const FDMXAttributeName& InAttribute)
 {
-	return InAttribute.GetName().ToString();
+	return InAttribute.Name.ToString();
 }
 
 FName UDMXAttributeNameConversions::Conv_DMXAttributeToName(const FDMXAttributeName& InAttribute)
 {
-	return InAttribute.GetName();
+	return InAttribute.Name;
 }
 
 TArray<FString> FDMXAttribute::GetKeywords() const
@@ -134,4 +109,4 @@ void FDMXAttribute::CleanupKeywords()
 	TArray<FString> KeywordsArray;
 	Keywords.ParseIntoArray(KeywordsArray, TEXT(","), true);
 	Keywords = FString::Join(KeywordsArray, TEXT(", "));
-}
\ No newline at end of file
+}
diff --git a/Engine/Plugins/VirtualProduction/DMX/DMXProtocol/Source/DMXProtocol/Private/DMXProtocolSettings.cpp b/Engine/Plugins/VirtualProduction/DMX/DMXProtocol/Source/DMXProtocol/Private/DMXProtocolSettings.cpp
--- a/Engine/Plugins/VirtualProduction/DMX/DMXProtocol/Source/DMXProtocol/Private/DMXProtocolSettings.cpp
+++ b/Engine/Plugins/VirtualProduction/DMX/DMXProtocol/Source/DMXProtocol/Private/DMXProtocolSettings.cpp
@@ -196,7 +196,7 @@ void UDMXProtocolSettings::PostEditChangeChainProperty(FPropertyChangedChainEven
 			Attribute.CleanupKeywords();
 		}
 
-		FDMXAttributeName::OnValuesChanged.Broadcast();
+		OnDefaultAttributesChanged.Broadcast();
 	}
 	else if (
 		PropertyName == FDMXOutputPortConfig::GetDestinationAddressesPropertyNameChecked() ||
diff --git a/Engine/Plugins/VirtualProduction/DMX/DMXProtocol/Source/DMXProtocol/Public/DMXAttribute.h b/Engine/Plugins/VirtualProduction/DMX/DMXProtocol/Source/DMXProtocol/Public/DMXAttribute.h
--- a/Engine/Plugins/VirtualProduction/DMX/DMXProtocol/Source/DMXProtocol/Public/DMXAttribute.h
+++ b/Engine/Plugins/VirtualProduction/DMX/DMXProtocol/Source/DMXProtocol/Public/DMXAttribute.h
@@ -48,15 +48,15 @@ FORCEINLINE uint32 GetTypeHash(const FDMXAttribute& Attribute)
 
 USTRUCT(BlueprintType, Category = "DMX")
 struct DMXPROTOCOL_API FDMXAttributeName
-	: public FDMXNameListItem
 {
 	GENERATED_BODY()
 
-	DECLARE_DMX_NAMELISTITEM_STATICS(true)
-
 public:
 	//~ Constructors
 
+	/** */
+	static TArray<FName> GetPossibleValues();
+
 	/**
 	 * Default constructor. Assigns Attribute to the first available
 	 * Attribute from the plugin settings
@@ -70,7 +70,11 @@ public:
 	FDMXAttributeName(const FName& NameAttribute);
 
 	//~ FDMXNameListItem interface
-	virtual void SetFromName(const FName& InName) override;
+	void SetFromName(const FName& InName);
+
+	/** The Attribute Name */
+	UPROPERTY(EditAnywhere, Category = "DMX")
+	FName Name;
 
 	const FDMXAttribute& GetAttribute() const;
 	operator const FDMXAttribute& () const { return GetAttribute(); }
diff --git a/Engine/Plugins/VirtualProduction/DMX/DMXProtocol/Source/DMXProtocol/Public/DMXProtocolSettings.h b/Engine/Plugins/VirtualProduction/DMX/DMXProtocol/Source/DMXProtocol/Public/DMXProtocolSettings.h
--- a/Engine/Plugins/VirtualProduction/DMX/DMXProtocol/Source/DMXProtocol/Public/DMXProtocolSettings.h
+++ b/Engine/Plugins/VirtualProduction/DMX/DMXProtocol/Source/DMXProtocol/Public/DMXProtocolSettings.h
@@ -111,6 +111,12 @@ public:
 		return OnSetReceiveDMXEnabledDelegate;
 	}
 
+	/** Gets a Delegate Broadcast when default attributes were changed */
+	FSimpleMulticastDelegate& GetOnDefaultAttributesChanged()
+	{
+		return OnDefaultAttributesChanged;
+	}
+
 #if WITH_EDITOR
 	/** Returns true if Fixture Patches should receive DMX in Editor */
 	bool ShouldAllFixturePatchesReceiveDMXInEditor() const { return bAllFixturePatchesReceiveDMXInEditor; }
@@ -162,6 +168,9 @@ private:
 	/** Broadcast when receive DMX is enabled or disabled */
 	FDMXOnReceiveDMXEnabled OnSetReceiveDMXEnabledDelegate;
 
+	/** Broadcast when default attributets changed */
+	FSimpleMulticastDelegate OnDefaultAttributesChanged;
+
 #if WITH_EDITORONLY_DATA
 	/** If true, all fixture patches receive DMX in Editor. This overrides the fixture patches 'Receive DMX In Editor' property. */
 	UPROPERTY(Config, EditAnywhere, Category = "DMX|Communication Settings", Meta = (DisplayName = "All Fixture Patches receive DMX in Editor"))
