DMX - Fix CIS Build Engine Localization - Warning: Text conflict from LOCTEXT macro for namespace "DMXMVRExporter" and key "MVRExportCreateZipFailedReason"

diff --git a/Engine/Plugins/VirtualProduction/DMX/DMXEngine/Source/DMXEditor/Private/Exporters/DMXMVRExporter.cpp b/Engine/Plugins/VirtualProduction/DMX/DMXEngine/Source/DMXEditor/Private/Exporters/DMXMVRExporter.cpp
--- a/Engine/Plugins/VirtualProduction/DMX/DMXEngine/Source/DMXEditor/Private/Exporters/DMXMVRExporter.cpp
+++ b/Engine/Plugins/VirtualProduction/DMX/DMXEngine/Source/DMXEditor/Private/Exporters/DMXMVRExporter.cpp
@@ -68,7 +68,7 @@ void FDMXMVRExporter::ExportInternal(UDMXLibrary* DMXLibrary, const FString& Fil
 	ZipThirdPartyData(Zip, GeneralSceneDescription);
 	if (!Zip->SaveToFile(FilePathAndName))
 	{
-		OutErrorReason = FText::Format(LOCTEXT("MVRExportCreateZipFailedReason", "File is not writable or locked by another process. Cannot export {0}."), FText::FromString(FilePathAndName));
+		OutErrorReason = FText::Format(LOCTEXT("MVRExportWriteZipFailedReason", "File is not writable or locked by another process. Cannot export {0}."), FText::FromString(FilePathAndName));
 		return;
 	}
 }
diff --git a/Engine/Plugins/VirtualProduction/DMX/DMXEngine/Source/DMXEditor/Private/Widgets/FixturePatch/SDMXMVRFixtureList.cpp b/Engine/Plugins/VirtualProduction/DMX/DMXEngine/Source/DMXEditor/Private/Widgets/FixturePatch/SDMXMVRFixtureList.cpp
--- a/Engine/Plugins/VirtualProduction/DMX/DMXEngine/Source/DMXEditor/Private/Widgets/FixturePatch/SDMXMVRFixtureList.cpp
+++ b/Engine/Plugins/VirtualProduction/DMX/DMXEngine/Source/DMXEditor/Private/Widgets/FixturePatch/SDMXMVRFixtureList.cpp
@@ -164,7 +164,7 @@ private:
 			}
 			else if (Item->GetFixtureType()->Modes.IsEmpty())
 			{
-				const FText IssueText = LOCTEXT("NoFixtureTypeIssue", "Fixture Type has no Modes defined.");
+				const FText IssueText = LOCTEXT("NoModesIssue", "Fixture Type has no Modes defined.");
 				ItemToIssueMap.Add(Item, IssueText);
 			}
 			else if (Item->GetFixtureType() && Item->GetModeIndex() == INDEX_NONE)
