DMX - Fix some Fixture Types don't create a matrix fixture from GDTF where they could

diff --git a/Engine/Plugins/VirtualProduction/DMX/DMXEngine/Source/DMXEditor/Private/DMXInitializeFixtureTypeFromGDTFHelper.cpp b/Engine/Plugins/VirtualProduction/DMX/DMXEngine/Source/DMXEditor/Private/DMXInitializeFixtureTypeFromGDTFHelper.cpp
--- a/Engine/Plugins/VirtualProduction/DMX/DMXEngine/Source/DMXEditor/Private/DMXInitializeFixtureTypeFromGDTFHelper.cpp
+++ b/Engine/Plugins/VirtualProduction/DMX/DMXEngine/Source/DMXEditor/Private/DMXInitializeFixtureTypeFromGDTFHelper.cpp
@@ -11,8 +11,11 @@
 
 #include "XmlFile.h"
 #include "XmlNode.h"
+#include "Algo/Copy.h"
 #include "Algo/Count.h"
 #include "Algo/Find.h"
+#include "Algo/MaxElement.h"
+#include "Algo/MinElement.h"
 
 
 #define LOCTEXT_NAMESPACE "DMXInitializeFixtureTypeFromGDTFHelper"
@@ -57,6 +60,229 @@ namespace UE::DMX::DMXInitializeFixtureTypeFromGDTFHelper::Private
 
 		return Attribute;
 	}
+
+	/** Helper to interprete a GDTF Channel */
+	class FDMXGDTFChannelInterpreter
+	{
+	public:
+		/** Initializes the bIsMatrix member */
+		static bool IsMatrix(const FXmlNode* InDMXChannelNode, const TArray<const FXmlNode*>& InGeometryNodes)
+		{
+			return GetNumCells(InDMXChannelNode, InGeometryNodes) > 1;
+		}
+
+		/** Properties of a channel */
+		struct FChannelProperties
+		{
+			FString AttributeName;
+			int32 FirstChannel;
+			EDMXFixtureSignalFormat SignalFormat;
+			bool bLSBMode;
+			int32 NumCells;
+			int32 DefaultValue;
+		};
+
+		/** Returns the properties of the channel */
+		static bool GetChannelProperties(const FXmlNode* InDMXChannelNode, const TArray<const FXmlNode*>& GeometryNodes, FChannelProperties& OutChannelProperties)
+		{
+			if (!ensureMsgf(InDMXChannelNode->GetTag() == TEXT("dmxchannel") || InDMXChannelNode->GetTag() == TEXT("channel"), TEXT("Trying to read Channel node, but node is tagged '%s'."), *InDMXChannelNode->GetTag()))
+			{
+				return false;
+			}
+
+			OutChannelProperties.NumCells = GetNumCells(InDMXChannelNode, GeometryNodes);
+
+			if (!GetAttributeName(InDMXChannelNode, OutChannelProperties.AttributeName))
+			{
+				return false;
+			}
+
+			GetDefaultValue(InDMXChannelNode, OutChannelProperties.DefaultValue);
+			GetDataType(InDMXChannelNode, OutChannelProperties.SignalFormat, OutChannelProperties.bLSBMode);
+
+			TArray<int32> Offsets;
+			if (IsMatrix(InDMXChannelNode, GeometryNodes))
+			{
+				Offsets = GetMatrixOffsetArray(InDMXChannelNode, GeometryNodes);
+			}
+			else
+			{
+				Offsets = GetOffsetArray(InDMXChannelNode);
+			}
+			if (Offsets.IsEmpty())
+			{
+				return false;
+			}
+			OutChannelProperties.FirstChannel = *Algo::MinElement(Offsets);
+
+			return true;
+		}
+
+	private:
+		/** Returns the number of cells if it is a matrix, or 1, if it is not a matrix */
+		static int32 GetNumCells(const FXmlNode* InDMXChannelNode, const TArray<const FXmlNode*>& GeometryNodes)
+		{
+			if (!ensureMsgf(InDMXChannelNode->GetTag() == TEXT("dmxchannel") || InDMXChannelNode->GetTag() == TEXT("channel"), TEXT("Trying to read Channel node, but node is tagged '%s'."), *InDMXChannelNode->GetTag()))
+			{
+				return false;
+			}
+
+			constexpr TCHAR GeometryTag[] = TEXT("Geometry");
+			const FString GeometryName = InDMXChannelNode->GetAttribute(GeometryTag);
+
+			// Count the number of geometries this channel references
+			const int32 NumReferencesInGeometry = Algo::CountIf(GeometryNodes, [&GeometryName, GeometryTag](const FXmlNode* GeometryNode)
+				{
+					if (GeometryNode->GetAttribute(GeometryTag) == GeometryName)
+					{
+						return true;
+					}
+
+					return false;
+				});
+
+			return FMath::Max(1, NumReferencesInGeometry);
+		}
+
+		/** 
+		 * Returns the attribute name of the channel node. 
+		 * As per legacy considers only the first function - FDMXFixtureFunction is not in support of more than one. 
+		 */
+		static bool GetAttributeName(const FXmlNode* InDMXChannelNode, FString& OutName)
+		{
+			constexpr TCHAR LogicalChannelTag[] = TEXT("LogicalChannel");
+			const FXmlNode* LogicalChannelNode = InDMXChannelNode->FindChildNode(LogicalChannelTag);
+			if (!LogicalChannelNode)
+			{
+				return false;
+			}
+
+			constexpr TCHAR ChannelFunctionTag[] = TEXT("ChannelFunction");
+			const FXmlNode* ChannelFunctionNode = LogicalChannelNode->FindChildNode(ChannelFunctionTag);
+			if (!ChannelFunctionNode)
+			{
+				return false;
+			}
+
+			// Name
+			constexpr TCHAR NameTag[] = TEXT("Name");
+			FString Name = ChannelFunctionNode->GetAttribute(NameTag);
+			
+			OutName = Name;
+			return true;
+		}
+
+		/** 
+		 * Gets the default value of the channel. Fails quietly if no default value present.
+		 * As per legacy considers only the first function - FDMXFixtureFunction is not in support of more than one. 
+		 */
+		static void GetDefaultValue(const FXmlNode* InDMXChannelNode, int32& OutDefaultValue)
+		{
+			constexpr TCHAR LogicalChannelTag[] = TEXT("LogicalChannel");
+			const FXmlNode* LogicalChannelNode = InDMXChannelNode->FindChildNode(LogicalChannelTag);
+			if (!LogicalChannelNode)
+			{
+				return;
+			}
+
+			constexpr TCHAR ChannelFunctionTag[] = TEXT("ChannelFunction");
+			const FXmlNode* ChannelFunctionNode = LogicalChannelNode->FindChildNode(ChannelFunctionTag);
+			if (!ChannelFunctionNode)
+			{
+				return;
+			}
+
+			constexpr TCHAR DefaultTag[] = TEXT("Default");
+			if (LexTryParseString(OutDefaultValue, *ChannelFunctionNode->GetAttribute(DefaultTag)))
+			{
+				return;
+			}
+
+			return;
+		}
+
+		/** Returns an array of offsets from an XmlNode */
+		static TArray<int32> GetOffsetArray(const FXmlNode* XmlNode)
+		{
+			constexpr TCHAR OffsetTag[] = TEXT("DMXOffset");
+			TArray<int32> OffsetArray;
+			TArray<FString> OffsetStrArray;
+			FindAttributeEvenIfDMXSubstringIsMissing(*XmlNode, OffsetTag).ParseIntoArray(OffsetStrArray, TEXT(","));
+
+			for (int32 OffsetIndex = 0; OffsetIndex < OffsetStrArray.Num(); ++OffsetIndex)
+			{
+				int32 OffsetValue;
+				LexTryParseString(OffsetValue, *OffsetStrArray[OffsetIndex]);
+				OffsetArray.Add(OffsetValue);
+			}
+
+			return OffsetArray;
+		}
+
+		/** Returns an array of offsets from an XmlNode */
+		static TArray<int32> GetMatrixOffsetArray(const FXmlNode* InDMXChannelNode, const TArray<const FXmlNode*>& GeometryNodes)
+		{
+			constexpr TCHAR GeometryTag[] = TEXT("Geometry");
+			const FString GeometryName = InDMXChannelNode->GetAttribute(GeometryTag);
+
+			TArray<const FXmlNode*> ReferencedGeometryNodes;
+			Algo::CopyIf(GeometryNodes, ReferencedGeometryNodes, [GeometryTag, GeometryName](const FXmlNode* GeometryNode)
+				{
+					if (GeometryNode->GetAttribute(GeometryTag) == GeometryName)
+					{
+						return true;
+					}
+
+					return false;
+				});
+
+			if (ReferencedGeometryNodes.IsEmpty())
+			{
+				// Try to parse as normal channel instead
+				return GetOffsetArray(InDMXChannelNode);
+			}
+
+			if (const FXmlNode* XmlNode = ReferencedGeometryNodes[0])
+			{
+				constexpr TCHAR BreakTag[] = TEXT("Break");
+				if (const FXmlNode* BreakNode = XmlNode->FindChildNode(BreakTag))
+				{
+					return GetOffsetArray(BreakNode);
+				}
+			}
+
+			return TArray<int32>();
+		}
+
+		/** Returns the data type offset array implies */
+		static void GetDataType(const FXmlNode* InDMXChannelNode, EDMXFixtureSignalFormat& OutSignalFormat, bool& OutLSBOrder)
+		{
+			// For the data type we always need to refer to the offset of the channel, not the geometry
+			TArray<int32> ChannelOffestArray = GetOffsetArray(InDMXChannelNode);
+
+			// Compute number of used addresses
+			int32 AddressMin = DMX_MAX_ADDRESS;
+			int32 AddressMax = 0;
+			for (const int32& Address : ChannelOffestArray)
+			{
+				AddressMin = FMath::Min(AddressMin, Address);
+				AddressMax = FMath::Max(AddressMax, Address);
+			}
+			const int32 NumUsedAddresses = FMath::Clamp(AddressMax - AddressMin + 1, 1, DMX_MAX_FUNCTION_SIZE);
+
+			OutSignalFormat = static_cast<EDMXFixtureSignalFormat>(NumUsedAddresses - 1);
+
+			// Offsets represent the channels in MSB order. If they are in reverse order, it means this Function uses LSB format.
+			if (ChannelOffestArray.Num() > 1)
+			{
+				OutLSBOrder = ChannelOffestArray[0] > ChannelOffestArray[1];
+			}
+			else
+			{
+				OutLSBOrder = false;
+			}
+		}
+	};
 };
 
 bool FDMXInitializeFixtureTypeFromGDTFHelper::GenerateModesFromGDTF(UDMXEntityFixtureType& InOutFixtureType, const UDMXImportGDTF& InGDTF)
@@ -182,166 +408,122 @@ bool FDMXInitializeFixtureTypeFromGDTFHelper::GenerateMode(const FXmlNode& InFix
 	constexpr TCHAR NameTag[] = TEXT("Name");
 	Mode.ModeName = InDMXModeNode.GetAttribute(NameTag); // Not using the class's default attribute
 
-	int32 NextFunctionStartingChannel = 1;
 	for (const FXmlNode* DMXChannelNode : ChannelsNode->GetChildrenNodes())
 	{
-		FDMXFixtureFunction Function;
-		if (!GetChannelPropertiesWithoutStartingChannel(*DMXChannelNode, Function))
+		using namespace UE::DMX::DMXInitializeFixtureTypeFromGDTFHelper::Private;
+
+		FDMXGDTFChannelInterpreter::FChannelProperties ChannelProperties;
+		if (!FDMXGDTFChannelInterpreter::GetChannelProperties(DMXChannelNode, GeometryNodes, ChannelProperties))
 		{
 			continue;
 		}
 
-		UE_LOG(LogTemp, Warning, TEXT("Channel: %i"), NextFunctionStartingChannel);
-
-		constexpr TCHAR GeometryTag[] = TEXT("Geometry");
-		const FString GeometryName = DMXChannelNode->GetAttribute(GeometryTag);
-
-		// Count the number of geometries this channel references
-		const int32 NumReferencesInGeometry = Algo::CountIf(GeometryNodes, [&GeometryName, GeometryTag](const FXmlNode* GeometryNode)
-			{
-				if (GeometryNode->GetAttribute(GeometryTag) == GeometryName)
-				{
-					return true;
-				}
-
-				return false;
-			});
-
-		if (NumReferencesInGeometry < 2)
-		{
-			// A single or no referenced Geometry - Create a normal function.
-			Function.Channel = FMath::Max(NextFunctionStartingChannel, Mode.FixtureMatrixConfig.GetLastChannel() + 1);
-			Mode.Functions.Add(Function);
-		}
-		else
+		if (FDMXGDTFChannelInterpreter::IsMatrix(DMXChannelNode, GeometryNodes))
 		{
-			// Try to add a matrix function instead of a common function
-			if (Mode.bFixtureMatrixEnabled &&
-				Mode.Functions.Num() > 0 &&
-				Mode.FixtureMatrixConfig.GetLastChannel() > Mode.Functions.Last().GetLastChannel())
-			{
-				UE_LOG(LogDMXEditor, Warning, TEXT("Mode '%s' contains many matrices, but this version of Unreal Engine only supports one matrix. Skipping import of mode."), *Mode.ModeName)
-					Mode.Functions.Reset();
-				Mode.bFixtureMatrixEnabled = false;
-				Mode.FixtureMatrixConfig.CellAttributes.Reset();
-				Mode.ModeName = FString::Printf(TEXT("n/a '%s' [not supported in this Engine Version]"), *Mode.ModeName);
-				break;
-			}
-
 			FDMXFixtureCellAttribute MatrixAttribute;
-			MatrixAttribute.Attribute = Function.Attribute;
-			MatrixAttribute.bUseLSBMode = Function.bUseLSBMode;
-			MatrixAttribute.DataType = Function.DataType;
-			MatrixAttribute.DefaultValue = Function.DefaultValue;
+			MatrixAttribute.Attribute = FDMXAttributeName(*ChannelProperties.AttributeName);
+			MatrixAttribute.bUseLSBMode = ChannelProperties.bLSBMode;
+			MatrixAttribute.DataType = ChannelProperties.SignalFormat;
+			MatrixAttribute.DefaultValue = ChannelProperties.DefaultValue;
 
 			Mode.FixtureMatrixConfig.CellAttributes.Add(MatrixAttribute);
-			Mode.FixtureMatrixConfig.YCells = NumReferencesInGeometry;
+			Mode.FixtureMatrixConfig.YCells = ChannelProperties.NumCells;
 			if (!Mode.bFixtureMatrixEnabled)
 			{
-				Mode.FixtureMatrixConfig.FirstCellChannel = NextFunctionStartingChannel;
+				Mode.FixtureMatrixConfig.FirstCellChannel = ChannelProperties.FirstChannel;
 				Mode.bFixtureMatrixEnabled = true;
 			}
 		}
+		else
+		{
+			FDMXFixtureFunction Function;
+			Function.FunctionName = ChannelProperties.AttributeName;
+			Function.Attribute = FDMXAttributeName(*ChannelProperties.AttributeName);
+			Function.Channel = ChannelProperties.FirstChannel;
+			Function.bUseLSBMode = ChannelProperties.bLSBMode;
+			Function.DataType = ChannelProperties.SignalFormat;
+			Function.DefaultValue = ChannelProperties.DefaultValue;
 
-		NextFunctionStartingChannel += Function.GetNumChannels();
+			Mode.Functions.Add(Function);
+		}
 	}
 
+	CleanupMode(Mode);
+
 	OutMode = Mode;
 	return true;
 }
 
-bool FDMXInitializeFixtureTypeFromGDTFHelper::GetChannelPropertiesWithoutStartingChannel(const FXmlNode& DMXChannelNode, FDMXFixtureFunction& OutFunction) const
+void FDMXInitializeFixtureTypeFromGDTFHelper::CleanupMode(FDMXFixtureMode& InOutMode) const
 {
-	if (!ensureMsgf(DMXChannelNode.GetTag() == TEXT("dmxchannel") || DMXChannelNode.GetTag() == TEXT("channel"), TEXT("Trying to read Channel node, but node is tagged '%s'."), *DMXChannelNode.GetTag()))
+	TOptional<TRange<int32>> MatrixRange;
+	if (InOutMode.bFixtureMatrixEnabled)
 	{
-		return false;
-	}
-
-	// As per legacy consider only the first function, although there may be many.
-	constexpr TCHAR LogicalChannelTag[] = TEXT("LogicalChannel");
-	const FXmlNode* LogicalChannelNode = DMXChannelNode.FindChildNode(LogicalChannelTag);
-	if (!LogicalChannelNode)
-	{
-		return false;
-	}
-
-	constexpr TCHAR ChannelFunctionTag[] = TEXT("ChannelFunction");
-	const FXmlNode* ChannelFunctionNode = LogicalChannelNode->FindChildNode(ChannelFunctionTag);
-	if (!ChannelFunctionNode)
-	{
-		return false;
-	}
+		// Only one single, consecutive matrix is supported by the engine in this version.
+		MatrixRange = TRange<int32>(InOutMode.FixtureMatrixConfig.FirstCellChannel, InOutMode.FixtureMatrixConfig.GetLastChannel() + 1);
+		const FDMXFixtureFunction* OverlappingFunction = Algo::FindByPredicate(InOutMode.Functions, [MatrixRange](const FDMXFixtureFunction& Function)
+			{
+				const TRange<int32> FunctionRange(Function.Channel, Function.GetLastChannel() + 1);
+				return FunctionRange.Overlaps(MatrixRange.GetValue());
+			});
 
-	// Name
-	constexpr TCHAR NameTag[] = TEXT("Name");
-	const FString Name = ChannelFunctionNode->GetAttribute(NameTag);
-
-	OutFunction.FunctionName = Name.IsEmpty() ? TEXT("<empty>") : Name;
-	OutFunction.Attribute.Name = Name.IsEmpty() ? NAME_None : FName(*Name);
-	
-	// Default value
-	constexpr TCHAR DefaultTag[] = TEXT("Default");
-	int32 DefaultValue;
-	if (LexTryParseString(DefaultValue, *ChannelFunctionNode->GetAttribute(DefaultTag)))
-	{
-		OutFunction.DefaultValue = DefaultValue;
+		if (OverlappingFunction)
+		{
+			UE_LOG(LogDMXEditor, Warning, TEXT("Mode '%s' contains many matrices, but this version of Unreal Engine only supports one matrix. Skipping import of mode."), *InOutMode.ModeName);
+			InOutMode.Functions.Reset();
+			InOutMode.bFixtureMatrixEnabled = false;
+			InOutMode.FixtureMatrixConfig.CellAttributes.Reset();
+			InOutMode.ModeName = FString::Printf(TEXT("n/a '%s' [not supported in this Engine Version]"), *InOutMode.ModeName);
+		}
 	}
 
-	// Word size and endianness
-	constexpr TCHAR OffsetTag[] = TEXT("Offset");
-	TArray<int32> OffsetArray;
-	TArray<FString> OffsetStrArray;
-	DMXChannelNode.GetAttribute(OffsetTag).ParseIntoArray(OffsetStrArray, TEXT(","));
-
-	for (int32 OffsetIndex = 0; OffsetIndex < OffsetStrArray.Num(); ++OffsetIndex)
+	// Make sure functions are in consecutive order, insert 'reserved' channels where no channel is specified
+	if (!InOutMode.Functions.IsEmpty())
 	{
-		int32 OffsetValue;
-		LexTryParseString(OffsetValue, *OffsetStrArray[OffsetIndex]);
-		OffsetArray.Add(OffsetValue);
-	}
+		Algo::SortBy(InOutMode.Functions, [](const FDMXFixtureFunction& Function)
+			{
+				return Function.Channel;
+			});
 
-	if (OffsetArray.Num() > 0)
-	{
-		// Compute number of used addresses
-		int32 AddressMin = DMX_MAX_ADDRESS;
-		int32 AddressMax = 0;
-		for (const int32& Address : OffsetArray)
+		// Fill in blank functions
+		const int32 LastFunctionChannel = InOutMode.Functions.Last().Channel;
+		for (int32 Channel = 1; Channel < LastFunctionChannel; Channel++)
 		{
-			AddressMin = FMath::Min(AddressMin, Address);
-			AddressMax = FMath::Max(AddressMax, Address);
-		}
-		const int32 NumUsedAddresses = FMath::Clamp(AddressMax - AddressMin + 1, 1, DMX_MAX_FUNCTION_SIZE);
+			const FDMXFixtureFunction* FunctionOnChannelPtr = Algo::FindByPredicate(InOutMode.Functions, [Channel](const FDMXFixtureFunction& Function)
+				{
+					const TRange<int32> FunctionRange(Function.Channel, Function.GetLastChannel() + 1);
+					return FunctionRange.Contains(Channel);
+				});
 
-		OutFunction.DataType = static_cast<EDMXFixtureSignalFormat>(NumUsedAddresses - 1);
+			const bool bChannelHasFunction = FunctionOnChannelPtr != nullptr;
+			const bool bChannelHasMatrix = MatrixRange.IsSet() && MatrixRange.GetValue().Contains(Channel);
 
-		// Offsets represent the channels in MSB order. If they are in reverse order, it means this Function uses LSB format.
-		if (OffsetArray.Num() > 1)
-		{
-			OutFunction.bUseLSBMode = OffsetArray[0] > OffsetArray[1];
-		}
-		else
-		{
-			OutFunction.bUseLSBMode = false;
+			if (!bChannelHasFunction && !bChannelHasMatrix)
+			{
+				FDMXFixtureFunction EmptyFunction;
+				EmptyFunction.Channel = Channel;
+				EmptyFunction.FunctionName = TEXT("<empty>");
+				InOutMode.Functions.Add(EmptyFunction);
+			}
 		}
 	}
-	else
-	{
-		OutFunction.DataType = EDMXFixtureSignalFormat::E8Bit;
-	}
 
-	return true;
+	Algo::SortBy(InOutMode.Functions, [](const FDMXFixtureFunction& Function)
+		{
+			return Function.Channel;
+		});
 }
 
 void FDMXInitializeFixtureTypeFromGDTFHelper::CleanupAttributes(UDMXEntityFixtureType& InOutFixtureType) const
 {
 	// Get Protocol Setting's default attributes
 	const UDMXProtocolSettings* ProtocolSettings = GetDefault<UDMXProtocolSettings>();
-	TMap<FName, TArray<FString> > AttributeNameToKeywordsMap;
+	TMap<FName, TArray<FString> > SettingsAttributeNameToKeywordsMap;
 	for (const FDMXAttribute& Attribute : ProtocolSettings->Attributes)
 	{
 		TArray<FString> Keywords = Attribute.GetKeywords();
 
-		AttributeNameToKeywordsMap.Emplace(Attribute.Name, Keywords);
+		SettingsAttributeNameToKeywordsMap.Emplace(Attribute.Name, Keywords);
 	}
 
 	TArray<FName> AssignedAttributeNames;
@@ -349,7 +531,7 @@ void FDMXInitializeFixtureTypeFromGDTFHelper::CleanupAttributes(UDMXEntityFixtur
 	{
 		for (FDMXFixtureFunction& Function : Mode.Functions)
 		{
-			const TTuple<FName, TArray<FString>>* AttributeNameToKeywordPairPtr = Algo::FindByPredicate(AttributeNameToKeywordsMap, [Function](const TTuple<FName, TArray<FString>> AttributeNameToKeywordPair)
+			const TTuple<FName, TArray<FString>>* SettingsAttributeNameToKeywordPairPtr = Algo::FindByPredicate(SettingsAttributeNameToKeywordsMap, [Function](const TTuple<FName, TArray<FString>> AttributeNameToKeywordPair)
 				{
 					if (Function.Attribute.Name == AttributeNameToKeywordPair.Key ||
 						AttributeNameToKeywordPair.Value.Contains(Function.Attribute.Name.ToString()))
@@ -360,22 +542,14 @@ void FDMXInitializeFixtureTypeFromGDTFHelper::CleanupAttributes(UDMXEntityFixtur
 					return false;
 				});
 
-			if (AttributeNameToKeywordPairPtr && !AssignedAttributeNames.Contains(AttributeNameToKeywordPairPtr->Key))
+			if (SettingsAttributeNameToKeywordPairPtr && !AssignedAttributeNames.Contains(SettingsAttributeNameToKeywordPairPtr->Key))
 			{
-				Function.Attribute.Name = AttributeNameToKeywordPairPtr->Key;
-				AssignedAttributeNames.Add(AttributeNameToKeywordPairPtr->Key);
+				Function.Attribute.Name = SettingsAttributeNameToKeywordPairPtr->Key;
+				AssignedAttributeNames.Add(SettingsAttributeNameToKeywordPairPtr->Key);
 			}
 			else
 			{
-				if (AssignedAttributeNames.Contains(Function.Attribute.Name))
-				{
-					UE_LOG(LogDMXEditor, Warning, TEXT("Attribute '%s' occurs twice in Fixture Type '%s'. Reseting second occurence to None"), *Function.Attribute.Name.ToString(), *InOutFixtureType.GetName());
-				}
-				else
-				{
-					// Attribute not specified in project settings, but unique
-					AssignedAttributeNames.Add(Function.Attribute.Name);
-				}
+				AssignedAttributeNames.Add(Function.Attribute.Name);
 			}
 		}
 	}
diff --git a/Engine/Plugins/VirtualProduction/DMX/DMXEngine/Source/DMXEditor/Private/DMXInitializeFixtureTypeFromGDTFHelper.h b/Engine/Plugins/VirtualProduction/DMX/DMXEngine/Source/DMXEditor/Private/DMXInitializeFixtureTypeFromGDTFHelper.h
--- a/Engine/Plugins/VirtualProduction/DMX/DMXEngine/Source/DMXEditor/Private/DMXInitializeFixtureTypeFromGDTFHelper.h
+++ b/Engine/Plugins/VirtualProduction/DMX/DMXEngine/Source/DMXEditor/Private/DMXInitializeFixtureTypeFromGDTFHelper.h
@@ -34,8 +34,8 @@ private:
 	/** Generates a Mode from specified Nodes. Returns true on success. */
 	bool GenerateMode(const FXmlNode& InFixtureTypeNode, const FXmlNode& InDMXModeNode, FDMXFixtureMode& OutMode) const;
 
-	/** Gets the properties of the channel node, as a fixture function. Does not set the starting channel. */
-	bool GetChannelPropertiesWithoutStartingChannel(const FXmlNode& DMXChannelNode, FDMXFixtureFunction& OutFunction) const;
+	/** Makes sure the mode can be used by the engine */
+	void CleanupMode(FDMXFixtureMode& InOutMode) const;
 
 	/** Makes sure that attributes are replaced with attributes specified in protocol settings where possible */
 	void CleanupAttributes(UDMXEntityFixtureType& InOutFixtureType) const;
