DMX - Various MVR objects now correctly use the FixtureID as short name of a fixture

diff --git a/Engine/Plugins/VirtualProduction/DMX/DMXEngine/Source/DMXEditor/Private/Widgets/FixturePatch/DMXMVRFixtureListItem.cpp b/Engine/Plugins/VirtualProduction/DMX/DMXEngine/Source/DMXEditor/Private/Widgets/FixturePatch/DMXMVRFixtureListItem.cpp
--- a/Engine/Plugins/VirtualProduction/DMX/DMXEngine/Source/DMXEditor/Private/Widgets/FixturePatch/DMXMVRFixtureListItem.cpp
+++ b/Engine/Plugins/VirtualProduction/DMX/DMXEngine/Source/DMXEditor/Private/Widgets/FixturePatch/DMXMVRFixtureListItem.cpp
@@ -93,20 +93,20 @@ void FDMXMVRFixtureListItem::SetFixturePatchName(const FString& InDesiredName, F
 	}
 }
 
-int32 FDMXMVRFixtureListItem::GetUnitNumber() const
+FString FDMXMVRFixtureListItem::GetFixtureID() const
 {
-	return MVRFixtureNode->UnitNumber;
+	return MVRFixtureNode->FixtureID;
 }
 
-void FDMXMVRFixtureListItem::SetUnitNumber(int32 InUnitNumber)
+void FDMXMVRFixtureListItem::SetFixtureID(int32 InFixtureID)
 {
 	UDMXLibrary* DMXLibrary = GetDMXLibrary();
 	if (MVRFixtureNode && DMXLibrary)
 	{
-		const FScopedTransaction SetMVRFixtureUnitNumberTransaction(LOCTEXT("SetMVRFixtureUnitNumberTransaction", "Set MVR Fixture UnitNumber"));
+		const FScopedTransaction SetMVRFixtureFixtureIDTransaction(LOCTEXT("SetMVRFixtureFixtureIDTransaction", "Set MVR Fixture ID"));
 		DMXLibrary->PreEditChange(UDMXLibrary::StaticClass()->FindPropertyByName(UDMXLibrary::GetGeneralSceneDescriptionPropertyName()));
 
-		MVRFixtureNode->UnitNumber = InUnitNumber;
+		MVRFixtureNode->FixtureID = FString::FromInt(InFixtureID);
 
 		DMXLibrary->PostEditChange();
 	}
diff --git a/Engine/Plugins/VirtualProduction/DMX/DMXEngine/Source/DMXEditor/Private/Widgets/FixturePatch/DMXMVRFixtureListItem.h b/Engine/Plugins/VirtualProduction/DMX/DMXEngine/Source/DMXEditor/Private/Widgets/FixturePatch/DMXMVRFixtureListItem.h
--- a/Engine/Plugins/VirtualProduction/DMX/DMXEngine/Source/DMXEditor/Private/Widgets/FixturePatch/DMXMVRFixtureListItem.h
+++ b/Engine/Plugins/VirtualProduction/DMX/DMXEngine/Source/DMXEditor/Private/Widgets/FixturePatch/DMXMVRFixtureListItem.h
@@ -37,11 +37,11 @@ public:
 	/** Sets the Name of the Fixture Patch */
 	void SetFixturePatchName(const FString& InDesiredName, FString& OutNewName);
 
-	/** Returns the Unit Number */
-	int32 GetUnitNumber() const;
+	/** Returns the Fixture ID */
+	FString GetFixtureID() const;
 
-	/** Sets the Unit Number */
-	void SetUnitNumber(int32 InUnitNumber);
+	/** Sets the Fixture ID. Note, as other hard- and software, we allow to set integer values only. */
+	void SetFixtureID(int32 InFixtureID);
 
 	/** Returns the Name of the MVR Fixture */
 	const FString& GetMVRFixtureName() const;
diff --git a/Engine/Plugins/VirtualProduction/DMX/DMXEngine/Source/DMXEditor/Private/Widgets/FixturePatch/SDMXMVRFixtureList.cpp b/Engine/Plugins/VirtualProduction/DMX/DMXEngine/Source/DMXEditor/Private/Widgets/FixturePatch/SDMXMVRFixtureList.cpp
--- a/Engine/Plugins/VirtualProduction/DMX/DMXEngine/Source/DMXEditor/Private/Widgets/FixturePatch/SDMXMVRFixtureList.cpp
+++ b/Engine/Plugins/VirtualProduction/DMX/DMXEngine/Source/DMXEditor/Private/Widgets/FixturePatch/SDMXMVRFixtureList.cpp
@@ -44,8 +44,11 @@ public:
 		TMap<TSharedPtr<FDMXMVRFixtureListItem>, FText> FixtureTypeIssues = GetFixtureTypeIssues();
 		AppendConflictTexts(FixtureTypeIssues, AccumulatedConflicts);
 
-		TMap<TSharedPtr<FDMXMVRFixtureListItem>, FText> UnitNumberConflicts = GetUnitNumberConflicts();
-		AppendConflictTexts(UnitNumberConflicts, AccumulatedConflicts);
+		TMap<TSharedPtr<FDMXMVRFixtureListItem>, FText> FixtureIDIssues = GetFixtureIDIssues();
+		AppendConflictTexts(FixtureIDIssues, AccumulatedConflicts);
+
+		TMap<TSharedPtr<FDMXMVRFixtureListItem>, FText> FixtureIDConflicts = GetFixtureIDConflicts();
+		AppendConflictTexts(FixtureIDConflicts, AccumulatedConflicts);
 
 		TMap<TSharedPtr<FDMXMVRFixtureListItem>, FText> ChannelExcessConflicts = GetChannelExcessConflicts();
 		AppendConflictTexts(ChannelExcessConflicts, AccumulatedConflicts);
@@ -75,7 +78,7 @@ private:
 
 	}
 
-	/** The patch of an item */
+	/** The patch of an item. Useful to Get Conflicts with Other */
 	struct FItemPatch
 	{
 		FItemPatch(const TSharedPtr<FDMXMVRFixtureListItem>& InItem)
@@ -85,6 +88,7 @@ private:
 			AddressRange = TRange<int32>(Item->GetAddress(), Item->GetAddress() + Item->GetNumChannels());
 		}
 
+		/** Returns a conflict text if this item conflicts with Other */
 		FText GetConfictsWithOther(const FItemPatch& Other) const
 		{
 			// No conflict with self
@@ -119,31 +123,30 @@ private:
 				return FText::GetEmpty();
 			}
 
-			const FText UnitNumberText = MakeBeautifulItemText(Item);
-			const FText OtherUnitNumberText = MakeBeautifulItemText(Other.Item);
+			const FText FixtureIDText = MakeBeautifulItemText(Item);
+			const FText OtherFixtureIDText = MakeBeautifulItemText(Other.Item);
 			if (AddressRange.GetLowerBound() == Other.AddressRange.GetLowerBound() &&
 				Item->GetFixtureType() == Other.Item->GetFixtureType())
 			{
 				// Modes confict
 				check(Item->GetModeIndex() != Other.Item->GetModeIndex());
-				return FText::Format(LOCTEXT("ModeConflict", "Uses same Address and Fixture Type as Fixture {1}, but Modes differ."), UnitNumberText, OtherUnitNumberText);
+				return FText::Format(LOCTEXT("ModeConflict", "Uses same Address and Fixture Type as Fixture {1}, but Modes differ."), FixtureIDText, OtherFixtureIDText);
 			}
 			else if (AddressRange.GetLowerBound() == Other.AddressRange.GetLowerBound())
 			{
 				// Fixture Types conflict
 				check(Item->GetFixtureType() != Other.Item->GetFixtureType());
-				return FText::Format(LOCTEXT("FixtureTypeConflict", "Uses same Address as Fixture {1}, but Fixture Types differ."), UnitNumberText, OtherUnitNumberText);
+				return FText::Format(LOCTEXT("FixtureTypeConflict", "Uses same Address as Fixture {1}, but Fixture Types differ."), FixtureIDText, OtherFixtureIDText);
 			}
 			else
 			{
 				// Addresses conflict
-				return FText::Format(LOCTEXT("AddressConflict", "Overlaps Addresses with Fixture {1}"), UnitNumberText, OtherUnitNumberText);
+				return FText::Format(LOCTEXT("AddressConflict", "Overlaps Addresses with Fixture {1}"), FixtureIDText, OtherFixtureIDText);
 			}
 		}
 
 		FORCEINLINE const TSharedPtr<FDMXMVRFixtureListItem>& GetItem() const { return Item; };
 
-
 	private:
 		int32 Universe = -1;
 		TRange<int32> AddressRange;
@@ -240,35 +243,50 @@ private:
 		return ItemToConflictMap;
 	}
 
-	/** Returns an Map of Items to Unit Number conflict Texts */
-	TMap<TSharedPtr<FDMXMVRFixtureListItem>, FText> GetUnitNumberConflicts() const
+	/** Returns an Map of Items to Fixture IDs issues Texts */
+	TMap<TSharedPtr<FDMXMVRFixtureListItem>, FText> GetFixtureIDIssues() const
+	{
+		TMap<TSharedPtr<FDMXMVRFixtureListItem>, FText> Result;
+		for (const TSharedPtr<FDMXMVRFixtureListItem>& Item : Items)
+		{
+			int32 FixtureIDNumerical;
+			if (!LexTryParseString(FixtureIDNumerical, *Item->GetFixtureID()))
+			{
+				Result.Add(Item, LOCTEXT("FixtureIDNotNumericalIssueText", "FID has to be a number."));
+			}
+		}
+		return Result;
+	}
+
+	/** Returns an Map of Items to Fixture IDs conflict Texts */
+	TMap<TSharedPtr<FDMXMVRFixtureListItem>, FText> GetFixtureIDConflicts() const
 	{
-		TMap<int32, TArray<TSharedPtr<FDMXMVRFixtureListItem>>> UnitNumberMap;
-		UnitNumberMap.Reserve(Items.Num());
+		TMap<FString, TArray<TSharedPtr<FDMXMVRFixtureListItem>>> FixtureIDMap;
+		FixtureIDMap.Reserve(Items.Num());
 		for (const TSharedPtr<FDMXMVRFixtureListItem>& Item : Items)
 		{
-			UnitNumberMap.FindOrAdd(Item->GetUnitNumber()).Add(Item);
+			FixtureIDMap.FindOrAdd(Item->GetFixtureID()).Add(Item);
 		}
-		TArray<TArray<TSharedPtr<FDMXMVRFixtureListItem>>> UnitNumberConflicts;
-		UnitNumberMap.GenerateValueArray(UnitNumberConflicts);
-		UnitNumberConflicts.RemoveAll([](const TArray<TSharedPtr<FDMXMVRFixtureListItem>>& ConflictingItems)
+		TArray<TArray<TSharedPtr<FDMXMVRFixtureListItem>>> FixtureIDConflicts;
+		FixtureIDMap.GenerateValueArray(FixtureIDConflicts);
+		FixtureIDConflicts.RemoveAll([](const TArray<TSharedPtr<FDMXMVRFixtureListItem>>& ConflictingItems)
 			{
 				return ConflictingItems.Num() < 2;
 			});
 		
 		TMap<TSharedPtr<FDMXMVRFixtureListItem>, FText> ItemToConflictMap;
-		for (TArray<TSharedPtr<FDMXMVRFixtureListItem>>& ConflictingItems : UnitNumberConflicts)
+		for (TArray<TSharedPtr<FDMXMVRFixtureListItem>>& ConflictingItems : FixtureIDConflicts)
 		{
 			ConflictingItems.Sort([](const TSharedPtr<FDMXMVRFixtureListItem>& ItemA, const TSharedPtr<FDMXMVRFixtureListItem>& ItemB)
 				{
-					return ItemA->GetUnitNumber() < ItemB->GetUnitNumber();
+					return ItemA->GetFixtureID() < ItemB->GetFixtureID();
 				});
 
 			check(ConflictingItems.Num() > 0);
-			FText ConflictText = FText::Format(LOCTEXT("BaseUnitNumberConflictText", "Ambiguous IDs in {0}"), MakeBeautifulItemText(ConflictingItems[0]));
+			FText ConflictText = FText::Format(LOCTEXT("BaseFixtureIDConflictText", "Ambiguous FIDs in {0}"), MakeBeautifulItemText(ConflictingItems[0]));
 			for (int32 ConflictingItemIndex = 1; ConflictingItemIndex < ConflictingItems.Num(); ConflictingItemIndex++)
 			{
-				ConflictText = FText::Format(LOCTEXT("AppendUnitNumberConflictText", "{0}, {1}"), ConflictText, MakeBeautifulItemText(ConflictingItems[ConflictingItemIndex]));
+				ConflictText = FText::Format(LOCTEXT("AppendFixtureIDConflictText", "{0}, {1}"), ConflictText, MakeBeautifulItemText(ConflictingItems[ConflictingItemIndex]));
 			}
 			
 			for (const TSharedPtr<FDMXMVRFixtureListItem>& Item : ConflictingItems)
@@ -295,7 +313,7 @@ private:
 
 const FName FDMXMVRFixtureListCollumnIDs::Status = "Status";
 const FName FDMXMVRFixtureListCollumnIDs::FixturePatchName = "FixturePatchName";
-const FName FDMXMVRFixtureListCollumnIDs::UnitNumber = "UnitNumber";
+const FName FDMXMVRFixtureListCollumnIDs::FixtureID = "FixtureID";
 const FName FDMXMVRFixtureListCollumnIDs::MVRFixtureName = "MVRFixtureName";
 const FName FDMXMVRFixtureListCollumnIDs::FixtureType = "FixtureType";
 const FName FDMXMVRFixtureListCollumnIDs::Mode = "Mode";
@@ -738,10 +756,10 @@ TSharedRef<SHeaderRow> SDMXMVRFixtureList::GenerateHeaderRow()
 
 	HeaderRow->AddColumn(
 		SHeaderRow::FColumn::FArguments()
-		.ColumnId(FDMXMVRFixtureListCollumnIDs::UnitNumber)
-		.SortMode(this, &SDMXMVRFixtureList::GetColumnSortMode, FDMXMVRFixtureListCollumnIDs::UnitNumber)
+		.ColumnId(FDMXMVRFixtureListCollumnIDs::FixtureID)
+		.SortMode(this, &SDMXMVRFixtureList::GetColumnSortMode, FDMXMVRFixtureListCollumnIDs::FixtureID)
 		.OnSort(this, &SDMXMVRFixtureList::SortByColumnID)
-		.DefaultLabel(LOCTEXT("UnitNumberColumnLabel", "ID"))
+		.DefaultLabel(LOCTEXT("FixtureIDColumnLabel", "FID"))
 		.FillWidth(0.1f)
 	);
 
@@ -788,9 +806,9 @@ void SDMXMVRFixtureList::SaveHeaderRowSettings()
 	{
 		for (const SHeaderRow::FColumn& Column : HeaderRow->GetColumns())
 		{
-			if (Column.ColumnId == FDMXMVRFixtureListCollumnIDs::UnitNumber)
+			if (Column.ColumnId == FDMXMVRFixtureListCollumnIDs::FixtureID)
 			{
-				EditorSettings->MVRFixtureListSettings.UnitNumberColumnWidth = Column.Width.Get();
+				EditorSettings->MVRFixtureListSettings.FixtureIDColumnWidth = Column.Width.Get();
 			}
 			if (Column.ColumnId == FDMXMVRFixtureListCollumnIDs::MVRFixtureName)
 			{
@@ -818,10 +836,10 @@ void SDMXMVRFixtureList::RestoresHeaderRowSettings()
 {
 	if (const UDMXEditorSettings* EditorSettings = GetDefault<UDMXEditorSettings>())
 	{
-		const float UnitNumberColumnWidth = EditorSettings->MVRFixtureListSettings.UnitNumberColumnWidth;
-		if (UnitNumberColumnWidth > 20.f)
+		const float FixtureIDColumnWidth = EditorSettings->MVRFixtureListSettings.FixtureIDColumnWidth;
+		if (FixtureIDColumnWidth > 20.f)
 		{
-			HeaderRow->SetColumnWidth(FDMXMVRFixtureListCollumnIDs::UnitNumber, UnitNumberColumnWidth);
+			HeaderRow->SetColumnWidth(FDMXMVRFixtureListCollumnIDs::FixtureID, FixtureIDColumnWidth);
 		}
 
 		const float NameColumnWidth = EditorSettings->MVRFixtureListSettings.NameColumnWidth;
@@ -865,20 +883,20 @@ void SDMXMVRFixtureList::SortByColumnID(const EColumnSortPriority::Type SortPrio
 	SortMode = InSortMode;
 	SortedByColumnID = ColumnId;
 
-	if (ColumnId == FDMXMVRFixtureListCollumnIDs::UnitNumber)
+	if (ColumnId == FDMXMVRFixtureListCollumnIDs::FixtureID)
 	{
 		if (InSortMode == EColumnSortMode::Descending)
 		{
 			ListSource.Sort([](const TSharedPtr<FDMXMVRFixtureListItem>& ItemA, const TSharedPtr<FDMXMVRFixtureListItem>& ItemB)
 				{
-					return !(ItemA->GetUnitNumber() <= ItemB->GetUnitNumber());
+					return !(ItemA->GetFixtureID() <= ItemB->GetFixtureID());
 				});
 		}
 		else if (InSortMode == EColumnSortMode::Ascending)
 		{
 			ListSource.Sort([](const TSharedPtr<FDMXMVRFixtureListItem>& ItemA, const TSharedPtr<FDMXMVRFixtureListItem>& ItemB)
 				{
-					return !(ItemA->GetUnitNumber() >= ItemB->GetUnitNumber());
+					return !(ItemA->GetFixtureID() >= ItemB->GetFixtureID());
 				});
 		}
 	}
diff --git a/Engine/Plugins/VirtualProduction/DMX/DMXEngine/Source/DMXEditor/Private/Widgets/FixturePatch/SDMXMVRFixtureList.h b/Engine/Plugins/VirtualProduction/DMX/DMXEngine/Source/DMXEditor/Private/Widgets/FixturePatch/SDMXMVRFixtureList.h
--- a/Engine/Plugins/VirtualProduction/DMX/DMXEngine/Source/DMXEditor/Private/Widgets/FixturePatch/SDMXMVRFixtureList.h
+++ b/Engine/Plugins/VirtualProduction/DMX/DMXEngine/Source/DMXEditor/Private/Widgets/FixturePatch/SDMXMVRFixtureList.h
@@ -29,7 +29,7 @@ struct FDMXMVRFixtureListCollumnIDs
 {
 	static const FName Status;
 	static const FName FixturePatchName;
-	static const FName UnitNumber;
+	static const FName FixtureID;
 	static const FName MVRFixtureName;
 	static const FName FixtureType;
 	static const FName Mode;
diff --git a/Engine/Plugins/VirtualProduction/DMX/DMXEngine/Source/DMXEditor/Private/Widgets/FixturePatch/SDMXMVRFixtureListRow.cpp b/Engine/Plugins/VirtualProduction/DMX/DMXEngine/Source/DMXEditor/Private/Widgets/FixturePatch/SDMXMVRFixtureListRow.cpp
--- a/Engine/Plugins/VirtualProduction/DMX/DMXEngine/Source/DMXEditor/Private/Widgets/FixturePatch/SDMXMVRFixtureListRow.cpp
+++ b/Engine/Plugins/VirtualProduction/DMX/DMXEngine/Source/DMXEditor/Private/Widgets/FixturePatch/SDMXMVRFixtureListRow.cpp
@@ -380,9 +380,9 @@ TSharedRef<SWidget> SDMXMVRFixtureListRow::GenerateWidgetForColumn(const FName&
 	{
 		return GenerateStatusRow();
 	}
-	else if (ColumnName == FDMXMVRFixtureListCollumnIDs::UnitNumber)
+	else if (ColumnName == FDMXMVRFixtureListCollumnIDs::FixtureID)
 	{
-		return GenerateUnitNumberRow();
+		return GenerateFixtureIDRow();
 	}
 	else if (ColumnName == FDMXMVRFixtureListCollumnIDs::MVRFixtureName)
 	{
@@ -496,7 +496,7 @@ TSharedRef<SWidget> SDMXMVRFixtureListRow::GenerateStatusRow()
 		];
 }
 
-TSharedRef<SWidget> SDMXMVRFixtureListRow::GenerateUnitNumberRow()
+TSharedRef<SWidget> SDMXMVRFixtureListRow::GenerateFixtureIDRow()
 {
 	return
 		SNew(SBorder)
@@ -504,49 +504,48 @@ TSharedRef<SWidget> SDMXMVRFixtureListRow::GenerateUnitNumberRow()
 		.VAlign(VAlign_Center)
 		.Padding(4.f)
 		.BorderImage(FAppStyle::GetBrush("NoBorder"))
-		.OnMouseDoubleClick(this, &SDMXMVRFixtureListRow::OnUnitNumberBorderDoubleClicked)
+		.OnMouseDoubleClick(this, &SDMXMVRFixtureListRow::OnFixtureIDBorderDoubleClicked)
 		[
 			SNew(SBorder)
 			.HAlign(HAlign_Fill)
 			.VAlign(VAlign_Center)
 			.BorderImage(FDMXEditorStyle::Get().GetBrush("DMXEditor.RoundedPropertyBorder"))
 			[
-				SAssignNew(UnitNumberTextBlocK, SInlineEditableTextBlock)
+				SAssignNew(FixtureIDTextBlocK, SInlineEditableTextBlock)
 				.Text_Lambda([this]()
 				{
-					const int32 UnitNumber = Item->GetUnitNumber();
-					return FText::AsNumber(UnitNumber);
+					return FText::FromString(Item->GetFixtureID());
 				})
 				.Font(FCoreStyle::GetDefaultFontStyle("Regular", 10))
-				.OnTextCommitted(this, &SDMXMVRFixtureListRow::OnUnitNumberCommitted)
+				.OnTextCommitted(this, &SDMXMVRFixtureListRow::OnFixtureIDCommitted)
 				.IsSelected(IsSelected)
 			]
 		];
 }
 
-FReply SDMXMVRFixtureListRow::OnUnitNumberBorderDoubleClicked(const FGeometry& InMyGeometry, const FPointerEvent& InMouseEvent)
+FReply SDMXMVRFixtureListRow::OnFixtureIDBorderDoubleClicked(const FGeometry& InMyGeometry, const FPointerEvent& InMouseEvent)
 {
 	if (InMouseEvent.GetEffectingButton() == EKeys::LeftMouseButton)
 	{
-		if (UnitNumberTextBlocK.IsValid())
+		if (FixtureIDTextBlocK.IsValid())
 		{
-			UnitNumberTextBlocK->EnterEditingMode();
+			FixtureIDTextBlocK->EnterEditingMode();
 		}
 	}
 
 	return FReply::Handled();
 }
 
-void SDMXMVRFixtureListRow::OnUnitNumberCommitted(const FText& InNewText, ETextCommit::Type InTextCommit)
+void SDMXMVRFixtureListRow::OnFixtureIDCommitted(const FText& InNewText, ETextCommit::Type InTextCommit)
 {
 	const FString StringValue = InNewText.ToString();
-	int32 NewUnitNumber;
-	if (LexTryParseString<int32>(NewUnitNumber, *StringValue))
+	int32 NewFixtureID;
+	if (LexTryParseString<int32>(NewFixtureID, *StringValue))
 	{
-		Item->SetUnitNumber(NewUnitNumber);
+		Item->SetFixtureID(NewFixtureID);
 
-		const FString ParsedUnitNumberString = FString::FromInt(NewUnitNumber);
-		UnitNumberTextBlocK->SetText(FText::FromString(ParsedUnitNumberString));
+		const FString ParsedFixtureIDString = FString::FromInt(NewFixtureID);
+		FixtureIDTextBlocK->SetText(FText::FromString(ParsedFixtureIDString));
 
 		OnRowRequestsStatusRefresh.ExecuteIfBound();
 	}
diff --git a/Engine/Plugins/VirtualProduction/DMX/DMXEngine/Source/DMXEditor/Private/Widgets/FixturePatch/SDMXMVRFixtureListRow.h b/Engine/Plugins/VirtualProduction/DMX/DMXEngine/Source/DMXEditor/Private/Widgets/FixturePatch/SDMXMVRFixtureListRow.h
--- a/Engine/Plugins/VirtualProduction/DMX/DMXEngine/Source/DMXEditor/Private/Widgets/FixturePatch/SDMXMVRFixtureListRow.h
+++ b/Engine/Plugins/VirtualProduction/DMX/DMXEngine/Source/DMXEditor/Private/Widgets/FixturePatch/SDMXMVRFixtureListRow.h
@@ -59,14 +59,14 @@ private:
 	/** Generates the row that displays the Status */
 	TSharedRef<SWidget> GenerateStatusRow();
 
-	/** Generates the row that displays the Unit Number */
-	TSharedRef<SWidget> GenerateUnitNumberRow();
+	/** Generates the row that displays the Fixture ID */
+	TSharedRef<SWidget> GenerateFixtureIDRow();
 
-	/** Called when the Unit Number Border was double-clicked */
-	FReply OnUnitNumberBorderDoubleClicked(const FGeometry& InMyGeometry, const FPointerEvent& InMouseEvent);
+	/** Called when the Fixture ID Border was double-clicked */
+	FReply OnFixtureIDBorderDoubleClicked(const FGeometry& InMyGeometry, const FPointerEvent& InMouseEvent);
 
-	/** Called when a Unit Number was committed */
-	void OnUnitNumberCommitted(const FText& InNewText, ETextCommit::Type InTextCommit);
+	/** Called when a Fixture ID was committed */
+	void OnFixtureIDCommitted(const FText& InNewText, ETextCommit::Type InTextCommit);
 
 	/** Generates the row that displays the MVR Fixture Name */
 	TSharedRef<SWidget> GenerateMVRFixtureNameRow();
@@ -107,8 +107,8 @@ private:
 	/** The text block to edit the Fixture Patch Name */
 	TSharedPtr<SInlineEditableTextBlock> FixturePatchNameTextBlock;
 
-	/** The text block to edit the Unit Number */
-	TSharedPtr<SInlineEditableTextBlock> UnitNumberTextBlocK;
+	/** The text block to edit the Fixture ID */
+	TSharedPtr<SInlineEditableTextBlock> FixtureIDTextBlocK;
 
 	/** The text block to edit the Name */
 	TSharedPtr<SInlineEditableTextBlock> NameTextBlock;
diff --git a/Engine/Plugins/VirtualProduction/DMX/DMXEngine/Source/DMXEditor/Private/Widgets/FixturePatch/SDMXMVRFixtureListToolbar.cpp b/Engine/Plugins/VirtualProduction/DMX/DMXEngine/Source/DMXEditor/Private/Widgets/FixturePatch/SDMXMVRFixtureListToolbar.cpp
--- a/Engine/Plugins/VirtualProduction/DMX/DMXEngine/Source/DMXEditor/Private/Widgets/FixturePatch/SDMXMVRFixtureListToolbar.cpp
+++ b/Engine/Plugins/VirtualProduction/DMX/DMXEngine/Source/DMXEditor/Private/Widgets/FixturePatch/SDMXMVRFixtureListToolbar.cpp
@@ -107,14 +107,14 @@ namespace UE::DMXRuntime::SDMXMVRFixtureListToolbar::Private
 		return false;
 	}
 
-	UE_NODISCARD bool ParseUnitNumber(const FString& InputString, int32& OutUnitNumber)
+	UE_NODISCARD bool ParseFixtureID(const FString& InputString, int32& OutFixtureID)
 	{
-		if (LexTryParseString<int32>(OutUnitNumber, *InputString))
+		if (LexTryParseString<int32>(OutFixtureID, *InputString))
 		{
 			return true;
 		}
 
-		OutUnitNumber = -1;
+		OutFixtureID = -1;
 		return false;
 	}
 }
@@ -235,18 +235,23 @@ TArray<TSharedPtr<FDMXMVRFixtureListItem>> SDMXMVRFixtureListToolbar::FilterItem
 		return Result;
 	}
 	
-	int32 UnitNumber;
-	if (ParseUnitNumber(SearchString, UnitNumber))
+	int32 FixtureIDNumerical;
+	if (ParseFixtureID(SearchString, FixtureIDNumerical))
 	{
-		TArray<TSharedPtr<FDMXMVRFixtureListItem>> UnitNumbersOnlyResult = Result;
-		UnitNumbersOnlyResult.RemoveAll([UnitNumber](const TSharedPtr<FDMXMVRFixtureListItem>& Item)
+		TArray<TSharedPtr<FDMXMVRFixtureListItem>> FixtureIDsOnlyResult = Result;
+		FixtureIDsOnlyResult.RemoveAll([FixtureIDNumerical](const TSharedPtr<FDMXMVRFixtureListItem>& Item)
 			{
-				return Item->GetUnitNumber() != UnitNumber;
+				int32 OtherFixtureIDNumerical;
+				if (ParseFixtureID(Item->GetFixtureID(), OtherFixtureIDNumerical))
+				{
+					return OtherFixtureIDNumerical != FixtureIDNumerical;
+				}
+				return true;
 			});
 
-		if (UnitNumbersOnlyResult.Num() > 0)
+		if (FixtureIDsOnlyResult.Num() > 0)
 		{
-			return UnitNumbersOnlyResult;
+			return FixtureIDsOnlyResult;
 		}
 	}
 
diff --git a/Engine/Plugins/VirtualProduction/DMX/DMXEngine/Source/DMXEditor/Public/DMXEditorSettings.h b/Engine/Plugins/VirtualProduction/DMX/DMXEngine/Source/DMXEditor/Public/DMXEditorSettings.h
--- a/Engine/Plugins/VirtualProduction/DMX/DMXEngine/Source/DMXEditor/Public/DMXEditorSettings.h
+++ b/Engine/Plugins/VirtualProduction/DMX/DMXEngine/Source/DMXEditor/Public/DMXEditorSettings.h
@@ -15,9 +15,9 @@ struct FDMXMVRFixtureListSettings
 {
 	GENERATED_BODY()
 
-	/** Width of the Unit Number column */
+	/** Width of the Fixture ID column */
 	UPROPERTY()
-	float UnitNumberColumnWidth = 0.f;
+	float FixtureIDColumnWidth = 0.f;
 
 	/** Width of the Name column */
 	UPROPERTY()
diff --git a/Engine/Plugins/VirtualProduction/DMX/DMXEngine/Source/DMXRuntime/Private/MVR/DMXMVRGeneralSceneDescription.cpp b/Engine/Plugins/VirtualProduction/DMX/DMXEngine/Source/DMXRuntime/Private/MVR/DMXMVRGeneralSceneDescription.cpp
--- a/Engine/Plugins/VirtualProduction/DMX/DMXEngine/Source/DMXRuntime/Private/MVR/DMXMVRGeneralSceneDescription.cpp
+++ b/Engine/Plugins/VirtualProduction/DMX/DMXEngine/Source/DMXRuntime/Private/MVR/DMXMVRGeneralSceneDescription.cpp
@@ -144,14 +144,12 @@ void UDMXMVRGeneralSceneDescription::WriteFixturePatchToGeneralSceneDescription(
 		UDMXMVRChildListNode& AnyChildList = RootNode->GetOrCreateFirstChildListNode();
 		MVRFixtureNode = AnyChildList.CreateParametricObject<UDMXMVRFixtureNode>();
 
-		const TArray<int32> UnitNumbersInUse = GetUnitNumbersInUse(*DMXLibrary);
-		const int32 UnitNumber = UnitNumbersInUse.Num() > 0 ? UnitNumbersInUse.Last() + 1 : 1;
+		const TArray<int32> FixtureIDsInUse = GetNumericalFixtureIDsInUse(*DMXLibrary);
+		const int32 FixtureID = FixtureIDsInUse.Num() > 0 ? FixtureIDsInUse.Last() + 1 : 1;
 
 		MVRFixtureNode->Name = FixturePatch.Name;
 		MVRFixtureNode->UUID = MVRFixtureUUID;
-
-		// Create a new Unit Number and add it to the array of Unit Numbers in use
-		MVRFixtureNode->UnitNumber = UnitNumber;
+		MVRFixtureNode->FixtureID = FString::FromInt(FixtureID);
 	}
 	check(MVRFixtureNode);
 
@@ -200,7 +198,7 @@ void UDMXMVRGeneralSceneDescription::ParseGeneralSceneDescriptionXml(const TShar
 }
 #endif // WITH_EDITOR
 
-TArray<int32> UDMXMVRGeneralSceneDescription::GetUnitNumbersInUse(const UDMXLibrary& DMXLibrary)
+TArray<int32> UDMXMVRGeneralSceneDescription::GetNumericalFixtureIDsInUse(const UDMXLibrary& DMXLibrary)
 {
 	checkf(RootNode, TEXT("Unexpected: MVR General Scene Description Root Node is invalid."));
 
@@ -212,22 +210,26 @@ TArray<int32> UDMXMVRGeneralSceneDescription::GetUnitNumbersInUse(const UDMXLibr
 		MVRFixtureUUIDsInUse.Add(FixturePatch->GetMVRFixtureUUID());
 	}
 
-	TArray<int32> UnitNumbersInUse;
+	TArray<int32> FixtureIDsInUse;
 	for (const FGuid& MVRFixtureUUID : MVRFixtureUUIDsInUse)
 	{
 		if (UDMXMVRParametricObjectNodeBase** ObjectNodePtr = RootNode->FindParametricObjectNodeByUUID(MVRFixtureUUID))
 		{
 			if (UDMXMVRFixtureNode* FixtureNode = Cast<UDMXMVRFixtureNode>(*ObjectNodePtr))
 			{
-				UnitNumbersInUse.Add(FixtureNode->UnitNumber);
+				int32 FixtureIDNumerical;
+				if (LexTryParseString(FixtureIDNumerical, *FixtureNode->FixtureID))
+				{
+					FixtureIDsInUse.Add(FixtureIDNumerical);
+				}
 			}
 		}
 	}
 
-	UnitNumbersInUse.Sort([](int32 UnitNumberA, int32 UnitNumberB)
+	FixtureIDsInUse.Sort([](int32 FixtureIDA, int32 FixtureIDB)
 		{
-			return UnitNumberA < UnitNumberB;
+			return FixtureIDA < FixtureIDB;
 		});
 
-	return UnitNumbersInUse;
+	return FixtureIDsInUse;
 }
diff --git a/Engine/Plugins/VirtualProduction/DMX/DMXEngine/Source/DMXRuntime/Private/MVR/Types/DMXMVRFixtureNode.cpp b/Engine/Plugins/VirtualProduction/DMX/DMXEngine/Source/DMXRuntime/Private/MVR/Types/DMXMVRFixtureNode.cpp
--- a/Engine/Plugins/VirtualProduction/DMX/DMXEngine/Source/DMXRuntime/Private/MVR/Types/DMXMVRFixtureNode.cpp
+++ b/Engine/Plugins/VirtualProduction/DMX/DMXEngine/Source/DMXRuntime/Private/MVR/Types/DMXMVRFixtureNode.cpp
@@ -132,11 +132,11 @@ namespace UE::DMXRuntime::DMXMVRFixture::Private
 	}
 
 	/** Parses the Fixture Id from the MVR Fixture Node, returns true on success. */
-	bool ParseFixtureId(const FXmlNode& InFixtureNode, FString& OutFixtureId)
+	bool ParseFixtureID(const FXmlNode& InFixtureNode, FString& OutFixtureId)
 	{
 		OutFixtureId.Reset();
 
-		constexpr TCHAR FixtureIdTag[] = TEXT("FixtureId");
+		constexpr TCHAR FixtureIdTag[] = TEXT("FixtureID");
 		const FXmlNode* FixtureIdNode = InFixtureNode.FindChildNode(FixtureIdTag);
 		if (FixtureIdNode)
 		{
@@ -151,7 +151,7 @@ namespace UE::DMXRuntime::DMXMVRFixture::Private
 	bool ParseUnitNumber(const FXmlNode& InFixtureNode, int32& OutUnitNumber)
 	{
 		// @todo: Temp fix for the larger fixture id/unit number mismatch
-		constexpr TCHAR UnitNumberTag[] = TEXT("FixtureId");
+		constexpr TCHAR UnitNumberTag[] = TEXT("UnitNumber");
 		const FXmlNode* UnitNumberNode = InFixtureNode.FindChildNode(UnitNumberTag);
 		if (UnitNumberNode)
 		{
@@ -524,7 +524,7 @@ bool UDMXMVRFixtureNode::InitializeFromFixtureXmlNode(const FXmlNode& FixtureXml
 	ParseCastShadow(FixtureXmlNode, CastShadow);
 	ParsePosition(FixtureXmlNode, Position);
 	
-	if (!ParseFixtureId(FixtureXmlNode, FixtureID))
+	if (!ParseFixtureID(FixtureXmlNode, FixtureID))
 	{
 		// Non-optional, mend with empty string
 		UE_LOG(LogDMXRuntime, Warning, TEXT("Imported Fixture '%s' from MVR. No valid Fixture Id specified. Fixture Id set to Fixture Name."), *Name);
diff --git a/Engine/Plugins/VirtualProduction/DMX/DMXEngine/Source/DMXRuntime/Public/MVR/DMXMVRGeneralSceneDescription.h b/Engine/Plugins/VirtualProduction/DMX/DMXEngine/Source/DMXRuntime/Public/MVR/DMXMVRGeneralSceneDescription.h
--- a/Engine/Plugins/VirtualProduction/DMX/DMXEngine/Source/DMXRuntime/Public/MVR/DMXMVRGeneralSceneDescription.h
+++ b/Engine/Plugins/VirtualProduction/DMX/DMXEngine/Source/DMXRuntime/Public/MVR/DMXMVRGeneralSceneDescription.h
@@ -63,8 +63,8 @@ private:
 	void ParseGeneralSceneDescriptionXml(const TSharedRef<FXmlFile>& GeneralSceneDescriptionXml);
 #endif
 
-	/** Returns the Unit Numbers currently in use, sorted from lowest to highest */
-	TArray<int32> GetUnitNumbersInUse(const UDMXLibrary& DMXLibrary);
+	/** Returns the Fixture IDs currently in use, as number, sorted from lowest to highest */
+	TArray<int32> GetNumericalFixtureIDsInUse(const UDMXLibrary& DMXLibrary);
 
 	/** The Root Node of the General Scene Description */
 	UPROPERTY()
