DMX - Fix pixelmapping distribution sort does not distribute correctly

diff --git a/Engine/Plugins/VirtualProduction/DMX/DMXEngine/Source/DMXRuntime/Public/DMXRuntimeUtils.h b/Engine/Plugins/VirtualProduction/DMX/DMXEngine/Source/DMXRuntime/Public/DMXRuntimeUtils.h
--- a/Engine/Plugins/VirtualProduction/DMX/DMXEngine/Source/DMXRuntime/Public/DMXRuntimeUtils.h
+++ b/Engine/Plugins/VirtualProduction/DMX/DMXEngine/Source/DMXRuntime/Public/DMXRuntimeUtils.h
@@ -79,36 +79,19 @@
 		}
 		else if (InDistribution == EDMXPixelMappingDistribution::TopLeftToBottom)
 		{
-			for (int32 XIndex = 0; XIndex < InNumXPanels; ++XIndex)
+			for (int32 YIndex = 0; YIndex < InNumYPanels; ++YIndex)
 			{
-				for (int32 YIndex = 0; YIndex < InNumYPanels; ++YIndex)
+				for (int32 XIndex = 0; XIndex < InNumXPanels; ++XIndex)
 				{
-					OutSortedList.Add(InUnorderedList[XIndex + YIndex * InNumXPanels]);
+					OutSortedList.Add(InUnorderedList[YIndex + XIndex * InNumYPanels]);
 				}
 			}
 		}
 		else if (InDistribution == EDMXPixelMappingDistribution::TopLeftToClockwise)
 		{
-			for (int32 XIndex = 0; XIndex < InNumXPanels; ++XIndex)
+			for (int32 YIndex = 0; YIndex < InNumYPanels; ++YIndex)
 			{
-				for (int32 YIndex = 0; YIndex < InNumYPanels; ++YIndex)
-				{
-					if ((XIndex % 2) == 0)
-					{
-						OutSortedList.Add(InUnorderedList[YIndex + (XIndex * InNumYPanels)]);
-					}
-					else
-					{
-						OutSortedList.Add(InUnorderedList[((XIndex + 1) * InNumYPanels) - (YIndex + 1)]);
-					}
-				}
-			}
-		}
-		else if (InDistribution == EDMXPixelMappingDistribution::TopLeftToAntiClockwise)
-		{
-			for (int32 XIndex = 0; XIndex < InNumXPanels; ++XIndex)
-			{
-				for (int32 YIndex = 0; YIndex < InNumYPanels; ++YIndex)
+				for (int32 XIndex = 0; XIndex < InNumXPanels; ++XIndex)
 				{
 					if ((YIndex % 2) == 0)
 					{
@@ -121,48 +104,48 @@
 				}
 			}
 		}
+		else if (InDistribution == EDMXPixelMappingDistribution::TopLeftToAntiClockwise)
+		{
+			for (int32 YIndex = 0; YIndex < InNumYPanels; ++YIndex)
+			{
+				for (int32 XIndex = 0; XIndex < InNumXPanels; ++XIndex)
+				{
+					if ((XIndex % 2) == 0)
+					{
+						OutSortedList.Add(InUnorderedList[YIndex + (XIndex * InNumYPanels)]);
+					}
+					else
+					{
+						OutSortedList.Add(InUnorderedList[((XIndex + 1) * InNumYPanels) - (YIndex + 1)]);
+					}
+				}
+			}
+		}
 		else if (InDistribution == EDMXPixelMappingDistribution::TopRightToLeft)
 		{
-			for (int32 XIndex = 0; XIndex < InNumXPanels; ++XIndex)
+			for (int32 YIndex = 0; YIndex < InNumYPanels; ++YIndex)
 			{
-				for (int32 YIndex = 0; YIndex < InNumYPanels; ++YIndex)
+				for (int32 XIndex = 0; XIndex < InNumXPanels; ++XIndex)
 				{
-					OutSortedList.Add(InUnorderedList[(InNumYPanels * (XIndex + 1)) - (YIndex + 1)]);
+					OutSortedList.Add(InUnorderedList[InNumXPanels - 1 - XIndex + YIndex * InNumXPanels]);
 				}
 			}
 		}
 		else if (InDistribution == EDMXPixelMappingDistribution::BottomLeftToTop)
 		{
-			for (int32 XIndex = 0; XIndex < InNumXPanels; ++XIndex)
+			for (int32 YIndex = InNumYPanels - 1; YIndex >= 0; --YIndex)
 			{
-				for (int32 YIndex = 0; YIndex < InNumYPanels; ++YIndex)
+				for (int32 XIndex = 0; XIndex < InNumXPanels; ++XIndex)
 				{
-					OutSortedList.Add(InUnorderedList[(InNumXPanels - XIndex - 1) + (YIndex * InNumXPanels)]);
+					OutSortedList.Add(InUnorderedList[YIndex + XIndex * InNumYPanels]);
 				}
 			}
 		}
 		else if (InDistribution == EDMXPixelMappingDistribution::TopRightToAntiClockwise)
 		{
-			for (int32 XIndex = 0; XIndex < InNumXPanels; ++XIndex)
+			for (int32 YIndex = 0; YIndex < InNumYPanels; ++YIndex)
 			{
-				for (int32 YIndex = 0; YIndex < InNumYPanels; ++YIndex)
-				{
-					if ((XIndex % 2) == 0)
-					{
-						OutSortedList.Add(InUnorderedList[(InNumYPanels * (XIndex + 1)) - (YIndex + 1)]);
-					}
-					else
-					{
-						OutSortedList.Add(InUnorderedList[YIndex + (XIndex * InNumYPanels)]);
-					}
-				}
-			}
-		}
-		else if (InDistribution == EDMXPixelMappingDistribution::BottomLeftToClockwise)
-		{
-			for (int32 XIndex = 0; XIndex < InNumXPanels; ++XIndex)
-			{
-				for (int32 YIndex = 0; YIndex < InNumYPanels; ++YIndex)
+				for (int32 XIndex = 0; XIndex < InNumXPanels; ++XIndex)
 				{
 					if ((YIndex % 2) == 0)
 					{
@@ -175,64 +158,49 @@
 				}
 			}
 		}
+		else if (InDistribution == EDMXPixelMappingDistribution::BottomLeftToClockwise)
+		{
+			for (int32 YIndex = 0; YIndex < InNumYPanels; ++YIndex)
+			{
+				for (int32 XIndex = 0; XIndex < InNumXPanels; ++XIndex)
+				{
+					if ((XIndex % 2) == 0)
+					{
+						OutSortedList.Add(InUnorderedList[(InNumYPanels * (XIndex + 1)) - (YIndex + 1)]);
+					}
+					else
+					{
+						OutSortedList.Add(InUnorderedList[YIndex + (XIndex * InNumYPanels)]);
+					}
+				}
+			}
+		}
 		else if (InDistribution == EDMXPixelMappingDistribution::BottomLeftToRight)
 		{
-			for (int32 XIndex = 0; XIndex < InNumXPanels; ++XIndex)
+			for (int32 YIndex = InNumYPanels - 1; YIndex >= 0; --YIndex)
 			{
-				for (int32 YIndex = 0; YIndex < InNumYPanels; ++YIndex)
+				for (int32 XIndex = 0; XIndex < InNumXPanels; ++XIndex)
 				{
-					OutSortedList.Add(InUnorderedList[((InNumXPanels - XIndex) * InNumYPanels) - (InNumYPanels - YIndex)]);
+					OutSortedList.Add(InUnorderedList[YIndex * InNumXPanels + XIndex]);
 				}
 			}
 		}
 		else if (InDistribution == EDMXPixelMappingDistribution::TopRightToBottom)
 		{
-			for (int32 XIndex = 0; XIndex < InNumXPanels; ++XIndex)
+			for (int32 YIndex = 0; YIndex < InNumYPanels; ++YIndex)
 			{
-				for (int32 YIndex = 0; YIndex < InNumYPanels; ++YIndex)
+				for (int32 XIndex = InNumXPanels - 1; XIndex >= 0; --XIndex)
 				{
-					OutSortedList.Add(InUnorderedList[((InNumYPanels - YIndex - 1) * InNumXPanels) + XIndex]);
+					OutSortedList.Add(InUnorderedList[YIndex + XIndex * InNumYPanels]);
 				}
 			}
 		}
 		else if (InDistribution == EDMXPixelMappingDistribution::BottomLeftAntiClockwise)
 		{
-			for (int32 XIndex = 0; XIndex < InNumXPanels; ++XIndex)
+			for (int32 YIndex = 0; YIndex < InNumYPanels; ++YIndex)
 			{
-				for (int32 YIndex = 0; YIndex < InNumYPanels; ++YIndex)
+				for (int32 XIndex = 0; XIndex < InNumXPanels; ++XIndex)
 				{
-					if ((InNumXPanels) % 2 == 0)
-					{
-						if ((XIndex % 2) == 0)
-						{
-							OutSortedList.Add(InUnorderedList[((InNumXPanels - XIndex) * InNumYPanels) - (YIndex + 1)]);
-						}
-						else
-						{
-							OutSortedList.Add(InUnorderedList[((InNumXPanels - XIndex) * InNumYPanels) - (InNumYPanels - YIndex)]);
-						}
-					}
-					else
-					{
-						if ((XIndex % 2) == 0)
-						{
-							OutSortedList.Add(InUnorderedList[((InNumXPanels - XIndex) * InNumYPanels) - (InNumYPanels - YIndex)]);
-						}
-						else
-						{
-							OutSortedList.Add(InUnorderedList[((InNumXPanels - XIndex) * InNumYPanels) - (YIndex + 1)]);
-						}
-					}
-
-				}
-			}
-		}
-		else if (InDistribution == EDMXPixelMappingDistribution::TopRightToClockwise)
-		{
-			for (int32 XIndex = 0; XIndex < InNumXPanels; ++XIndex)
-			{
-				for (int32 YIndex = 0; YIndex < InNumYPanels; ++YIndex)
-				{
 					if ((InNumYPanels) % 2 == 0)
 					{
 						if ((YIndex % 2) == 0)
@@ -256,3 +224,34 @@
 						}
 					}
 				}
+			}
+		}
+		else if (InDistribution == EDMXPixelMappingDistribution::TopRightToClockwise)
+		{
+			for (int32 YIndex = 0; YIndex < InNumYPanels; ++YIndex)
+			{
+				for (int32 XIndex = 0; XIndex < InNumXPanels; ++XIndex)
+				{
+					if ((InNumXPanels) % 2 == 0)
+					{
+						if ((XIndex % 2) == 0)
+						{
+							OutSortedList.Add(InUnorderedList[((InNumXPanels - XIndex) * InNumYPanels) - (YIndex + 1)]);
+						}
+						else
+						{
+							OutSortedList.Add(InUnorderedList[((InNumXPanels - XIndex) * InNumYPanels) - (InNumYPanels - YIndex)]);
+						}
+					}
+					else
+					{
+						if ((XIndex % 2) == 0)
+						{
+							OutSortedList.Add(InUnorderedList[((InNumXPanels - XIndex) * InNumYPanels) - (InNumYPanels - YIndex)]);
+						}
+						else
+						{
+							OutSortedList.Add(InUnorderedList[((InNumXPanels - XIndex) * InNumYPanels) - (YIndex + 1)]);
+						}
+					}
+				}
@@ -260,50 +259,50 @@
 		}
 		else if (InDistribution == EDMXPixelMappingDistribution::BottomRightToLeft)
 		{
-			for (int32 XIndex = 0; XIndex < InNumXPanels; ++XIndex)
+			for (int32 YIndex = 0; YIndex < InNumYPanels; ++YIndex)
 			{
-				for (int32 YIndex = 0; YIndex < InNumYPanels; ++YIndex)
+				for (int32 XIndex = 0; XIndex < InNumXPanels; ++XIndex)
 				{
-					OutSortedList.Add(InUnorderedList[((InNumXPanels - XIndex) * InNumYPanels) - (YIndex + 1)]);
+					OutSortedList.Add(InUnorderedList[InNumXPanels * InNumYPanels - 1 - (XIndex + YIndex * InNumXPanels)]);
 				}
 			}
 		}
 		else if (InDistribution == EDMXPixelMappingDistribution::BottomRightToTop)
 		{
-			for (int32 XIndex = 0; XIndex < InNumXPanels; ++XIndex)
+			for (int32 YIndex = InNumYPanels - 1; YIndex >= 0; --YIndex)
 			{
-				for (int32 YIndex = 0; YIndex < InNumYPanels; ++YIndex)
+				for (int32 XIndex = InNumXPanels - 1; XIndex >= 0; --XIndex)
 				{
-					OutSortedList.Add(InUnorderedList[(InNumXPanels * InNumYPanels) - (YIndex * InNumXPanels) - (XIndex + 1)]);
+					OutSortedList.Add(InUnorderedList[YIndex + XIndex * InNumYPanels]);
 				}
 			}
 		}
 		else if (InDistribution == EDMXPixelMappingDistribution::BottomRightToClockwise)
 		{
-			for (int32 XIndex = 0; XIndex < InNumXPanels; ++XIndex)
+			for (int32 YIndex = 0; YIndex < InNumYPanels; ++YIndex)
 			{
-				for (int32 YIndex = 0; YIndex < InNumYPanels; ++YIndex)
+				for (int32 XIndex = 0; XIndex < InNumXPanels; ++XIndex)
 				{
-					if ((InNumXPanels) % 2 == 0)
+					if ((InNumYPanels) % 2 == 0)
 					{
-						if ((XIndex % 2) == 0)
+						if ((YIndex % 2) == 0)
 						{
-							OutSortedList.Add(InUnorderedList[(InNumXPanels * InNumYPanels) - ((XIndex + 1) * InNumYPanels) + YIndex]);
+							OutSortedList.Add(InUnorderedList[(InNumXPanels * InNumYPanels) - ((1 + YIndex) * InNumXPanels) + XIndex]);
 						}
 						else
 						{
-							OutSortedList.Add(InUnorderedList[(InNumYPanels * (InNumXPanels - XIndex)) - (YIndex + 1)]);
+							OutSortedList.Add(InUnorderedList[(InNumXPanels * InNumYPanels) - (YIndex * InNumXPanels) - (1 + XIndex)]);
 						}
 					}
 					else
 					{
-						if ((XIndex % 2) == 0)
+						if ((YIndex % 2) == 0)
 						{
-							OutSortedList.Add(InUnorderedList[(InNumYPanels * (InNumXPanels - XIndex)) - (YIndex + 1)]);
+							OutSortedList.Add(InUnorderedList[(InNumXPanels * InNumYPanels) - (YIndex * InNumXPanels) - (1 + XIndex)]);
 						}
 						else
 						{
-							OutSortedList.Add(InUnorderedList[(InNumXPanels * InNumYPanels) - ((XIndex + 1) * InNumYPanels) + YIndex]);
+							OutSortedList.Add(InUnorderedList[(InNumXPanels * InNumYPanels) - ((1 + YIndex) * InNumXPanels) + XIndex]);
 						}
 					}
 				}
@@ -311,30 +310,30 @@
 		}
 		else if (InDistribution == EDMXPixelMappingDistribution::BottomRightToAntiClockwise)
 		{
-			for (int32 XIndex = 0; XIndex < InNumXPanels; ++XIndex)
+			for (int32 YIndex = 0; YIndex < InNumYPanels; ++YIndex)
 			{
-				for (int32 YIndex = 0; YIndex < InNumYPanels; ++YIndex)
+				for (int32 XIndex = 0; XIndex < InNumXPanels; ++XIndex)
 				{
-					if ((InNumYPanels) % 2 == 0)
+					if ((InNumXPanels) % 2 == 0)
 					{
-						if ((YIndex % 2) == 0)
+						if ((XIndex % 2) == 0)
 						{
-							OutSortedList.Add(InUnorderedList[(InNumXPanels * InNumYPanels) - ((1 + YIndex) * InNumXPanels) + XIndex]);
+							OutSortedList.Add(InUnorderedList[(InNumXPanels * InNumYPanels) - ((XIndex + 1) * InNumYPanels) + YIndex]);
 						}
 						else
 						{
-							OutSortedList.Add(InUnorderedList[(InNumXPanels * InNumYPanels) - (YIndex * InNumXPanels) - (1 + XIndex)]);
+							OutSortedList.Add(InUnorderedList[(InNumYPanels * (InNumXPanels - XIndex)) - (YIndex + 1)]);
 						}
 					}
 					else
 					{
-						if ((YIndex % 2) == 0)
+						if ((XIndex % 2) == 0)
 						{
-							OutSortedList.Add(InUnorderedList[(InNumXPanels * InNumYPanels) - (YIndex * InNumXPanels) - (1 + XIndex)]);
+							OutSortedList.Add(InUnorderedList[(InNumYPanels * (InNumXPanels - XIndex)) - (YIndex + 1)]);
 						}
 						else
 						{
-							OutSortedList.Add(InUnorderedList[(InNumXPanels * InNumYPanels) - ((1 + YIndex) * InNumXPanels) + XIndex]);
+							OutSortedList.Add(InUnorderedList[(InNumXPanels * InNumYPanels) - ((XIndex + 1) * InNumYPanels) + YIndex]);
 						}
 					}
 				}
