DMX - Fix CIS Cook Win64 - Unhandled Exception in DMXMVRChildListNode.cpp

diff --git a/Engine/Plugins/VirtualProduction/DMX/DMXEngine/Source/DMXRuntime/Private/MVR/DMXMVRGeneralSceneDescription.cpp b/Engine/Plugins/VirtualProduction/DMX/DMXEngine/Source/DMXRuntime/Private/MVR/DMXMVRGeneralSceneDescription.cpp
--- a/Engine/Plugins/VirtualProduction/DMX/DMXEngine/Source/DMXRuntime/Private/MVR/DMXMVRGeneralSceneDescription.cpp
+++ b/Engine/Plugins/VirtualProduction/DMX/DMXEngine/Source/DMXRuntime/Private/MVR/DMXMVRGeneralSceneDescription.cpp
@@ -45,7 +45,7 @@ UDMXMVRFixtureNode* UDMXMVRGeneralSceneDescription::FindFixtureNode(const FGuid&
 {
 	checkf(RootNode, TEXT("Unexpected: MVR General Scene Description Root Node is invalid."));
 
-	if (UDMXMVRParametricObjectNodeBase** ObjectNodePtr = RootNode->FindParametricObjectNodeByUUID(FixtureUUID))
+	if (TObjectPtr<UDMXMVRParametricObjectNodeBase>* ObjectNodePtr = RootNode->FindParametricObjectNodeByUUID(FixtureUUID))
 	{
 		return Cast<UDMXMVRFixtureNode>(*ObjectNodePtr);
 	}
@@ -131,7 +131,7 @@ void UDMXMVRGeneralSceneDescription::WriteFixturePatchToGeneralSceneDescription(
 	}
 
 	const FGuid& MVRFixtureUUID = FixturePatch.GetMVRFixtureUUID();
-	UDMXMVRParametricObjectNodeBase** ParametricObjectNodePtr = RootNode->FindParametricObjectNodeByUUID(MVRFixtureUUID);
+	TObjectPtr<UDMXMVRParametricObjectNodeBase>* ParametricObjectNodePtr = RootNode->FindParametricObjectNodeByUUID(MVRFixtureUUID);
 
 	UDMXMVRFixtureNode* MVRFixtureNode = nullptr;
 	if (ParametricObjectNodePtr)
@@ -213,7 +213,7 @@ TArray<int32> UDMXMVRGeneralSceneDescription::GetNumericalFixtureIDsInUse(const
 	TArray<int32> FixtureIDsInUse;
 	for (const FGuid& MVRFixtureUUID : MVRFixtureUUIDsInUse)
 	{
-		if (UDMXMVRParametricObjectNodeBase** ObjectNodePtr = RootNode->FindParametricObjectNodeByUUID(MVRFixtureUUID))
+		if (TObjectPtr<UDMXMVRParametricObjectNodeBase>* ObjectNodePtr = RootNode->FindParametricObjectNodeByUUID(MVRFixtureUUID))
 		{
 			if (UDMXMVRFixtureNode* FixtureNode = Cast<UDMXMVRFixtureNode>(*ObjectNodePtr))
 			{
diff --git a/Engine/Plugins/VirtualProduction/DMX/DMXEngine/Source/DMXRuntime/Private/MVR/Types/DMXMVRChildListNode.cpp b/Engine/Plugins/VirtualProduction/DMX/DMXEngine/Source/DMXRuntime/Private/MVR/Types/DMXMVRChildListNode.cpp
--- a/Engine/Plugins/VirtualProduction/DMX/DMXEngine/Source/DMXRuntime/Private/MVR/Types/DMXMVRChildListNode.cpp
+++ b/Engine/Plugins/VirtualProduction/DMX/DMXEngine/Source/DMXRuntime/Private/MVR/Types/DMXMVRChildListNode.cpp
@@ -92,13 +92,12 @@ bool UDMXMVRChildListNode::Contains(UDMXMVRParametricObjectNodeBase* ParametricO
 	return false;
 }
 
-UDMXMVRParametricObjectNodeBase** UDMXMVRChildListNode::FindParametricObjectNodeByUUID(const FGuid& UUID)
+TObjectPtr<UDMXMVRParametricObjectNodeBase>* UDMXMVRChildListNode::FindParametricObjectNodeByUUID(const FGuid& UUID)
 {
-	UE_TRANSITIONAL_OBJECT_PTR(UDMXMVRParametricObjectNodeBase)* TempParametricObjectNodePtr = Algo::FindByPredicate(ParametricObjectNodes, [UUID](UDMXMVRParametricObjectNodeBase* ParametricObjectNode)
+	TObjectPtr<UDMXMVRParametricObjectNodeBase>* ParametricObjectNodePtr = Algo::FindByPredicate(ParametricObjectNodes, [UUID](UDMXMVRParametricObjectNodeBase* ParametricObjectNode)
 		{
 			return ParametricObjectNode->UUID == UUID;
 		});
-	UDMXMVRParametricObjectNodeBase** ParametricObjectNodePtr = &static_cast<UDMXMVRParametricObjectNodeBase*&>(*TempParametricObjectNodePtr);
 	if (ParametricObjectNodePtr)
 	{
 		return ParametricObjectNodePtr;
diff --git a/Engine/Plugins/VirtualProduction/DMX/DMXEngine/Source/DMXRuntime/Private/MVR/Types/DMXMVRGroupObjectNode.cpp b/Engine/Plugins/VirtualProduction/DMX/DMXEngine/Source/DMXRuntime/Private/MVR/Types/DMXMVRGroupObjectNode.cpp
--- a/Engine/Plugins/VirtualProduction/DMX/DMXEngine/Source/DMXRuntime/Private/MVR/Types/DMXMVRGroupObjectNode.cpp
+++ b/Engine/Plugins/VirtualProduction/DMX/DMXEngine/Source/DMXRuntime/Private/MVR/Types/DMXMVRGroupObjectNode.cpp
@@ -101,7 +101,7 @@ bool UDMXMVRGroupObjectNode::Contains(UDMXMVRParametricObjectNodeBase* Parametri
 	return false;
 }
 
-UDMXMVRParametricObjectNodeBase** UDMXMVRGroupObjectNode::FindParametricObjectNodeByUUID(const FGuid& InUUID)
+TObjectPtr<UDMXMVRParametricObjectNodeBase>* UDMXMVRGroupObjectNode::FindParametricObjectNodeByUUID(const FGuid& InUUID)
 {
 	checkf(ChildListNode, TEXT("Default Subobject 'ChildListNode' is invalid, this is not expected."));
 
diff --git a/Engine/Plugins/VirtualProduction/DMX/DMXEngine/Source/DMXRuntime/Private/MVR/Types/DMXMVRLayerNode.cpp b/Engine/Plugins/VirtualProduction/DMX/DMXEngine/Source/DMXRuntime/Private/MVR/Types/DMXMVRLayerNode.cpp
--- a/Engine/Plugins/VirtualProduction/DMX/DMXEngine/Source/DMXRuntime/Private/MVR/Types/DMXMVRLayerNode.cpp
+++ b/Engine/Plugins/VirtualProduction/DMX/DMXEngine/Source/DMXRuntime/Private/MVR/Types/DMXMVRLayerNode.cpp
@@ -119,10 +119,10 @@ bool UDMXMVRLayerNode::Contains(UDMXMVRParametricObjectNodeBase* ParametricObjec
 	return false;
 }
 
-UDMXMVRParametricObjectNodeBase** UDMXMVRLayerNode::FindParametricObjectNodeByUUID(const FGuid& InUUID) const
+TObjectPtr<UDMXMVRParametricObjectNodeBase>* UDMXMVRLayerNode::FindParametricObjectNodeByUUID(const FGuid& InUUID) const
 {
 	checkf(ChildListNode, TEXT("Unexpected: ChildListNode is instanced, but its pointer is invalid."));
-	UDMXMVRParametricObjectNodeBase** ParametricObjectNodePtr = ChildListNode->FindParametricObjectNodeByUUID(InUUID);
+	TObjectPtr<UDMXMVRParametricObjectNodeBase>* ParametricObjectNodePtr = ChildListNode->FindParametricObjectNodeByUUID(InUUID);
 	if (ParametricObjectNodePtr)
 	{
 		return ParametricObjectNodePtr;
diff --git a/Engine/Plugins/VirtualProduction/DMX/DMXEngine/Source/DMXRuntime/Private/MVR/Types/DMXMVRLayersNode.cpp b/Engine/Plugins/VirtualProduction/DMX/DMXEngine/Source/DMXRuntime/Private/MVR/Types/DMXMVRLayersNode.cpp
--- a/Engine/Plugins/VirtualProduction/DMX/DMXEngine/Source/DMXRuntime/Private/MVR/Types/DMXMVRLayersNode.cpp
+++ b/Engine/Plugins/VirtualProduction/DMX/DMXEngine/Source/DMXRuntime/Private/MVR/Types/DMXMVRLayersNode.cpp
@@ -45,11 +45,11 @@ void UDMXMVRLayersNode::GetFixtureNodes(TArray<UDMXMVRFixtureNode*>& OutFixtureN
 	}
 }
 
-UDMXMVRParametricObjectNodeBase** UDMXMVRLayersNode::FindParametricObjectNodeByUUID(const FGuid& UUID) const
+TObjectPtr<UDMXMVRParametricObjectNodeBase>* UDMXMVRLayersNode::FindParametricObjectNodeByUUID(const FGuid& UUID) const
 {
 	for (UDMXMVRLayerNode* Layer : LayerNodes)
 	{
-		UDMXMVRParametricObjectNodeBase** ParametricObjectNodePtr = Layer->FindParametricObjectNodeByUUID(UUID);
+		TObjectPtr<UDMXMVRParametricObjectNodeBase>* ParametricObjectNodePtr = Layer->FindParametricObjectNodeByUUID(UUID);
 		if (ParametricObjectNodePtr)
 		{
 			return ParametricObjectNodePtr;
diff --git a/Engine/Plugins/VirtualProduction/DMX/DMXEngine/Source/DMXRuntime/Private/MVR/Types/DMXMVRRootNode.cpp b/Engine/Plugins/VirtualProduction/DMX/DMXEngine/Source/DMXRuntime/Private/MVR/Types/DMXMVRRootNode.cpp
--- a/Engine/Plugins/VirtualProduction/DMX/DMXEngine/Source/DMXRuntime/Private/MVR/Types/DMXMVRRootNode.cpp
+++ b/Engine/Plugins/VirtualProduction/DMX/DMXEngine/Source/DMXRuntime/Private/MVR/Types/DMXMVRRootNode.cpp
@@ -149,7 +149,7 @@ UDMXMVRChildListNode& UDMXMVRRootNode::GetOrCreateFirstChildListNode()
 	return *ChildListNode;
 }
 
-UDMXMVRParametricObjectNodeBase** UDMXMVRRootNode::FindParametricObjectNodeByUUID(const FGuid& UUID) const
+TObjectPtr<UDMXMVRParametricObjectNodeBase>* UDMXMVRRootNode::FindParametricObjectNodeByUUID(const FGuid& UUID) const
 {
 	UDMXMVRLayersNode* LayersNode = SceneNode ? SceneNode->GetLayersNode() : nullptr;
 	if (LayersNode)
diff --git a/Engine/Plugins/VirtualProduction/DMX/DMXEngine/Source/DMXRuntime/Public/MVR/Types/DMXMVRChildListNode.h b/Engine/Plugins/VirtualProduction/DMX/DMXEngine/Source/DMXRuntime/Public/MVR/Types/DMXMVRChildListNode.h
--- a/Engine/Plugins/VirtualProduction/DMX/DMXEngine/Source/DMXRuntime/Public/MVR/Types/DMXMVRChildListNode.h
+++ b/Engine/Plugins/VirtualProduction/DMX/DMXEngine/Source/DMXRuntime/Public/MVR/Types/DMXMVRChildListNode.h
@@ -33,7 +33,7 @@ public:
 	bool Contains(UDMXMVRParametricObjectNodeBase* ParametricObjectNode);
 
 	/** Returns a pointer ot the Parametric Object Node that corresponds to the UUID, or nullptr if there is no corresponding Parametric Object Node */
-	UDMXMVRParametricObjectNodeBase** FindParametricObjectNodeByUUID(const FGuid& UUID);
+	TObjectPtr<UDMXMVRParametricObjectNodeBase>* FindParametricObjectNodeByUUID(const FGuid& UUID);
 
 	/** Creates a Parametric Object Node as a Child in the Child List Node */
 	template<typename ParametricObjectNodeClass, typename TEnableIf<TIsDerivedFrom<ParametricObjectNodeClass, UDMXMVRParametricObjectNodeBase>::Value, bool>::Type = true>
diff --git a/Engine/Plugins/VirtualProduction/DMX/DMXEngine/Source/DMXRuntime/Public/MVR/Types/DMXMVRGroupObjectNode.h b/Engine/Plugins/VirtualProduction/DMX/DMXEngine/Source/DMXRuntime/Public/MVR/Types/DMXMVRGroupObjectNode.h
--- a/Engine/Plugins/VirtualProduction/DMX/DMXEngine/Source/DMXRuntime/Public/MVR/Types/DMXMVRGroupObjectNode.h
+++ b/Engine/Plugins/VirtualProduction/DMX/DMXEngine/Source/DMXRuntime/Public/MVR/Types/DMXMVRGroupObjectNode.h
@@ -42,7 +42,7 @@ public:
 	bool Contains(UDMXMVRParametricObjectNodeBase* ParametricObjectNode);
 
 	/** Returns a pointer ot the Parametric Object Node that corresponds to the UUID, or nullptr if there is no corresponding Parametric Object Node */
-	UDMXMVRParametricObjectNodeBase** FindParametricObjectNodeByUUID(const FGuid& InUUID);
+	TObjectPtr<UDMXMVRParametricObjectNodeBase>* FindParametricObjectNodeByUUID(const FGuid& InUUID);
 
 	/** Removes the Parametric Object Node from Child List. Returns true if the Node was removed. */
 	bool RemoveParametricObjectNode(UDMXMVRParametricObjectNodeBase* ParametricObjectNode);
diff --git a/Engine/Plugins/VirtualProduction/DMX/DMXEngine/Source/DMXRuntime/Public/MVR/Types/DMXMVRLayerNode.h b/Engine/Plugins/VirtualProduction/DMX/DMXEngine/Source/DMXRuntime/Public/MVR/Types/DMXMVRLayerNode.h
--- a/Engine/Plugins/VirtualProduction/DMX/DMXEngine/Source/DMXRuntime/Public/MVR/Types/DMXMVRLayerNode.h
+++ b/Engine/Plugins/VirtualProduction/DMX/DMXEngine/Source/DMXRuntime/Public/MVR/Types/DMXMVRLayerNode.h
@@ -39,7 +39,7 @@ public:
 	bool Contains(UDMXMVRParametricObjectNodeBase* ParametricObjectNode);
 
 	/** Returns a pointer ot the Parametric Object Node that corresponds to the UUID, or nullptr if there is no corresponding Parametric Object Node */
-	UDMXMVRParametricObjectNodeBase** FindParametricObjectNodeByUUID(const FGuid& InUUID) const;
+	TObjectPtr<UDMXMVRParametricObjectNodeBase>* FindParametricObjectNodeByUUID(const FGuid& InUUID) const;
 
 	/** Returns a list of graphic objects that are part of the layer. */
 	FORCEINLINE UDMXMVRChildListNode* GetChildListNode() const { return ChildListNode; }
diff --git a/Engine/Plugins/VirtualProduction/DMX/DMXEngine/Source/DMXRuntime/Public/MVR/Types/DMXMVRLayersNode.h b/Engine/Plugins/VirtualProduction/DMX/DMXEngine/Source/DMXRuntime/Public/MVR/Types/DMXMVRLayersNode.h
--- a/Engine/Plugins/VirtualProduction/DMX/DMXEngine/Source/DMXRuntime/Public/MVR/Types/DMXMVRLayersNode.h
+++ b/Engine/Plugins/VirtualProduction/DMX/DMXEngine/Source/DMXRuntime/Public/MVR/Types/DMXMVRLayersNode.h
@@ -32,7 +32,7 @@ public:
 	void GetFixtureNodes(TArray<UDMXMVRFixtureNode*>& OutFixtureNodes) const;
 
 	/** Returns a pointer ot the Parametric Object Node that corresponds to the UUID, or nullptr if there is no corresponding Parametric Object Node */
-	UDMXMVRParametricObjectNodeBase** FindParametricObjectNodeByUUID(const FGuid& UUID) const;
+	TObjectPtr<UDMXMVRParametricObjectNodeBase>* FindParametricObjectNodeByUUID(const FGuid& UUID) const;
 
 	/** Creates a new Layer. Returns the newly created Layer */
 	UDMXMVRLayerNode* CreateLayer();
diff --git a/Engine/Plugins/VirtualProduction/DMX/DMXEngine/Source/DMXRuntime/Public/MVR/Types/DMXMVRRootNode.h b/Engine/Plugins/VirtualProduction/DMX/DMXEngine/Source/DMXRuntime/Public/MVR/Types/DMXMVRRootNode.h
--- a/Engine/Plugins/VirtualProduction/DMX/DMXEngine/Source/DMXRuntime/Public/MVR/Types/DMXMVRRootNode.h
+++ b/Engine/Plugins/VirtualProduction/DMX/DMXEngine/Source/DMXRuntime/Public/MVR/Types/DMXMVRRootNode.h
@@ -48,7 +48,7 @@ public:
 	UDMXMVRChildListNode& GetOrCreateFirstChildListNode();
 
 	/** Returns a pointer ot the Parametric Object Node that corresponds to the UUID, or nullptr if there is no corresponding Parametric Object Node */
-	UDMXMVRParametricObjectNodeBase** FindParametricObjectNodeByUUID(const FGuid& UUID) const;
+	TObjectPtr<UDMXMVRParametricObjectNodeBase>* FindParametricObjectNodeByUUID(const FGuid& UUID) const;
 
 	/** Removes the Parametric Object Node from the General Scene Description. Returns true if the Node was removed. */
 	bool RemoveParametricObjectNode(UDMXMVRParametricObjectNodeBase* ParametricObjectNode);
