DMX - Fix MVR only imports one patch per addresses, improve log message if mode cannot be found in GDTF

diff --git a/Engine/Plugins/VirtualProduction/DMX/DMXEngine/Source/DMXEditor/Private/Factories/DMXLibraryFromMVRFactory.cpp b/Engine/Plugins/VirtualProduction/DMX/DMXEngine/Source/DMXEditor/Private/Factories/DMXLibraryFromMVRFactory.cpp
--- a/Engine/Plugins/VirtualProduction/DMX/DMXEngine/Source/DMXEditor/Private/Factories/DMXLibraryFromMVRFactory.cpp
+++ b/Engine/Plugins/VirtualProduction/DMX/DMXEngine/Source/DMXEditor/Private/Factories/DMXLibraryFromMVRFactory.cpp
@@ -270,16 +270,6 @@ void UDMXLibraryFromMVRFactory::InitDMXLibrary(UDMXLibrary* DMXLibrary, const TA
 		}
 
 		UDMXEntityFixtureType* FixtureType = GDTFSpecToFixtureTypeMap[FixtureNode->GDTFSpec];
-		UDMXEntityFixturePatch* const* FixturePatchPtr = DMXLibrary->GetEntitiesTypeCast<UDMXEntityFixturePatch>().FindByPredicate([FixtureNode, FixtureType](UDMXEntityFixturePatch* FixturePatch)
-			{
-				const FDMXFixtureMode* ActiveModePtr = FixturePatch->GetActiveMode();
-				return
-					FixturePatch->GetFixtureType() == FixtureType &&
-					FixturePatch->GetUniverseID() == FixtureNode->GetUniverseID() &&
-					FixturePatch->GetStartingChannel() == FixtureNode->GetStartingChannel() &&
-					ActiveModePtr && ActiveModePtr->ModeName == FixtureNode->GDTFMode;
-			});
-
 		if (!FixtureTypeToColorMap.Contains(FixtureType))
 		{
 			FLinearColor FixtureTypeColor = FLinearColor::MakeRandomColor();
@@ -293,30 +283,26 @@ void UDMXLibraryFromMVRFactory::InitDMXLibrary(UDMXLibrary* DMXLibrary, const TA
 			FixtureTypeToColorMap.Add(FixtureType, FixtureTypeColor);
 		}
 
-		if (!FixturePatchPtr)
-		{
-			int32 ActiveModeIndex = FixtureType->Modes.IndexOfByPredicate([FixtureNode](const FDMXFixtureMode& Mode)
-				{
-					return Mode.ModeName == FixtureNode->GDTFMode;
-				});
-
-			if (ActiveModeIndex != INDEX_NONE)
-			{
-				FDMXEntityFixturePatchConstructionParams FixturePatchConstructionParams;
-				FixturePatchConstructionParams.ActiveMode = ActiveModeIndex;
-				FixturePatchConstructionParams.FixtureTypeRef = FDMXEntityFixtureTypeRef(FixtureType);
-				FixturePatchConstructionParams.UniverseID = FixtureNode->GetUniverseID();
-				FixturePatchConstructionParams.StartingAddress = FixtureNode->GetStartingChannel();
-				FixturePatchConstructionParams.MVRFixtureUUID = FixtureNode->UUID;
-
-				UDMXEntityFixturePatch* FixturePatch = UDMXEntityFixturePatch::CreateFixturePatchInLibrary(FixturePatchConstructionParams, FixtureNode->Name);
-				FixturePatch->EditorColor = FixtureTypeToColorMap.FindChecked(FixtureType);
-				FixturePatchPtr = &FixturePatch;
-			}
-			else
+		int32 ActiveModeIndex = FixtureType->Modes.IndexOfByPredicate([FixtureNode](const FDMXFixtureMode& Mode)
 			{
-				UE_LOG(LogDMXEditor, Warning, TEXT("Skipped creating a Fixture Patch for '%s', GDTF '%s' has no valid Mode."), *FixtureNode->Name, *FixtureNode->GDTFSpec)
-			}
+				return Mode.ModeName == FixtureNode->GDTFMode;
+			});
+
+		if (ActiveModeIndex != INDEX_NONE)
+		{
+			FDMXEntityFixturePatchConstructionParams FixturePatchConstructionParams;
+			FixturePatchConstructionParams.ActiveMode = ActiveModeIndex;
+			FixturePatchConstructionParams.FixtureTypeRef = FDMXEntityFixtureTypeRef(FixtureType);
+			FixturePatchConstructionParams.UniverseID = FixtureNode->GetUniverseID();
+			FixturePatchConstructionParams.StartingAddress = FixtureNode->GetStartingChannel();
+			FixturePatchConstructionParams.MVRFixtureUUID = FixtureNode->UUID;
+
+			UDMXEntityFixturePatch* FixturePatch = UDMXEntityFixturePatch::CreateFixturePatchInLibrary(FixturePatchConstructionParams, FixtureNode->Name);
+			FixturePatch->EditorColor = FixtureTypeToColorMap.FindChecked(FixtureType);
+		}
+		else
+		{
+			UE_LOG(LogDMXEditor, Warning, TEXT("Skipped creating a Fixture Patch for '%s', cannot find Mode '%s' in '%s'."), *FixtureNode->Name, *FixtureNode->GDTFMode, *FixtureNode->GDTFSpec)
 		}
 	}
 
