DMX - Fix MVR does not export if UE created the layer node, prevent from exporting MVR if no patches are available, fix GDTF and MVR assets are losing their source file path

diff --git a/Engine/Plugins/VirtualProduction/DMX/DMXEngine/Source/DMXEditor/Private/Exporters/DMXMVRExporter.cpp b/Engine/Plugins/VirtualProduction/DMX/DMXEngine/Source/DMXEditor/Private/Exporters/DMXMVRExporter.cpp
--- a/Engine/Plugins/VirtualProduction/DMX/DMXEngine/Source/DMXEditor/Private/Exporters/DMXMVRExporter.cpp
+++ b/Engine/Plugins/VirtualProduction/DMX/DMXEngine/Source/DMXEditor/Private/Exporters/DMXMVRExporter.cpp
@@ -29,7 +29,6 @@
 
 void FDMXMVRExporter::Export(UDMXLibrary* DMXLibrary, const FString& FilePathAndName, FText& OutErrorReason)
 {
-
 	FDMXMVRExporter Instance;
 	Instance.ExportInternal(DMXLibrary, FilePathAndName, OutErrorReason);
 }
@@ -53,9 +52,8 @@ void FDMXMVRExporter::ExportInternal(UDMXLibrary* DMXLibrary, const FString& Fil
 	}
 
 	const TSharedRef<FDMXZipper> Zip = MakeShared<FDMXZipper>();
-	if (!ZipGeneralSceneDescription(Zip, GeneralSceneDescription))
+	if (!ZipGeneralSceneDescription(Zip, GeneralSceneDescription, OutErrorReason))
 	{
-		OutErrorReason = FText::Format(LOCTEXT("MVRExportCreateZipFailedReason", "Failed to write General Scene Description. Cannot export {0}."), FText::FromString(FilePathAndName));
 		return;
 	}
 
@@ -73,10 +71,19 @@ void FDMXMVRExporter::ExportInternal(UDMXLibrary* DMXLibrary, const FString& Fil
 	}
 }
 
-bool FDMXMVRExporter::ZipGeneralSceneDescription(const TSharedRef<FDMXZipper>& Zip, const UDMXMVRGeneralSceneDescription* GeneralSceneDescription)
+bool FDMXMVRExporter::ZipGeneralSceneDescription(const TSharedRef<FDMXZipper>& Zip, const UDMXMVRGeneralSceneDescription* GeneralSceneDescription, FText& OutErrorReason)
 {
+	if (!GeneralSceneDescription->CanCreateXmlFile(OutErrorReason))
+	{
+		return false;
+	}
+
 	TSharedPtr<FXmlFile> XmlFile = GeneralSceneDescription->CreateXmlFile();
-	
+	if (!ensureMsgf(XmlFile.IsValid(), TEXT("General Scene Descriptions returns CanCreateXmlFile() as true, but CreateXmlFile() fails.")))
+	{
+		return false;
+	}
+
 	// Try to merge the source General Scene Description Xml
 	const TSharedPtr<FXmlFile> SourceXmlFile = CreateSourceGeneralSceneDescriptionXmlFile(GeneralSceneDescription);
 	if (SourceXmlFile.IsValid())
@@ -153,7 +160,7 @@ bool FDMXMVRExporter::ZipGDTFs(const TSharedRef<FDMXZipper>& Zip, UDMXLibrary* D
 		const TArray64<uint8>& RawSourceData = RefreshSourceDataAndFixtureType(*FixtureType, *GDTFAssetImportData);
 		if (RawSourceData.Num() > 0)
 		{
-			const FString GDTFFilename = FPaths::GetCleanFilename(GDTFAssetImportData->GetSourceFilePathAndName());
+			const FString GDTFFilename = FPaths::GetCleanFilename(GDTFAssetImportData->GetFilePathAndName());
 			Zip->AddFile(GDTFFilename, GDTFAssetImportData->GetRawSourceData());
 		}
 		else
@@ -274,9 +281,9 @@ const TArray64<uint8>& FDMXMVRExporter::RefreshSourceDataAndFixtureType(UDMXEnti
 
 			TArray<FString> Filenames;
 			if (!InOutGDTFAssetImportData.GetSourceData().SourceFiles.IsEmpty() &&
-				FPaths::FileExists(InOutGDTFAssetImportData.GetSourceFilePathAndName()))
+				FPaths::FileExists(InOutGDTFAssetImportData.GetFilePathAndName()))
 			{
-				Filenames.Add(InOutGDTFAssetImportData.GetSourceFilePathAndName());
+				Filenames.Add(InOutGDTFAssetImportData.GetFilePathAndName());
 			}
 			else
 			{
diff --git a/Engine/Plugins/VirtualProduction/DMX/DMXEngine/Source/DMXEditor/Private/Exporters/DMXMVRExporter.h b/Engine/Plugins/VirtualProduction/DMX/DMXEngine/Source/DMXEditor/Private/Exporters/DMXMVRExporter.h
--- a/Engine/Plugins/VirtualProduction/DMX/DMXEngine/Source/DMXEditor/Private/Exporters/DMXMVRExporter.h
+++ b/Engine/Plugins/VirtualProduction/DMX/DMXEngine/Source/DMXEditor/Private/Exporters/DMXMVRExporter.h
@@ -26,7 +26,7 @@ private:
 	void ExportInternal(UDMXLibrary* DMXLibrary, const FString& FilePathAndName, FText& OutErrorReason);
 
 	/** Zips the GeneralSceneDescription.xml */
-	UE_NODISCARD bool ZipGeneralSceneDescription(const TSharedRef<FDMXZipper>& Zip, const UDMXMVRGeneralSceneDescription* GeneralSceneDescription);
+	UE_NODISCARD bool ZipGeneralSceneDescription(const TSharedRef<FDMXZipper>& Zip, const UDMXMVRGeneralSceneDescription* GeneralSceneDescription, FText& OutErrorReason);
 
 	/** Zips the GDTFs from the Library */
 	UE_NODISCARD bool ZipGDTFs(const TSharedRef<FDMXZipper>& Zip, UDMXLibrary* DMXLibrary);
diff --git a/Engine/Plugins/VirtualProduction/DMX/DMXEngine/Source/DMXEditor/Private/Factories/DMXGDTFFactory.cpp b/Engine/Plugins/VirtualProduction/DMX/DMXEngine/Source/DMXEditor/Private/Factories/DMXGDTFFactory.cpp
--- a/Engine/Plugins/VirtualProduction/DMX/DMXEngine/Source/DMXEditor/Private/Factories/DMXGDTFFactory.cpp
+++ b/Engine/Plugins/VirtualProduction/DMX/DMXEngine/Source/DMXEditor/Private/Factories/DMXGDTFFactory.cpp
@@ -188,7 +188,7 @@ bool UDMXGDTFFactory::CanReimport(UObject* Obj, TArray<FString>& OutFilenames)
 	UDMXImportGDTF* GDTF = Cast<UDMXImportGDTF>(Obj);
 	if (GDTF && GDTF->GetGDTFAssetImportData())
 	{
-		const FString SourceFilename = GDTF->GetGDTFAssetImportData()->GetSourceFilePathAndName();
+		const FString SourceFilename = GDTF->GetGDTFAssetImportData()->GetFilePathAndName();
 		OutFilenames.Add(SourceFilename);
 		return true;
 	}
@@ -213,7 +213,7 @@ EReimportResult::Type UDMXGDTFFactory::Reimport(UObject* InObject)
 		return EReimportResult::Failed;
 	}
 
-	const FString SourceFilename = GDTF->GetGDTFAssetImportData()->GetSourceFilePathAndName();
+	const FString SourceFilename = GDTF->GetGDTFAssetImportData()->GetFilePathAndName();
 	if (!FPaths::FileExists(SourceFilename))
 	{
 		return EReimportResult::Failed;
diff --git a/Engine/Plugins/VirtualProduction/DMX/DMXEngine/Source/DMXEditor/Private/Factories/DMXLibraryFromMVRFactory.cpp b/Engine/Plugins/VirtualProduction/DMX/DMXEngine/Source/DMXEditor/Private/Factories/DMXLibraryFromMVRFactory.cpp
--- a/Engine/Plugins/VirtualProduction/DMX/DMXEngine/Source/DMXEditor/Private/Factories/DMXLibraryFromMVRFactory.cpp
+++ b/Engine/Plugins/VirtualProduction/DMX/DMXEngine/Source/DMXEditor/Private/Factories/DMXLibraryFromMVRFactory.cpp
@@ -174,7 +174,7 @@ TArray<UDMXImportGDTF*> UDMXLibraryFromMVRFactory::CreateGDTFAssets(UObject* Par
 
 					if (UDMXGDTFAssetImportData* GDTFAssetImportData = GDTFAsset->GetGDTFAssetImportData())
 					{
-						const FString SourceFilename = FPaths::GetCleanFilename(GDTFAssetImportData->GetSourceFilePathAndName());
+						const FString SourceFilename = FPaths::GetCleanFilename(GDTFAssetImportData->GetFilePathAndName());
 						ImportedGDTFNames.Add(SourceFilename);
 					}
 				}
@@ -247,7 +247,7 @@ void UDMXLibraryFromMVRFactory::InitDMXLibrary(UDMXLibrary* DMXLibrary, const TA
 			continue;
 		}
 
-		const FString GDTFSourceFilename = GDTFAssetImportData->GetSourceFilePathAndName();
+		const FString GDTFSourceFilename = GDTFAssetImportData->GetFilePathAndName();
 		const FString GDTFFilename = FPaths::GetCleanFilename(GDTFSourceFilename);
 		FDMXEntityFixtureTypeConstructionParams FixtureTypeConstructionParams;
 		FixtureTypeConstructionParams.DMXCategory = FDMXFixtureCategory(FDMXFixtureCategory::GetFirstValue());
diff --git a/Engine/Plugins/VirtualProduction/DMX/DMXEngine/Source/DMXRuntime/Private/Library/DMXGDTFAssetImportData.cpp b/Engine/Plugins/VirtualProduction/DMX/DMXEngine/Source/DMXRuntime/Private/Library/DMXGDTFAssetImportData.cpp
--- a/Engine/Plugins/VirtualProduction/DMX/DMXEngine/Source/DMXRuntime/Private/Library/DMXGDTFAssetImportData.cpp
+++ b/Engine/Plugins/VirtualProduction/DMX/DMXEngine/Source/DMXRuntime/Private/Library/DMXGDTFAssetImportData.cpp
@@ -5,17 +5,12 @@
 #include "Misc/FileHelper.h"
 #include "Misc/Paths.h"
 
-#if WITH_EDITOR
-FString UDMXGDTFAssetImportData::GetSourceFilePathAndName() const
-{
-	// Not GetFirstFilename(), we want the name even if the file doesn't exist on disk locally.
-	return SourceData.SourceFiles.Num() > 0 ? SourceData.SourceFiles[0].RelativeFilename : FString();
-}
-#endif // WITH_EDITOR
 
 #if WITH_EDITOR
-void UDMXGDTFAssetImportData::SetSourceFile(const FString& FilePathAndName)
+void UDMXGDTFAssetImportData::SetSourceFile(const FString& InFilePathAndName)
 {
+	FilePathAndName = InFilePathAndName;
+
 	FAssetImportInfo Info;
 	Info.Insert(FAssetImportInfo::FSourceFile(FilePathAndName));
 	SourceData = Info;
@@ -29,3 +24,14 @@ void UDMXGDTFAssetImportData::SetSourceFile(const FString& FilePathAndName)
 	FFileHelper::LoadFileToArray(RawSourceData.ByteArray, *FilePathAndName);
 }
 #endif // WITH_EDITOR
+
+void UDMXGDTFAssetImportData::PostLoad() 
+{
+	Super::PostLoad();
+
+#if WITH_EDITOR
+	FAssetImportInfo Info;
+	Info.Insert(FAssetImportInfo::FSourceFile(FilePathAndName));
+	SourceData = Info;
+#endif // WITH_EDITOR
+}
diff --git a/Engine/Plugins/VirtualProduction/DMX/DMXEngine/Source/DMXRuntime/Private/MVR/DMXMVRAssetImportData.cpp b/Engine/Plugins/VirtualProduction/DMX/DMXEngine/Source/DMXRuntime/Private/MVR/DMXMVRAssetImportData.cpp
--- a/Engine/Plugins/VirtualProduction/DMX/DMXEngine/Source/DMXRuntime/Private/MVR/DMXMVRAssetImportData.cpp
+++ b/Engine/Plugins/VirtualProduction/DMX/DMXEngine/Source/DMXRuntime/Private/MVR/DMXMVRAssetImportData.cpp
@@ -7,8 +7,10 @@
 
 
 #if WITH_EDITOR
-void UDMXMVRAssetImportData::SetSourceFile(const FString& FilePathAndName)
+void UDMXMVRAssetImportData::SetSourceFile(const FString& InFilePathAndName)
 {
+	FilePathAndName = InFilePathAndName;
+
 	FAssetImportInfo Info;
 	Info.Insert(FAssetImportInfo::FSourceFile(FilePathAndName));
 	SourceData = Info;
@@ -22,3 +24,14 @@ void UDMXMVRAssetImportData::SetSourceFile(const FString& FilePathAndName)
 	FFileHelper::LoadFileToArray(RawSourceData.ByteArray, *FilePathAndName);
 }
 #endif // WITH_EDITOR
+
+void UDMXMVRAssetImportData::PostLoad()
+{
+	Super::PostLoad();
+
+#if WITH_EDITOR
+	FAssetImportInfo Info;
+	Info.Insert(FAssetImportInfo::FSourceFile(FilePathAndName));
+	SourceData = Info;
+#endif
+}
diff --git a/Engine/Plugins/VirtualProduction/DMX/DMXEngine/Source/DMXRuntime/Private/MVR/DMXMVRGeneralSceneDescription.cpp b/Engine/Plugins/VirtualProduction/DMX/DMXEngine/Source/DMXRuntime/Private/MVR/DMXMVRGeneralSceneDescription.cpp
--- a/Engine/Plugins/VirtualProduction/DMX/DMXEngine/Source/DMXRuntime/Private/MVR/DMXMVRGeneralSceneDescription.cpp
+++ b/Engine/Plugins/VirtualProduction/DMX/DMXEngine/Source/DMXRuntime/Private/MVR/DMXMVRGeneralSceneDescription.cpp
@@ -24,6 +24,8 @@
 #include "Misc/Paths.h"
 
 
+#define LOCTEXT_NAMESPACE "DMXMVRGeneralSceneDescription"
+
 UDMXMVRGeneralSceneDescription::UDMXMVRGeneralSceneDescription()
 {
 	RootNode = CreateDefaultSubobject<UDMXMVRRootNode>("MVRRootNode");
@@ -112,6 +114,20 @@ void UDMXMVRGeneralSceneDescription::WriteDMXLibraryToGeneralSceneDescription(co
 }
 #endif // WITH_EDITOR
 
+#if WITH_EDITOR
+bool UDMXMVRGeneralSceneDescription::CanCreateXmlFile(FText& OutReason) const
+{
+	TArray<UDMXMVRFixtureNode*> FixtureNodes;
+	GetFixtureNodes(FixtureNodes);
+	if (FixtureNodes.IsEmpty())
+	{
+		OutReason = LOCTEXT("CannotCreateXmlFileBecauseNoFixtures", "DMX Library does not define valid MVR fixtures.");
+		return false;
+	}
+	return true;
+}
+#endif // WITH_EDITOR
+
 #if WITH_EDITOR
 TSharedPtr<FXmlFile> UDMXMVRGeneralSceneDescription::CreateXmlFile() const
 {
@@ -170,7 +186,7 @@ void UDMXMVRGeneralSceneDescription::WriteFixturePatchToGeneralSceneDescription(
 			{
 				if (GDTF && GDTF->GetGDTFAssetImportData())
 				{
-					return GDTF->GetGDTFAssetImportData()->GetSourceFilePathAndName();
+					return GDTF->GetGDTFAssetImportData()->GetFilePathAndName();
 				}
 				return FString();
 			}();
@@ -233,3 +249,5 @@ TArray<int32> UDMXMVRGeneralSceneDescription::GetNumericalFixtureIDsInUse(const
 
 	return FixtureIDsInUse;
 }
+
+#undef LOCTEXT_NAMESPACE
diff --git a/Engine/Plugins/VirtualProduction/DMX/DMXEngine/Source/DMXRuntime/Private/MVR/Types/DMXMVRLayerNode.cpp b/Engine/Plugins/VirtualProduction/DMX/DMXEngine/Source/DMXRuntime/Private/MVR/Types/DMXMVRLayerNode.cpp
--- a/Engine/Plugins/VirtualProduction/DMX/DMXEngine/Source/DMXRuntime/Private/MVR/Types/DMXMVRLayerNode.cpp
+++ b/Engine/Plugins/VirtualProduction/DMX/DMXEngine/Source/DMXRuntime/Private/MVR/Types/DMXMVRLayerNode.cpp
@@ -13,6 +13,7 @@
 
 UDMXMVRLayerNode::UDMXMVRLayerNode()
 {
+	UUID = FGuid::NewGuid();
 	ChildListNode = CreateDefaultSubobject<UDMXMVRChildListNode>("ChildListNode");
 }
 
diff --git a/Engine/Plugins/VirtualProduction/DMX/DMXEngine/Source/DMXRuntime/Private/MVR/Types/DMXMVRParametricObjectNodeBase.cpp b/Engine/Plugins/VirtualProduction/DMX/DMXEngine/Source/DMXRuntime/Private/MVR/Types/DMXMVRParametricObjectNodeBase.cpp
--- a/Engine/Plugins/VirtualProduction/DMX/DMXEngine/Source/DMXRuntime/Private/MVR/Types/DMXMVRParametricObjectNodeBase.cpp
+++ b/Engine/Plugins/VirtualProduction/DMX/DMXEngine/Source/DMXRuntime/Private/MVR/Types/DMXMVRParametricObjectNodeBase.cpp
@@ -6,6 +6,11 @@
 #include "MVR/Types/DMXMVRLayerNode.h"
 
 
+UDMXMVRParametricObjectNodeBase::UDMXMVRParametricObjectNodeBase()
+{
+	UUID = FGuid::NewGuid();
+}
+
 UDMXMVRLayerNode* UDMXMVRParametricObjectNodeBase::GetLayer()
 {
 	return const_cast<UDMXMVRLayerNode*>(const_cast<const UDMXMVRParametricObjectNodeBase*>(this)->GetLayer());
diff --git a/Engine/Plugins/VirtualProduction/DMX/DMXEngine/Source/DMXRuntime/Private/MVR/Types/DMXMVRRootNode.cpp b/Engine/Plugins/VirtualProduction/DMX/DMXEngine/Source/DMXRuntime/Private/MVR/Types/DMXMVRRootNode.cpp
--- a/Engine/Plugins/VirtualProduction/DMX/DMXEngine/Source/DMXRuntime/Private/MVR/Types/DMXMVRRootNode.cpp
+++ b/Engine/Plugins/VirtualProduction/DMX/DMXEngine/Source/DMXRuntime/Private/MVR/Types/DMXMVRRootNode.cpp
@@ -83,6 +83,15 @@ TSharedPtr<FXmlFile> UDMXMVRRootNode::CreateXmlFile()
 	checkf(UserDataNode, TEXT("Default Subobject 'UserDataNode' is invalid, this is not expected."));
 	checkf(SceneNode, TEXT("Default Subobject 'SceneNode' is invalid, this is not expected."));
 
+	// Don't export if there's nothing to export
+	TArray<UDMXMVRFixtureNode*> FixtureNodes;
+	GetFixtureNodes(FixtureNodes);
+	if (FixtureNodes.IsEmpty())
+	{
+		return nullptr;
+	}
+
+	// Create Xml File
 	const FString Buffer = "<?xml version=\"1.0\" encoding=\"UTF - 8\" standalone=\"no\" ?>\n<GeneralSceneDescription>\n</GeneralSceneDescription>";
 
 	TSharedRef<FXmlFile> XmlFile = MakeShared<FXmlFile>();
diff --git a/Engine/Plugins/VirtualProduction/DMX/DMXEngine/Source/DMXRuntime/Private/MVR/Types/DMXMVRUnrealEngineDataNode.cpp b/Engine/Plugins/VirtualProduction/DMX/DMXEngine/Source/DMXRuntime/Private/MVR/Types/DMXMVRUnrealEngineDataNode.cpp
--- a/Engine/Plugins/VirtualProduction/DMX/DMXEngine/Source/DMXRuntime/Private/MVR/Types/DMXMVRUnrealEngineDataNode.cpp
+++ b/Engine/Plugins/VirtualProduction/DMX/DMXEngine/Source/DMXRuntime/Private/MVR/Types/DMXMVRUnrealEngineDataNode.cpp
@@ -32,6 +32,6 @@ void UDMXMVRUnrealEngineDataNode::CreateXmlNodeInParent(FXmlNode& ParentNode) co
 	DataXmlNode->SetAttributes(Attributes);
 
 	// 'Date' child node
-	constexpr TCHAR ChildTag_Date[] = TEXT("Date");
+	constexpr TCHAR ChildTag_Date[] = TEXT("CreationDate");
 	DataXmlNode->AppendChildNode(ChildTag_Date, FDateTime::Now().ToString());
 }
diff --git a/Engine/Plugins/VirtualProduction/DMX/DMXEngine/Source/DMXRuntime/Public/Library/DMXGDTFAssetImportData.h b/Engine/Plugins/VirtualProduction/DMX/DMXEngine/Source/DMXRuntime/Public/Library/DMXGDTFAssetImportData.h
--- a/Engine/Plugins/VirtualProduction/DMX/DMXEngine/Source/DMXRuntime/Public/Library/DMXGDTFAssetImportData.h
+++ b/Engine/Plugins/VirtualProduction/DMX/DMXEngine/Source/DMXRuntime/Public/Library/DMXGDTFAssetImportData.h
@@ -19,20 +19,26 @@ class DMXRUNTIME_API UDMXGDTFAssetImportData
 	GENERATED_BODY()
 
 public:
-#if WITH_EDITOR
-	/** Returns path and name of the source file the GDTF asset originated from */
-	FString GetSourceFilePathAndName() const;
+	/** Returns the imported file path and name */
+	FORCEINLINE const FString& GetFilePathAndName() const { return FilePathAndName; };
 
+#if WITH_EDITOR
 	/** Creates new asset GDTF Asset Import Data from the File. Returns the GDTF Asset Import Data or nullptr if not a valid GDTF File */
-	void SetSourceFile(const FString& FilePathAndName);
+	void SetSourceFile(const FString& InFilePathAndName);
 
 	/** Returns the source data the asset was generated from */
 	FORCEINLINE const TArray64<uint8>& GetRawSourceData() const { return RawSourceData.ByteArray; }
 #endif // WITH_EDITOR 
 
-private:
-	/** We need to remember the literal File Name to */
+protected:
+	//~ Begin UObject interface
+	virtual void PostLoad() override;
+	//~ End UObject interface
 
+private:
+	/** The imported file path and name */
+	UPROPERTY()
+	FString FilePathAndName;
 
 #if WITH_EDITORONLY_DATA
 	/** The raw GDTF zip file as byte array */
diff --git a/Engine/Plugins/VirtualProduction/DMX/DMXEngine/Source/DMXRuntime/Public/MVR/DMXMVRAssetImportData.h b/Engine/Plugins/VirtualProduction/DMX/DMXEngine/Source/DMXRuntime/Public/MVR/DMXMVRAssetImportData.h
--- a/Engine/Plugins/VirtualProduction/DMX/DMXEngine/Source/DMXRuntime/Public/MVR/DMXMVRAssetImportData.h
+++ b/Engine/Plugins/VirtualProduction/DMX/DMXEngine/Source/DMXRuntime/Public/MVR/DMXMVRAssetImportData.h
@@ -20,23 +20,28 @@ class DMXRUNTIME_API UDMXMVRAssetImportData
 	GENERATED_BODY()
 
 public:
+	/** Returns the imported file path and name */
+	FORCEINLINE const FString& GetFilePathAndName() const { return FilePathAndName; };
+
 #if WITH_EDITOR
 	/** Sets the source file and initializes members from it */
-	void SetSourceFile(const FString& FilePathAndName);
-
-	/** Returns the Asset Import Data */
-	FORCEINLINE UAssetImportData* GetAssetImportData() { return AssetImportData;};
+	void SetSourceFile(const FString& InFilePathAndName);
 
 	/** Returns the raw source file as byte array, as it was imported */
 	FORCEINLINE const TArray64<uint8>& GetRawSourceData() const { return RawSourceData.ByteArray; }
 #endif // WITH_EDITOR 
 
+protected:
+	//~ Begin UObject interface
+	virtual void PostLoad() override;
+	//~ End UObject interface
+
 private:
-#if WITH_EDITORONLY_DATA
-	/** The Asset Import Data used to generate the GDTF asset or nullptr, if not generated from a GDTF file */
+	/** The imported file path and name */
 	UPROPERTY()
-	TObjectPtr<UAssetImportData> AssetImportData;
+	FString FilePathAndName;
 
+#if WITH_EDITORONLY_DATA
 	/** The raw source file as byte array, as it was imported */
 	UPROPERTY()
 	FDMXByteArray64 RawSourceData;
diff --git a/Engine/Plugins/VirtualProduction/DMX/DMXEngine/Source/DMXRuntime/Public/MVR/DMXMVRGeneralSceneDescription.h b/Engine/Plugins/VirtualProduction/DMX/DMXEngine/Source/DMXRuntime/Public/MVR/DMXMVRGeneralSceneDescription.h
--- a/Engine/Plugins/VirtualProduction/DMX/DMXEngine/Source/DMXRuntime/Public/MVR/DMXMVRGeneralSceneDescription.h
+++ b/Engine/Plugins/VirtualProduction/DMX/DMXEngine/Source/DMXRuntime/Public/MVR/DMXMVRGeneralSceneDescription.h
@@ -47,6 +47,9 @@ public:
 	 */
 	void WriteDMXLibraryToGeneralSceneDescription(const UDMXLibrary& DMXLibrary);
 
+	/** Returns true if an Xml File can be created. Gives a reason if no MVR can be exported */
+	bool CanCreateXmlFile(FText& OutReason) const;
+
 	/** Creates an General Scene Description Xml File from this. */
 	TSharedPtr<FXmlFile> CreateXmlFile() const;
 
diff --git a/Engine/Plugins/VirtualProduction/DMX/DMXEngine/Source/DMXRuntime/Public/MVR/Types/DMXMVRParametricObjectNodeBase.h b/Engine/Plugins/VirtualProduction/DMX/DMXEngine/Source/DMXRuntime/Public/MVR/Types/DMXMVRParametricObjectNodeBase.h
--- a/Engine/Plugins/VirtualProduction/DMX/DMXEngine/Source/DMXRuntime/Public/MVR/Types/DMXMVRParametricObjectNodeBase.h
+++ b/Engine/Plugins/VirtualProduction/DMX/DMXEngine/Source/DMXRuntime/Public/MVR/Types/DMXMVRParametricObjectNodeBase.h
@@ -24,6 +24,9 @@ class UDMXMVRParametricObjectNodeBase
 	GENERATED_BODY()
 
 public:
+	/** Constructor */
+	UDMXMVRParametricObjectNodeBase();
+
 	/** Creates a DMX Xml Node as per MVR standard in parent, or logs warnings if no compliant Node can be created. */
 	virtual void CreateXmlNodeInParent(FXmlNode& ParentNode) const PURE_VIRTUAL(UDMXMVRParametricObjectNodeBase::CreateXmlNodeInParent, return;);
 
