DMX - Fix issues caused by reworked dmx attribute names, deprecate FDMXNameListItem

diff --git a/Engine/Plugins/VirtualProduction/DMX/DMXEngine/Source/DMXBlueprintGraph/Private/DMXGraphPanelPinFactory.cpp b/Engine/Plugins/VirtualProduction/DMX/DMXEngine/Source/DMXBlueprintGraph/Private/DMXGraphPanelPinFactory.cpp
--- a/Engine/Plugins/VirtualProduction/DMX/DMXEngine/Source/DMXBlueprintGraph/Private/DMXGraphPanelPinFactory.cpp
+++ b/Engine/Plugins/VirtualProduction/DMX/DMXEngine/Source/DMXBlueprintGraph/Private/DMXGraphPanelPinFactory.cpp
@@ -7,8 +7,9 @@
 #include "Library/DMXEntityReference.h"
 #include "Library/DMXEntity.h"
 #include "Library/DMXEntityFixtureType.h"
+#include "Widgets/SDMXAttributeNameGraphPin.h"
+#include "Widgets/SDMXFixtureCategoryGraphPin.h"
 #include "Widgets/SDMXEntityReferenceGraphPin.h"
-#include "Widgets/SDynamicNameListGraphPin.h"
 #include "Widgets/SNullGraphPin.h"
 
 #include "EdGraphSchema_K2.h"
@@ -22,17 +23,13 @@ TSharedPtr<class SGraphPin> FDMXGraphPanelPinFactory::CreatePin(class UEdGraphPi
 	{
 		if (UScriptStruct* PinStructType = Cast<UScriptStruct>(InPin->PinType.PinSubCategoryObject.Get()))
 		{
-			if (PinStructType->IsChildOf(FDMXProtocolName::StaticStruct()))
+			if (PinStructType->IsChildOf(FDMXFixtureCategory::StaticStruct()))
 			{
-				return SNew(SDynamicNameListGraphPin<FDMXProtocolName>, InPin);
-			}
-			else if (PinStructType->IsChildOf(FDMXFixtureCategory::StaticStruct()))
-			{
-				return SNew(SDynamicNameListGraphPin<FDMXFixtureCategory>, InPin);
+				return SNew(SDMXFixtureCategoryGraphPin, InPin);
 			}
 			else if (PinStructType->IsChildOf(FDMXAttributeName::StaticStruct()))
 			{
-				return SNew(SDynamicNameListGraphPin<FDMXAttributeName>, InPin);
+				return SNew(SDMXAttributeNameGraphPin, InPin);
 			}
 			else if (PinStructType->IsChildOf(FDMXEntityReference::StaticStruct()))
 			{
@@ -52,4 +49,4 @@ TSharedPtr<class SGraphPin> FDMXGraphPanelPinFactory::CreatePin(class UEdGraphPi
 	}
 
 	return FGraphPanelPinFactory::CreatePin(InPin);
-}
\ No newline at end of file
+}
diff --git a/Engine/Plugins/VirtualProduction/DMX/DMXEngine/Source/DMXBlueprintGraph/Private/K2Node_GetDMXAttributeValues.cpp b/Engine/Plugins/VirtualProduction/DMX/DMXEngine/Source/DMXBlueprintGraph/Private/K2Node_GetDMXAttributeValues.cpp
--- a/Engine/Plugins/VirtualProduction/DMX/DMXEngine/Source/DMXBlueprintGraph/Private/K2Node_GetDMXAttributeValues.cpp
+++ b/Engine/Plugins/VirtualProduction/DMX/DMXEngine/Source/DMXBlueprintGraph/Private/K2Node_GetDMXAttributeValues.cpp
@@ -436,9 +436,7 @@ void UK2Node_GetDMXAttributeValues::ExposeAttributes()
 	{
 		for (const FDMXFixtureFunction& Function : ActiveFixtureMode->Functions)
 		{
-			auto Attribute = Function.Attribute.GetAttribute();
-
-			if (Attribute.Name.IsNone())
+			if (!Function.Attribute.IsValid())
 			{
 				continue;
 			}
diff --git a/Engine/Plugins/VirtualProduction/DMX/DMXEngine/Source/DMXBlueprintGraph/Private/Widgets/SDMXAttributeNameGraphPin.cpp b/Engine/Plugins/VirtualProduction/DMX/DMXEngine/Source/DMXBlueprintGraph/Private/Widgets/SDMXAttributeNameGraphPin.cpp
new file mode 100644
--- /dev/null
+++ b/Engine/Plugins/VirtualProduction/DMX/DMXEngine/Source/DMXBlueprintGraph/Private/Widgets/SDMXAttributeNameGraphPin.cpp
@@ -0,0 +1,45 @@
+// Copyright Epic Games, Inc. All Rights Reserved.
+
+#pragma once
+
+#include "SDMXAttributeNameGraphPin.h"
+
+#include "DMXAttribute.h"
+#include "Widgets/SNameListPicker.h"
+
+
+void SDMXAttributeNameGraphPin::Construct(const FArguments& InArgs, UEdGraphPin* InGraphPinObj)
+{
+	SGraphPin::Construct(SGraphPin::FArguments(), InGraphPinObj);
+}
+
+TSharedRef<SWidget>	SDMXAttributeNameGraphPin::GetDefaultValueWidget()
+{
+	return SNew(SNameListPicker)
+		.HasMultipleValues(false)
+		.Value(this, &SDMXAttributeNameGraphPin::GetValue)
+		.OnValueChanged(this, &SDMXAttributeNameGraphPin::SetValue)
+		.OptionsSource(MakeAttributeLambda(&FDMXAttributeName::GetPredefinedValues))
+		.bDisplayWarningIcon(true)
+		.Visibility(this, &SGraphPin::GetDefaultValueVisibility);
+}
+
+FName SDMXAttributeNameGraphPin::GetValue() const
+{
+	FDMXAttributeName NameItem;
+
+	if (!GraphPinObj->GetDefaultAsString().IsEmpty())
+	{
+		FDMXAttributeName::StaticStruct()->ImportText(*GraphPinObj->GetDefaultAsString(), &NameItem, nullptr, EPropertyPortFlags::PPF_None, GLog, FDMXAttributeName::StaticStruct()->GetName());
+	}
+
+	return NameItem.Name;
+}
+
+void SDMXAttributeNameGraphPin::SetValue(FName NewValue)
+{
+	FString ValueString;
+	FDMXAttributeName NewNameItem(NewValue);
+	FDMXAttributeName::StaticStruct()->ExportText(ValueString, &NewNameItem, nullptr, nullptr, EPropertyPortFlags::PPF_None, nullptr);
+	GraphPinObj->GetSchema()->TrySetDefaultValue(*GraphPinObj, ValueString);
+}
diff --git a/Engine/Plugins/VirtualProduction/DMX/DMXEngine/Source/DMXBlueprintGraph/Private/Widgets/SDMXAttributeNameGraphPin.h b/Engine/Plugins/VirtualProduction/DMX/DMXEngine/Source/DMXBlueprintGraph/Private/Widgets/SDMXAttributeNameGraphPin.h
new file mode 100644
--- /dev/null
+++ b/Engine/Plugins/VirtualProduction/DMX/DMXEngine/Source/DMXBlueprintGraph/Private/Widgets/SDMXAttributeNameGraphPin.h
@@ -0,0 +1,31 @@
+// Copyright Epic Games, Inc. All Rights Reserved.
+
+#pragma once
+
+#include "SGraphPin.h"
+
+#include "CoreMinimal.h"
+
+
+class SDMXAttributeNameGraphPin
+	: public SGraphPin
+{
+public:
+	SLATE_BEGIN_ARGS(SDMXAttributeNameGraphPin)
+	{}
+	SLATE_END_ARGS()
+
+	/**  Constructs this widet */
+	void Construct(const FArguments& InArgs, UEdGraphPin* InGraphPinObj);
+
+	//~ Begin SGraphPin Interface
+	virtual TSharedRef<SWidget>	GetDefaultValueWidget() override;
+	//~ End SGraphPin Interface
+
+private:
+	/** Gets the currently selected value */
+	FName GetValue() const;
+
+	/** Sets the currently selected value */
+	void SetValue(FName NewValue);
+};
diff --git a/Engine/Plugins/VirtualProduction/DMX/DMXEngine/Source/DMXBlueprintGraph/Private/Widgets/SDMXFixtureCategoryGraphPin.cpp b/Engine/Plugins/VirtualProduction/DMX/DMXEngine/Source/DMXBlueprintGraph/Private/Widgets/SDMXFixtureCategoryGraphPin.cpp
new file mode 100644
--- /dev/null
+++ b/Engine/Plugins/VirtualProduction/DMX/DMXEngine/Source/DMXBlueprintGraph/Private/Widgets/SDMXFixtureCategoryGraphPin.cpp
@@ -0,0 +1,46 @@
+// Copyright Epic Games, Inc. All Rights Reserved.
+
+#pragma once
+
+#include "SDMXFixtureCategoryGraphPin.h"
+
+#include "DMXProtocolTypes.h"
+#include "Widgets/SNameListPicker.h"
+#include "Widgets/SNullGraphPin.h"
+
+
+void SDMXFixtureCategoryGraphPin::Construct(const FArguments& InArgs, UEdGraphPin* InGraphPinObj)
+{
+	SGraphPin::Construct(SGraphPin::FArguments(), InGraphPinObj);
+}
+
+TSharedRef<SWidget>	SDMXFixtureCategoryGraphPin::GetDefaultValueWidget()
+{
+	return SNew(SNameListPicker)
+		.HasMultipleValues(false)
+		.Value(this, &SDMXFixtureCategoryGraphPin::GetValue)
+		.OnValueChanged(this, &SDMXFixtureCategoryGraphPin::SetValue)
+		.OptionsSource(MakeAttributeLambda(&FDMXFixtureCategory::GetPredefinedValues))
+		.bDisplayWarningIcon(true)
+		.Visibility(this, &SGraphPin::GetDefaultValueVisibility);
+}
+
+FName SDMXFixtureCategoryGraphPin::GetValue() const
+{
+	FDMXFixtureCategory NameItem;
+
+	if (!GraphPinObj->GetDefaultAsString().IsEmpty())
+	{
+		FDMXFixtureCategory::StaticStruct()->ImportText(*GraphPinObj->GetDefaultAsString(), &NameItem, nullptr, EPropertyPortFlags::PPF_None, GLog, FDMXFixtureCategory::StaticStruct()->GetName());
+	}
+
+	return NameItem.Name;
+}
+
+void SDMXFixtureCategoryGraphPin::SetValue(FName NewValue)
+{
+	FString ValueString;
+	FDMXFixtureCategory NewNameItem(NewValue);
+	FDMXFixtureCategory::StaticStruct()->ExportText(ValueString, &NewNameItem, nullptr, nullptr, EPropertyPortFlags::PPF_None, nullptr);
+	GraphPinObj->GetSchema()->TrySetDefaultValue(*GraphPinObj, ValueString);
+}
diff --git a/Engine/Plugins/VirtualProduction/DMX/DMXEngine/Source/DMXBlueprintGraph/Private/Widgets/SDMXFixtureCategoryGraphPin.h b/Engine/Plugins/VirtualProduction/DMX/DMXEngine/Source/DMXBlueprintGraph/Private/Widgets/SDMXFixtureCategoryGraphPin.h
new file mode 100644
--- /dev/null
+++ b/Engine/Plugins/VirtualProduction/DMX/DMXEngine/Source/DMXBlueprintGraph/Private/Widgets/SDMXFixtureCategoryGraphPin.h
@@ -0,0 +1,31 @@
+// Copyright Epic Games, Inc. All Rights Reserved.
+
+#pragma once
+
+#include "SGraphPin.h"
+
+#include "CoreMinimal.h"
+
+
+class SDMXFixtureCategoryGraphPin
+	: public SGraphPin
+{
+public:
+	SLATE_BEGIN_ARGS(SDMXFixtureCategoryGraphPin)
+	{}
+	SLATE_END_ARGS()
+
+	/**  Constructs this widet */
+	void Construct(const FArguments& InArgs, UEdGraphPin* InGraphPinObj);
+
+	//~ Begin SGraphPin Interface
+	virtual TSharedRef<SWidget>	GetDefaultValueWidget() override;
+	//~ End SGraphPin Interface
+
+private:
+	/** Gets the currently selected value */
+	FName GetValue() const;
+
+	/** Sets the currently selected value */
+	void SetValue(FName NewValue);
+};
diff --git a/Engine/Plugins/VirtualProduction/DMX/DMXEngine/Source/DMXBlueprintGraph/Private/Widgets/SDynamicNameListGraphPin.h b/Engine/Plugins/VirtualProduction/DMX/DMXEngine/Source/DMXBlueprintGraph/Private/Widgets/SDynamicNameListGraphPin.h
deleted file mode 100644
diff --git a/Engine/Plugins/VirtualProduction/DMX/DMXEngine/Source/DMXEditor/Private/Customizations/DMXAttributeNameCustomization.cpp b/Engine/Plugins/VirtualProduction/DMX/DMXEngine/Source/DMXEditor/Private/Customizations/DMXAttributeNameCustomization.cpp
new file mode 100644
--- /dev/null
+++ b/Engine/Plugins/VirtualProduction/DMX/DMXEngine/Source/DMXEditor/Private/Customizations/DMXAttributeNameCustomization.cpp
@@ -0,0 +1,147 @@
+// Copyright Epic Games, Inc. All Rights Reserved.
+
+#pragma once
+
+#include "DMXAttributeNameCustomization.h"
+
+#include "DMXAttribute.h"
+#include "DMXProtocolSettings.h"
+
+#include "DetailWidgetRow.h"
+#include "IPropertyTypeCustomization.h"
+#include "IPropertyUtilities.h"
+#include "Widgets/SNameListPicker.h"
+
+
+TSharedRef<IPropertyTypeCustomization> FDMXAttributeNameCustomization::MakeInstance()
+{
+	return MakeShared<FDMXAttributeNameCustomization>();
+}
+
+void FDMXAttributeNameCustomization::CustomizeHeader(TSharedRef<IPropertyHandle> InPropertyHandle, FDetailWidgetRow& InHeaderRow, IPropertyTypeCustomizationUtils& CustomizationUtils)
+{
+	StructPropertyHandle = InPropertyHandle;
+	check(CastFieldChecked<FStructProperty>(StructPropertyHandle->GetProperty())->Struct == FDMXAttributeName::StaticStruct());
+
+	TSharedPtr<IPropertyUtilities> PropertyUtilities = CustomizationUtils.GetPropertyUtilities();
+
+	UDMXProtocolSettings* ProtocolSettings = GetMutableDefault<UDMXProtocolSettings>();
+	if (!ProtocolSettings)
+	{
+		return;
+	}
+
+	TArray<FName> Attributes;
+	for (const FDMXAttribute& Attribute : ProtocolSettings->Attributes)
+	{
+		Attributes.Add(Attribute.Name);
+	}
+
+	InHeaderRow
+		.IsEnabled(MakeAttributeLambda([=] 
+			{ 
+				return !InPropertyHandle->IsEditConst() && PropertyUtilities->IsPropertyEditingEnabled();
+			}))
+		.NameContent()
+		[
+			StructPropertyHandle->CreatePropertyNameWidget()
+		]
+		.ValueContent()
+		.MinDesiredWidth(125.0f)
+		.MaxDesiredWidth(0.0f)
+		[
+			SNew(SNameListPicker)
+			.Font(CustomizationUtils.GetRegularFont())
+			.HasMultipleValues(this, &FDMXAttributeNameCustomization::HasMultipleValues)
+			.OptionsSource(Attributes)
+			.Value(this, &FDMXAttributeNameCustomization::GetValue)
+			.bDisplayWarningIcon(true)
+			.OnValueChanged(this, &FDMXAttributeNameCustomization::SetValue)
+		];
+
+	ProtocolSettings->GetOnDefaultAttributesChanged().AddLambda([PropertyUtilities]()
+		{
+			PropertyUtilities->ForceRefresh();
+		});
+}
+
+void FDMXAttributeNameCustomization::CustomizeChildren(TSharedRef<IPropertyHandle> InPropertyHandle, IDetailChildrenBuilder& InChildBuilder, IPropertyTypeCustomizationUtils& CustomizationUtils)
+{
+
+}
+
+FName FDMXAttributeNameCustomization::GetValue() const
+{
+	TArray<const void*> RawData;
+	StructPropertyHandle->AccessRawData(RawData);
+
+	for (const void* RawPtr : RawData)
+	{
+		if (RawPtr != nullptr)
+		{
+			// The types we use with this customization must have a cast constructor to FName
+			return reinterpret_cast<const FDMXAttributeName*>(RawPtr)->Name;
+		}
+	}
+
+	return FName();
+}
+
+void FDMXAttributeNameCustomization::SetValue(FName NewValue)
+{
+	FStructProperty* StructProperty = CastFieldChecked<FStructProperty>(StructPropertyHandle->GetProperty());
+
+	TArray<void*> RawData;
+	StructPropertyHandle->AccessRawData(RawData);
+
+	for (void* SingleRawData : RawData)
+	{
+		FDMXAttributeName* PreviousValue = reinterpret_cast<FDMXAttributeName*>(SingleRawData);
+		FDMXAttributeName NewAttributeName;
+		NewAttributeName.SetFromName(NewValue);
+
+		// Export new value to text format that can be imported later
+		FString TextValue;
+		StructProperty->Struct->ExportText(TextValue, &NewAttributeName, PreviousValue, nullptr, EPropertyPortFlags::PPF_None, nullptr);
+
+		// Set values on edited property handle from exported text
+		ensure(StructPropertyHandle->SetValueFromFormattedString(TextValue, EPropertyValueSetFlags::DefaultFlags) == FPropertyAccess::Result::Success);
+	}
+}
+
+bool FDMXAttributeNameCustomization::HasMultipleValues() const
+{
+	TArray<const void*> RawData;
+	StructPropertyHandle->AccessRawData(RawData);
+	if (RawData.Num() == 1)
+	{
+		return false;
+	}
+
+	TOptional<FDMXAttributeName> CompareAgainst;
+	for (const void* RawPtr : RawData)
+	{
+		if (RawPtr == nullptr)
+		{
+			if (CompareAgainst.IsSet())
+			{
+				return false;
+			}
+		}
+		else
+		{
+			const FDMXAttributeName* ThisValue = reinterpret_cast<const FDMXAttributeName*>(RawPtr);
+
+			if (!CompareAgainst.IsSet())
+			{
+				CompareAgainst = *ThisValue;
+			}
+			else if (!(*ThisValue == CompareAgainst.GetValue()))
+			{
+				return true;
+			}
+		}
+	}
+
+	return false;
+}
diff --git a/Engine/Plugins/VirtualProduction/DMX/DMXEngine/Source/DMXEditor/Private/Customizations/DMXAttributeNameCustomization.h b/Engine/Plugins/VirtualProduction/DMX/DMXEngine/Source/DMXEditor/Private/Customizations/DMXAttributeNameCustomization.h
new file mode 100644
--- /dev/null
+++ b/Engine/Plugins/VirtualProduction/DMX/DMXEngine/Source/DMXEditor/Private/Customizations/DMXAttributeNameCustomization.h
@@ -0,0 +1,42 @@
+// Copyright Epic Games, Inc. All Rights Reserved.
+
+#pragma once
+
+#include "DMXProtocolTypes.h"
+
+#include "CoreMinimal.h"
+#include "DetailWidgetRow.h"
+#include "IPropertyTypeCustomization.h"
+#include "IPropertyUtilities.h"
+#include "Widgets/SNameListPicker.h"
+
+
+/**  
+ * Customization for FDMXAttributeName struct
+ */
+class FDMXAttributeNameCustomization
+	: public IPropertyTypeCustomization
+{
+public:
+	/** Creates an instance of this property type customization */
+	static TSharedRef<IPropertyTypeCustomization> MakeInstance();
+
+	//~ IPropertyTypeCustomization interface begin
+	virtual void CustomizeHeader(TSharedRef<IPropertyHandle> InPropertyHandle, FDetailWidgetRow& InHeaderRow, IPropertyTypeCustomizationUtils& CustomizationUtils) override;
+	virtual void CustomizeChildren(TSharedRef<IPropertyHandle> InPropertyHandle, IDetailChildrenBuilder& InChildBuilder, IPropertyTypeCustomizationUtils& CustomizationUtils) override;
+	//~ IPropertyTypeCustomization interface end
+
+private:
+	/** Returns the value as FName */
+	FName GetValue() const;
+
+	/** Sets the value */
+	void SetValue(FName NewValue);
+
+	/** True if the customization displays multiple values */
+	bool HasMultipleValues() const;
+
+private:
+	/** Handle for the FDMXAttributeName struct self */
+	TSharedPtr<IPropertyHandle> StructPropertyHandle;
+};
diff --git a/Engine/Plugins/VirtualProduction/DMX/DMXEngine/Source/DMXEditor/Private/Customizations/DMXFixtureCategoryCustomization.cpp b/Engine/Plugins/VirtualProduction/DMX/DMXEngine/Source/DMXEditor/Private/Customizations/DMXFixtureCategoryCustomization.cpp
new file mode 100644
--- /dev/null
+++ b/Engine/Plugins/VirtualProduction/DMX/DMXEngine/Source/DMXEditor/Private/Customizations/DMXFixtureCategoryCustomization.cpp
@@ -0,0 +1,141 @@
+// Copyright Epic Games, Inc. All Rights Reserved.
+
+#pragma once
+
+#include "DMXFixtureCategoryCustomization.h"
+
+#include "DMXProtocolTypes.h"
+#include "DMXProtocolSettings.h"
+
+#include "DetailWidgetRow.h"
+#include "IPropertyTypeCustomization.h"
+#include "IPropertyUtilities.h"
+#include "Widgets/SNameListPicker.h"
+
+
+TSharedRef<IPropertyTypeCustomization> FDMXFixtureCategoryCustomization::MakeInstance()
+{
+	return MakeShared<FDMXFixtureCategoryCustomization>();
+}
+
+void FDMXFixtureCategoryCustomization::CustomizeHeader(TSharedRef<IPropertyHandle> InPropertyHandle, FDetailWidgetRow& InHeaderRow, IPropertyTypeCustomizationUtils& CustomizationUtils)
+{
+	StructPropertyHandle = InPropertyHandle;
+	check(CastFieldChecked<FStructProperty>(StructPropertyHandle->GetProperty())->Struct == FDMXFixtureCategory::StaticStruct());
+
+	TSharedPtr<IPropertyUtilities> PropertyUtilities = CustomizationUtils.GetPropertyUtilities();
+
+	UDMXProtocolSettings* ProtocolSettings = GetMutableDefault<UDMXProtocolSettings>();
+	if (!ProtocolSettings)
+	{
+		return;
+	}
+
+	InHeaderRow
+		.IsEnabled(MakeAttributeLambda([=] 
+			{ 
+				return !InPropertyHandle->IsEditConst() && PropertyUtilities->IsPropertyEditingEnabled();
+			}))
+		.NameContent()
+		[
+			StructPropertyHandle->CreatePropertyNameWidget()
+		]
+		.ValueContent()
+		.MinDesiredWidth(125.0f)
+		.MaxDesiredWidth(0.0f)
+		[
+			SNew(SNameListPicker)
+			.Font(CustomizationUtils.GetRegularFont())
+			.HasMultipleValues(this, &FDMXFixtureCategoryCustomization::HasMultipleValues)
+			.OptionsSource(ProtocolSettings->FixtureCategories.Array())
+			.Value(this, &FDMXFixtureCategoryCustomization::GetValue)
+			.bDisplayWarningIcon(true)
+			.OnValueChanged(this, &FDMXFixtureCategoryCustomization::SetValue)
+		];
+
+	ProtocolSettings->GetOnDefaultAttributesChanged().AddLambda([PropertyUtilities]()
+		{
+			PropertyUtilities->ForceRefresh();
+		});
+}
+
+void FDMXFixtureCategoryCustomization::CustomizeChildren(TSharedRef<IPropertyHandle> InPropertyHandle, IDetailChildrenBuilder& InChildBuilder, IPropertyTypeCustomizationUtils& CustomizationUtils)
+{
+
+}
+
+FName FDMXFixtureCategoryCustomization::GetValue() const
+{
+	TArray<const void*> RawData;
+	StructPropertyHandle->AccessRawData(RawData);
+
+	for (const void* RawPtr : RawData)
+	{
+		if (RawPtr != nullptr)
+		{
+			// The types we use with this customization must have a cast constructor to FName
+			return reinterpret_cast<const FDMXFixtureCategory*>(RawPtr)->Name;
+		}
+	}
+
+	return FName();
+}
+
+void FDMXFixtureCategoryCustomization::SetValue(FName NewValue)
+{
+	FStructProperty* StructProperty = CastFieldChecked<FStructProperty>(StructPropertyHandle->GetProperty());
+
+	TArray<void*> RawData;
+	StructPropertyHandle->AccessRawData(RawData);
+
+	for (void* SingleRawData : RawData)
+	{
+		FDMXFixtureCategory* PreviousValue = reinterpret_cast<FDMXFixtureCategory*>(SingleRawData);
+		FDMXFixtureCategory NewCategoryName;
+		NewCategoryName.Name = NewValue;
+
+		// Export new value to text format that can be imported later
+		FString TextValue;
+		StructProperty->Struct->ExportText(TextValue, &NewCategoryName, PreviousValue, nullptr, EPropertyPortFlags::PPF_None, nullptr);
+
+		// Set values on edited property handle from exported text
+		ensure(StructPropertyHandle->SetValueFromFormattedString(TextValue, EPropertyValueSetFlags::DefaultFlags) == FPropertyAccess::Result::Success);
+	}
+}
+
+bool FDMXFixtureCategoryCustomization::HasMultipleValues() const
+{
+	TArray<const void*> RawData;
+	StructPropertyHandle->AccessRawData(RawData);
+	if (RawData.Num() == 1)
+	{
+		return false;
+	}
+
+	TOptional<FDMXFixtureCategory> CompareAgainst;
+	for (const void* RawPtr : RawData)
+	{
+		if (RawPtr == nullptr)
+		{
+			if (CompareAgainst.IsSet())
+			{
+				return false;
+			}
+		}
+		else
+		{
+			const FDMXFixtureCategory* ThisValue = reinterpret_cast<const FDMXFixtureCategory*>(RawPtr);
+
+			if (!CompareAgainst.IsSet())
+			{
+				CompareAgainst = *ThisValue;
+			}
+			else if (!(*ThisValue == CompareAgainst.GetValue()))
+			{
+				return true;
+			}
+		}
+	}
+
+	return false;
+}
diff --git a/Engine/Plugins/VirtualProduction/DMX/DMXEngine/Source/DMXEditor/Private/Customizations/DMXFixtureCategoryCustomization.h b/Engine/Plugins/VirtualProduction/DMX/DMXEngine/Source/DMXEditor/Private/Customizations/DMXFixtureCategoryCustomization.h
new file mode 100644
--- /dev/null
+++ b/Engine/Plugins/VirtualProduction/DMX/DMXEngine/Source/DMXEditor/Private/Customizations/DMXFixtureCategoryCustomization.h
@@ -0,0 +1,42 @@
+// Copyright Epic Games, Inc. All Rights Reserved.
+
+#pragma once
+
+#include "DMXProtocolTypes.h"
+
+#include "CoreMinimal.h"
+#include "DetailWidgetRow.h"
+#include "IPropertyTypeCustomization.h"
+#include "IPropertyUtilities.h"
+#include "Widgets/SNameListPicker.h"
+
+
+/**  
+ * Customization for the FDMXFixtureCategory struct 
+ */
+class FDMXFixtureCategoryCustomization
+	: public IPropertyTypeCustomization
+{
+public:
+	/** Creates an instance of this property type customization */
+	static TSharedRef<IPropertyTypeCustomization> MakeInstance();
+
+	//~ IPropertyTypeCustomization interface begin
+	virtual void CustomizeHeader(TSharedRef<IPropertyHandle> InPropertyHandle, FDetailWidgetRow& InHeaderRow, IPropertyTypeCustomizationUtils& CustomizationUtils) override;
+	virtual void CustomizeChildren(TSharedRef<IPropertyHandle> InPropertyHandle, IDetailChildrenBuilder& InChildBuilder, IPropertyTypeCustomizationUtils& CustomizationUtils) override;
+	//~ IPropertyTypeCustomization interface end
+
+private:
+	/** Returns the value as FName */
+	FName GetValue() const;
+
+	/** Sets the value */
+	void SetValue(FName NewValue);
+
+	/** True if the customization displays multiple values */
+	bool HasMultipleValues() const;
+
+private:
+	/** Handle for the FDMXFixtureCategory struct self */
+	TSharedPtr<IPropertyHandle> StructPropertyHandle;
+};
diff --git a/Engine/Plugins/VirtualProduction/DMX/DMXEngine/Source/DMXEditor/Private/Customizations/DMXNameListCustomization.h b/Engine/Plugins/VirtualProduction/DMX/DMXEngine/Source/DMXEditor/Private/Customizations/DMXNameListCustomization.h
deleted file mode 100644
diff --git a/Engine/Plugins/VirtualProduction/DMX/DMXEngine/Source/DMXEditor/Private/DMXEditorModule.cpp b/Engine/Plugins/VirtualProduction/DMX/DMXEngine/Source/DMXEditor/Private/DMXEditorModule.cpp
--- a/Engine/Plugins/VirtualProduction/DMX/DMXEngine/Source/DMXEditor/Private/DMXEditorModule.cpp
+++ b/Engine/Plugins/VirtualProduction/DMX/DMXEngine/Source/DMXEditor/Private/DMXEditorModule.cpp
@@ -10,12 +10,13 @@
 #include "DMXProtocolTypes.h"
 #include "AssetTools/AssetTypeActions_DMXEditorLibrary.h"
 #include "Commands/DMXEditorCommands.h"
+#include "Customizations/DMXAttributeNameCustomization.h"
 #include "Customizations/DMXEntityFixtureTypeDetails.h"
 #include "Customizations/DMXEntityReferenceCustomization.h"
+#include "Customizations/DMXFixtureCategoryCustomization.h"
 #include "Customizations/DMXFixtureSignalFormatCustomization.h"
 #include "Customizations/DMXLibraryPortReferencesCustomization.h"
 #include "Customizations/DMXMVRSceneActorDetails.h"
-#include "Customizations/DMXNameListCustomization.h"
 #include "Customizations/DMXPixelMappingDistributionCustomization.h"
 #include "Customizations/TakeRecorderDMXLibrarySourceEditorCustomization.h"
 #include "Game/DMXComponent.h"
@@ -273,18 +274,17 @@ void FDMXEditorModule::RegisterPropertyTypeCustomizations()
 		FOnGetPropertyTypeCustomizationInstance::CreateStatic(&FDMXFixtureSignalFormatCustomization::MakeInstance)
 	);
 
-	// Property type customization for Name List structs
-	RegisterCustomPropertyTypeLayout(FDMXProtocolName::StaticStruct()->GetFName(), 
-		FOnGetPropertyTypeCustomizationInstance::CreateStatic(&FDMXNameListCustomization<FDMXProtocolName>::MakeInstance)
-	);
-	RegisterCustomPropertyTypeLayout(FDMXFixtureCategory::StaticStruct()->GetFName(), 
-		FOnGetPropertyTypeCustomizationInstance::CreateStatic(&FDMXNameListCustomization<FDMXFixtureCategory>::MakeInstance)
-	);
+	// Property type customization for the FDMXAttributeName struct
 	RegisterCustomPropertyTypeLayout(FDMXAttributeName::StaticStruct()->GetFName(), 
-		FOnGetPropertyTypeCustomizationInstance::CreateStatic(&FDMXNameListCustomization<FDMXAttributeName>::MakeInstance)
+		FOnGetPropertyTypeCustomizationInstance::CreateStatic(&FDMXAttributeNameCustomization::MakeInstance)
+	);
+
+	// Property type customization for the FDMXFixtureCategory struct
+	RegisterCustomPropertyTypeLayout(FDMXFixtureCategory::StaticStruct()->GetFName(),
+		FOnGetPropertyTypeCustomizationInstance::CreateStatic(&FDMXFixtureCategoryCustomization::MakeInstance)
 	);
 
-	// Customizations for the FDMXEntityReference structs
+	// Customizations for FDMXEntityReference structs
 	RegisterCustomPropertyTypeLayout(FDMXEntityFixtureTypeRef::StaticStruct()->GetFName(), 
 		FOnGetPropertyTypeCustomizationInstance::CreateStatic(&FDMXEntityReferenceCustomization::MakeInstance)
 	);
@@ -292,7 +292,7 @@ void FDMXEditorModule::RegisterPropertyTypeCustomizations()
 		FOnGetPropertyTypeCustomizationInstance::CreateStatic(&FDMXEntityReferenceCustomization::MakeInstance)
 	);
 
-	// Customization for the FDMXLibraryPortReferences struct type
+	// Customization for the FDMXLibraryPortReferences struct
 	RegisterCustomPropertyTypeLayout(FDMXLibraryPortReferences::StaticStruct()->GetFName(), 
 		FOnGetPropertyTypeCustomizationInstance::CreateStatic(&FDMXLibraryPortReferencesCustomization::MakeInstance)
 	);
diff --git a/Engine/Plugins/VirtualProduction/DMX/DMXEngine/Source/DMXEditor/Private/DMXFixtureTypeSharedData.cpp b/Engine/Plugins/VirtualProduction/DMX/DMXEngine/Source/DMXEditor/Private/DMXFixtureTypeSharedData.cpp
--- a/Engine/Plugins/VirtualProduction/DMX/DMXEngine/Source/DMXEditor/Private/DMXFixtureTypeSharedData.cpp
+++ b/Engine/Plugins/VirtualProduction/DMX/DMXEngine/Source/DMXEditor/Private/DMXFixtureTypeSharedData.cpp
@@ -901,7 +901,7 @@ void FDMXFixtureTypeSharedData::AddCellAttributeToSelectedMode()
 				FDMXFixtureMode& Mode = FixtureType->Modes[Selection->SelectedModeIndices[0]];
 
 				FDMXFixtureCellAttribute NewAttribute;
-				TArray<FName> AvailableAttributes = FDMXAttributeName::GetPossibleValues();
+				TArray<FName> AvailableAttributes = FDMXAttributeName::GetPredefinedValues();
 				NewAttribute.Attribute = AvailableAttributes.Num() > 0 ? FDMXAttributeName(AvailableAttributes[0]) : FDMXAttributeName();
 
 				const FDMXFixtureTypeSharedDataDetails::FDMXScopedModeEditChangeChainProperty_DEPRECATED ScopedChange(FixtureType);
diff --git a/Engine/Plugins/VirtualProduction/DMX/DMXEngine/Source/DMXEditor/Private/Widgets/FixtureType/DMXFixtureTypeFunctionsEditorFunctionItem.cpp b/Engine/Plugins/VirtualProduction/DMX/DMXEngine/Source/DMXEditor/Private/Widgets/FixtureType/DMXFixtureTypeFunctionsEditorFunctionItem.cpp
--- a/Engine/Plugins/VirtualProduction/DMX/DMXEngine/Source/DMXEditor/Private/Widgets/FixtureType/DMXFixtureTypeFunctionsEditorFunctionItem.cpp
+++ b/Engine/Plugins/VirtualProduction/DMX/DMXEngine/Source/DMXEditor/Private/Widgets/FixtureType/DMXFixtureTypeFunctionsEditorFunctionItem.cpp
@@ -41,7 +41,7 @@ bool FDMXFixtureTypeFunctionsEditorFunctionItem::HasValidAttribute() const
 	if (FixtureType.IsValid() && FixtureType->Modes.IsValidIndex(ModeIndex) && FixtureType->Modes[ModeIndex].Functions.IsValidIndex(FunctionIndex))
 	{
 		const FDMXFixtureFunction& Function = FixtureType->Modes[ModeIndex].Functions[FunctionIndex];
-		return Function.Attribute != FDMXNameListItem::None;
+		return !Function.Attribute.IsValid();
 	}
 
 	return false;
diff --git a/Engine/Plugins/VirtualProduction/DMX/DMXEngine/Source/DMXEditor/Private/Widgets/FixtureType/SDMXFixtureTypeFunctionsEditorFunctionRow.cpp b/Engine/Plugins/VirtualProduction/DMX/DMXEngine/Source/DMXEditor/Private/Widgets/FixtureType/SDMXFixtureTypeFunctionsEditorFunctionRow.cpp
--- a/Engine/Plugins/VirtualProduction/DMX/DMXEngine/Source/DMXEditor/Private/Widgets/FixtureType/SDMXFixtureTypeFunctionsEditorFunctionRow.cpp
+++ b/Engine/Plugins/VirtualProduction/DMX/DMXEngine/Source/DMXEditor/Private/Widgets/FixtureType/SDMXFixtureTypeFunctionsEditorFunctionRow.cpp
@@ -172,7 +172,7 @@ TSharedRef<SWidget> SDMXFixtureTypeFunctionsEditorFunctionRow::GenerateWidgetFor
 			.BorderImage(FAppStyle::GetBrush("NoBorder"))
 			[
 				SNew(SNameListPicker)
-				.OptionsSource(MakeAttributeLambda(&FDMXAttributeName::GetPossibleValues))
+				.OptionsSource(MakeAttributeLambda(&FDMXAttributeName::GetPredefinedValues))
 				.Value(this, &SDMXFixtureTypeFunctionsEditorFunctionRow::GetAttributeName)
 				.bDisplayWarningIcon(true)
 				.OnValueChanged(this, &SDMXFixtureTypeFunctionsEditorFunctionRow::OnUserChangedAttributeName)
diff --git a/Engine/Plugins/VirtualProduction/DMX/DMXEngine/Source/DMXEditor/Private/Widgets/FixtureType/SDMXFixtureTypeMatrixFunctionsEditorMatrixRow.cpp b/Engine/Plugins/VirtualProduction/DMX/DMXEngine/Source/DMXEditor/Private/Widgets/FixtureType/SDMXFixtureTypeMatrixFunctionsEditorMatrixRow.cpp
--- a/Engine/Plugins/VirtualProduction/DMX/DMXEngine/Source/DMXEditor/Private/Widgets/FixtureType/SDMXFixtureTypeMatrixFunctionsEditorMatrixRow.cpp
+++ b/Engine/Plugins/VirtualProduction/DMX/DMXEngine/Source/DMXEditor/Private/Widgets/FixtureType/SDMXFixtureTypeMatrixFunctionsEditorMatrixRow.cpp
@@ -83,7 +83,7 @@ TSharedRef<SWidget> SDMXFixtureTypeMatrixFunctionsEditorMatrixRow::GenerateWidge
 	{
 		return
 			SNew(SNameListPicker)
-			.OptionsSource(MakeAttributeLambda(&FDMXAttributeName::GetPossibleValues))
+			.OptionsSource(MakeAttributeLambda(&FDMXAttributeName::GetPredefinedValues))
 			.Value(this, &SDMXFixtureTypeMatrixFunctionsEditorMatrixRow::GetCellAttributeName)
 			.bDisplayWarningIcon(true)
 			.OnValueChanged(this, &SDMXFixtureTypeMatrixFunctionsEditorMatrixRow::SetCellAttributeName);
diff --git a/Engine/Plugins/VirtualProduction/DMX/DMXEngine/Source/DMXEditor/Private/Widgets/SNameListPicker.cpp b/Engine/Plugins/VirtualProduction/DMX/DMXEngine/Source/DMXEditor/Private/Widgets/SNameListPicker.cpp
--- a/Engine/Plugins/VirtualProduction/DMX/DMXEngine/Source/DMXEditor/Private/Widgets/SNameListPicker.cpp
+++ b/Engine/Plugins/VirtualProduction/DMX/DMXEngine/Source/DMXEditor/Private/Widgets/SNameListPicker.cpp
@@ -29,12 +29,6 @@ void SNameListPicker::Construct(const FArguments& InArgs)
 	OnValueChangedDelegate = InArgs._OnValueChanged;
 
 	UpdateOptionsSource();
-	UDMXProtocolSettings* ProtocolSettings = GetMutableDefault<UDMXProtocolSettings>();
-	if (!ProtocolSettings)
-	{
-		return;
-	}
-	ProtocolSettings->GetOnDefaultAttributesChanged().AddSP(this, &SNameListPicker::UpdateOptionsSource);
 
 	EdititableTextBox = SNew(SEditableTextBox)
 		.Font(FAppStyle::GetFontStyle(TEXT("PropertyWindow.NormalFont")))
@@ -115,7 +109,7 @@ void SNameListPicker::UpdateOptionsSource()
 	OptionsSource.Reserve(NumOptions);
 
 	// <None> is the first option
-	OptionsSource.Add(MakeShared<FName>(FDMXNameListItem::None));
+	OptionsSource.Add(MakeShared<FName>(NAME_None));
 
 	for (const FName& Name : OptionsSourceAttr.Get())
 	{
@@ -141,7 +135,7 @@ TSharedRef<ITableRow> SNameListPicker::GenerateNameItemWidget(TSharedPtr<FName>
 		return TableRow;
 	}
 
-	if (InItem->IsEqual(FDMXNameListItem::None))
+	if (InItem->IsEqual(NAME_None))
 	{
 		RowTextBlock->SetText(NoneLabel);
 		return TableRow;
@@ -270,7 +264,7 @@ void SNameListPicker::UpdateFilteredOptions(const FString& Filter)
 
 		for (TSharedPtr<FName>& NameOption : OptionsSource)
 		{
-			if (!NameOption.IsValid() || NameOption->IsEqual(FDMXNameListItem::None))
+			if (!NameOption.IsValid() || NameOption->IsEqual(NAME_None))
 			{
 				continue;
 			}
@@ -305,7 +299,7 @@ FText SNameListPicker::GetCurrentNameLabel() const
 	}
 
 	const FName CurrentName = ValueAttribute.Get();
-	if (CurrentName.IsEqual(FDMXNameListItem::None))
+	if (CurrentName.IsEqual(NAME_None))
 	{
 		return NoneLabel;
 	}
diff --git a/Engine/Plugins/VirtualProduction/DMX/DMXEngine/Source/DMXRuntime/Private/Library/DMXEntityFixtureType.cpp b/Engine/Plugins/VirtualProduction/DMX/DMXEngine/Source/DMXRuntime/Private/Library/DMXEntityFixtureType.cpp
--- a/Engine/Plugins/VirtualProduction/DMX/DMXEngine/Source/DMXRuntime/Private/Library/DMXEntityFixtureType.cpp
+++ b/Engine/Plugins/VirtualProduction/DMX/DMXEngine/Source/DMXRuntime/Private/Library/DMXEntityFixtureType.cpp
@@ -12,6 +12,7 @@
 #include "Library/DMXImportGDTF.h"
 #include "Library/DMXLibrary.h"
 
+#include "Algo/Find.h"
 
 #define LOCTEXT_NAMESPACE "DMXEntityFixtureType"
 
@@ -23,25 +24,9 @@ uint8 FDMXFixtureFunction::GetLastChannel() const
 FDMXFixtureMatrix::FDMXFixtureMatrix()
 {
 	// Add a default cell attribute	
-	if (!GIsInitialLoad)
-	{
-		if (const UDMXProtocolSettings* ProtocolSettings = GetDefault<UDMXProtocolSettings>())
-		{
-			FDMXAttributeName RedAttribute("Red");
-
-			if (ProtocolSettings->Attributes.Contains(RedAttribute.GetAttribute()))
-			{
-				FDMXFixtureCellAttribute RedCellAttribute;
-				RedCellAttribute.Attribute = RedAttribute;
-				CellAttributes = { RedCellAttribute };
-			}
-			else
-			{
-				FDMXFixtureCellAttribute VoidAttribute;
-				CellAttributes = { VoidAttribute };
-			}
-		}
-	}
+	FDMXFixtureCellAttribute RedCellAttribute;
+	RedCellAttribute.Attribute = FName("Red");
+	CellAttributes = { RedCellAttribute };
 }
 
 int32 FDMXFixtureMatrix::GetNumChannels() const
@@ -114,7 +99,7 @@ bool FDMXFixtureMatrix::GetChannelsFromCell(FIntPoint CellCoordinate, FDMXAttrib
 	for (const FDMXFixtureCellAttribute& CellAttribute : CellAttributes)
 	{
 		int32 CurrentFunctionSize = FDMXConversions::GetSizeOfSignalFormat(CellAttribute.DataType);
-		if (CellAttribute.Attribute.GetAttribute() == Attribute.GetAttribute())
+		if (CellAttribute.Attribute == Attribute)
 		{
 			CellAttributeIndex = CellSize;
 			CellAttributeSize = CurrentFunctionSize;
@@ -943,7 +928,7 @@ void UDMXEntityFixtureType::AddCellAttribute(int32 ModeIndex)
 		FDMXFixtureMatrix& Matrix = Mode.FixtureMatrixConfig;
 
 		FDMXFixtureCellAttribute NewAttribute;
-		TArray<FName> AvailableAttributes = FDMXAttributeName::GetPossibleValues();
+		TArray<FName> AvailableAttributes = FDMXAttributeName::GetPredefinedValues();
 		NewAttribute.Attribute = AvailableAttributes.Num() > 0 ? FDMXAttributeName(AvailableAttributes[0]) : FDMXAttributeName();
 
 		// Disable the matrix while editing it so other functions align when enabling it again
diff --git a/Engine/Plugins/VirtualProduction/DMX/DMXEngine/Source/DMXRuntime/Public/Library/DMXEntityFixtureType.h b/Engine/Plugins/VirtualProduction/DMX/DMXEngine/Source/DMXRuntime/Public/Library/DMXEntityFixtureType.h
--- a/Engine/Plugins/VirtualProduction/DMX/DMXEngine/Source/DMXRuntime/Public/Library/DMXEntityFixtureType.h
+++ b/Engine/Plugins/VirtualProduction/DMX/DMXEngine/Source/DMXRuntime/Public/Library/DMXEntityFixtureType.h
@@ -45,7 +45,7 @@ struct DMXRUNTIME_API FDMXFixtureFunction
 
 	/** Constructor */
 	FDMXFixtureFunction()
-		: Attribute(FDMXNameListItem::None)
+		: Attribute()
 		, FunctionName()
 		, Description()
 		, DefaultValue(0)
@@ -116,7 +116,7 @@ struct DMXRUNTIME_API FDMXFixtureCellAttribute
 
 	/** Constructor */
 	FDMXFixtureCellAttribute()
-		: Attribute(FDMXNameListItem::None)
+		: Attribute()
 		, Description()
 		, DefaultValue(0)
 		, DataType(EDMXFixtureSignalFormat::E8Bit)
diff --git a/Engine/Plugins/VirtualProduction/DMX/DMXProtocol/Source/DMXProtocol/Private/DMXAttribute.cpp b/Engine/Plugins/VirtualProduction/DMX/DMXProtocol/Source/DMXProtocol/Private/DMXAttribute.cpp
--- a/Engine/Plugins/VirtualProduction/DMX/DMXProtocol/Source/DMXProtocol/Private/DMXAttribute.cpp
+++ b/Engine/Plugins/VirtualProduction/DMX/DMXProtocol/Source/DMXProtocol/Private/DMXAttribute.cpp
@@ -6,23 +6,10 @@
 
 #include "Modules/ModuleManager.h"
 
-
-TArray<FName> FDMXAttributeName::GetPossibleValues()
-{
-	TArray<FName> DefaultValues;
-
-	if (const UDMXProtocolSettings* DMXSettings = GetDefault<UDMXProtocolSettings>())
-	{
-		DefaultValues.Reserve(DMXSettings->Attributes.Num());
-
-		for (const FDMXAttribute& Attribute : DMXSettings->Attributes)
-		{
-			DefaultValues.Emplace(Attribute.Name);
-		}
-	}
-
-	return DefaultValues;
-}
+PRAGMA_DISABLE_DEPRECATION_WARNINGS
+const bool FDMXAttributeName::bCanBeNone = true;
+FSimpleMulticastDelegate FDMXAttributeName::OnValuesChanged;
+PRAGMA_ENABLE_DEPRECATION_WARNINGS
 
 FDMXAttributeName::FDMXAttributeName()
 {
@@ -57,14 +44,15 @@ void FDMXAttributeName::SetFromName(const FName& InName)
 	*this = InName;
 }
 
-const FDMXAttribute& FDMXAttributeName::GetAttribute() const
+FDMXAttribute FDMXAttributeName::GetAttribute() const
 {
-	static const FDMXAttribute FailureAttribute;
-
+	FDMXAttribute Attribute;
+	Attribute.Name = Name;
+	return Attribute;
 	const UDMXProtocolSettings* DMXSettings = GetDefault<UDMXProtocolSettings>();
-	if (DMXSettings == nullptr)
+	if (!DMXSettings)
 	{
-		return FailureAttribute;
+		return Attribute;
 	}
 
 	for (const FDMXAttribute& SettingsAttribute : DMXSettings->Attributes)
@@ -75,7 +63,28 @@ const FDMXAttribute& FDMXAttributeName::GetAttribute() const
 		}
 	}
 
-	return FailureAttribute;
+	return Attribute;
+}
+
+TArray<FName> FDMXAttributeName::GetPredefinedValues()
+{
+	TArray<FName> Result;
+	const UDMXProtocolSettings* DMXSettings = GetDefault<UDMXProtocolSettings>();
+	if (!DMXSettings)
+	{
+		return Result;
+	}
+
+	for (const FDMXAttribute& Attribute : DMXSettings->Attributes)
+	{
+		Result.Add(Attribute.Name);
+	}
+	return Result;
+}
+
+TArray<FName> FDMXAttributeName::GetPossibleValues()
+{
+	return GetPredefinedValues();
 }
 
 FString UDMXAttributeNameConversions::Conv_DMXAttributeToString(const FDMXAttributeName& InAttribute)
diff --git a/Engine/Plugins/VirtualProduction/DMX/DMXProtocol/Source/DMXProtocol/Private/DMXProtocolSettings.cpp b/Engine/Plugins/VirtualProduction/DMX/DMXProtocol/Source/DMXProtocol/Private/DMXProtocolSettings.cpp
--- a/Engine/Plugins/VirtualProduction/DMX/DMXProtocol/Source/DMXProtocol/Private/DMXProtocolSettings.cpp
+++ b/Engine/Plugins/VirtualProduction/DMX/DMXProtocol/Source/DMXProtocol/Private/DMXProtocolSettings.cpp
@@ -179,7 +179,7 @@ void UDMXProtocolSettings::PostEditChangeChainProperty(FPropertyChangedChainEven
 			FixtureCategories.Add(TEXT("Other"));
 		}
 
-		FDMXFixtureCategory::OnValuesChanged.Broadcast();
+		OnDefaultFixtureCategoriesChanged.Broadcast();
 	}
 	else if (
 		PropertyName == GET_MEMBER_NAME_CHECKED(UDMXProtocolSettings, Attributes) ||
diff --git a/Engine/Plugins/VirtualProduction/DMX/DMXProtocol/Source/DMXProtocol/Private/DMXProtocolTypes.cpp b/Engine/Plugins/VirtualProduction/DMX/DMXProtocol/Source/DMXProtocol/Private/DMXProtocolTypes.cpp
--- a/Engine/Plugins/VirtualProduction/DMX/DMXProtocol/Source/DMXProtocol/Private/DMXProtocolTypes.cpp
+++ b/Engine/Plugins/VirtualProduction/DMX/DMXProtocol/Source/DMXProtocol/Private/DMXProtocolTypes.cpp
@@ -1,24 +1,15 @@
 // Copyright Epic Games, Inc. All Rights Reserved.
 
 #include "DMXProtocolTypes.h"
-#include "Interfaces/IDMXProtocol.h"
+
 #include "DMXProtocolSettings.h"
+#include "Interfaces/IDMXProtocol.h"
 
-IMPLEMENT_DMX_NAMELISTITEM_STATICVARS(FDMXProtocolName)
 
-IMPLEMENT_DMX_NAMELISTITEM_GetAllValues(FDMXProtocolName)
-{
-	return IDMXProtocol::GetProtocolNames();
-}
-
-IMPLEMENT_DMX_NAMELISTITEM_IsValid(FDMXProtocolName)
-{
-	if (InName.IsNone())
-	{
-		return false;
-	}
-	return IDMXProtocol::Get(InName).IsValid();
-}
+PRAGMA_DISABLE_DEPRECATION_WARNINGS
+const bool FDMXProtocolName::bCanBeNone = false;
+FSimpleMulticastDelegate FDMXProtocolName::OnValuesChanged;
+PRAGMA_ENABLE_DEPRECATION_WARNINGS
 
 FDMXProtocolName::FDMXProtocolName(IDMXProtocolPtr InProtocol)
 {
@@ -57,26 +48,16 @@ IDMXProtocolPtr FDMXProtocolName::GetProtocol() const
 	return IDMXProtocol::Get(Name);
 }
 
-IMPLEMENT_DMX_NAMELISTITEM_STATICVARS(FDMXFixtureCategory)
-
-IMPLEMENT_DMX_NAMELISTITEM_GetAllValues(FDMXFixtureCategory)
+TArray<FName> FDMXProtocolName::GetPossibleValues()
 {
-	return GetDefault<UDMXProtocolSettings>()->FixtureCategories.Array();
+	// DEPRECATED 5.1
+	return IDMXProtocol::GetProtocolNames();
 }
 
-IMPLEMENT_DMX_NAMELISTITEM_IsValid(FDMXFixtureCategory)
-{
-	const TArray<FName>&& AvailableNames = GetPossibleValues();
-	for (const FName& AvailableName : AvailableNames)
-	{
-		if (InName.IsEqual(AvailableName))
-		{
-			return true;
-		}
-	}
-
-	return false;
-}
+PRAGMA_DISABLE_DEPRECATION_WARNINGS
+const bool FDMXFixtureCategory::bCanBeNone = false;
+FSimpleMulticastDelegate FDMXFixtureCategory::OnValuesChanged;
+PRAGMA_ENABLE_DEPRECATION_WARNINGS
 
 FName FDMXFixtureCategory::GetFirstValue()
 {
@@ -90,6 +71,22 @@ FName FDMXFixtureCategory::GetFirstValue()
 	return FName();
 }
 
+TArray<FName> FDMXFixtureCategory::GetPredefinedValues()
+{
+	const UDMXProtocolSettings* ProtocolSettings = GetDefault<UDMXProtocolSettings>();
+	if (ProtocolSettings)
+	{
+		return ProtocolSettings->FixtureCategories.Array();
+	}
+	return TArray<FName>();
+}
+
+TArray<FName> FDMXFixtureCategory::GetPossibleValues()
+{
+	// Deprecated 5.1
+	return GetPredefinedValues();
+}
+
 FDMXFixtureCategory::FDMXFixtureCategory()
 {
 	Name = GetFirstValue();
diff --git a/Engine/Plugins/VirtualProduction/DMX/DMXProtocol/Source/DMXProtocol/Public/DMXAttribute.h b/Engine/Plugins/VirtualProduction/DMX/DMXProtocol/Source/DMXProtocol/Public/DMXAttribute.h
--- a/Engine/Plugins/VirtualProduction/DMX/DMXProtocol/Source/DMXProtocol/Public/DMXAttribute.h
+++ b/Engine/Plugins/VirtualProduction/DMX/DMXProtocol/Source/DMXProtocol/Public/DMXAttribute.h
@@ -52,11 +52,6 @@ struct DMXPROTOCOL_API FDMXAttributeName
 	GENERATED_BODY()
 
 public:
-	//~ Constructors
-
-	/** */
-	static TArray<FName> GetPossibleValues();
-
 	/**
 	 * Default constructor. Assigns Attribute to the first available
 	 * Attribute from the plugin settings
@@ -72,12 +67,31 @@ public:
 	//~ FDMXNameListItem interface
 	void SetFromName(const FName& InName);
 
+	/** IsValid member accessor */
+	bool IsValid() const { return !Name.IsNone(); }
+
 	/** The Attribute Name */
 	UPROPERTY(EditAnywhere, Category = "DMX")
 	FName Name;
 
-	const FDMXAttribute& GetAttribute() const;
-	operator const FDMXAttribute& () const { return GetAttribute(); }
+	/** Returns the predefined values */
+	static TArray<FName> GetPredefinedValues();
+
+	UE_DEPRECATED(5.1, "Deprecated since attributes now can take any name. If it is desired to resolve against default attributes from project settings, do a lookup on UDMXProtocolSettings::Attributes")
+	FDMXAttribute GetAttribute() const;
+
+	//////////////////////////////////////////////////////////////////
+	// Deprecated members originating from deprecated FDMXNameListItem
+	UE_DEPRECATED(5.1, "Please use GetPredefinedValues() instead.")
+	static TArray<FName> GetPossibleValues();
+
+	// Deprecated 5.1 without replacement. Always true
+	UE_DEPRECATED(5.1, "Obsolete. Value 'None' is always allowed.")
+	static const bool bCanBeNone;
+
+	// Deprecated 5.1 without replacement. To listen to changes of default defined protocols, please refer to UDMXProtocolSettings::GetOnDefaultAttributesChanged
+	UE_DEPRECATED(5.1, "To listen to changes of default defined protocols, please refer to UDMXProtocolSettings::GetOnDefaultAttributesChanged().")
+	static FSimpleMulticastDelegate OnValuesChanged;
 };
 
 inline uint32 GetTypeHash(const FDMXAttributeName& DMXNameListItem)
diff --git a/Engine/Plugins/VirtualProduction/DMX/DMXProtocol/Source/DMXProtocol/Public/DMXNameListItem.h b/Engine/Plugins/VirtualProduction/DMX/DMXProtocol/Source/DMXProtocol/Public/DMXNameListItem.h
--- a/Engine/Plugins/VirtualProduction/DMX/DMXProtocol/Source/DMXProtocol/Public/DMXNameListItem.h
+++ b/Engine/Plugins/VirtualProduction/DMX/DMXProtocol/Source/DMXProtocol/Public/DMXNameListItem.h
@@ -81,6 +81,7 @@ bool TypeName::IsValid(const FName& InName)
  * - inherit this struct and override the virtual functions.
  * - use the DMX_NAMELISTITEM_STATICS macros above to declare the expected statics.
  */
+struct UE_DEPRECATED(5.1, "Removed to overcome restrictions of the generic base struct. Ustructs that have a Name property should declare it themselves instead.") FDMXNameListItem;
 USTRUCT()
 struct DMXPROTOCOL_API FDMXNameListItem
 {
@@ -106,4 +107,4 @@ struct DMXPROTOCOL_API FDMXNameListItem
 
 	/** Set this object to represent a None value */
 	virtual void SetToNone() { Name = NAME_None; }
-};
\ No newline at end of file
+};
diff --git a/Engine/Plugins/VirtualProduction/DMX/DMXProtocol/Source/DMXProtocol/Public/DMXProtocolSettings.h b/Engine/Plugins/VirtualProduction/DMX/DMXProtocol/Source/DMXProtocol/Public/DMXProtocolSettings.h
--- a/Engine/Plugins/VirtualProduction/DMX/DMXProtocol/Source/DMXProtocol/Public/DMXProtocolSettings.h
+++ b/Engine/Plugins/VirtualProduction/DMX/DMXProtocol/Source/DMXProtocol/Public/DMXProtocolSettings.h
@@ -117,6 +117,12 @@ public:
 		return OnDefaultAttributesChanged;
 	}
 
+	/** Gets a Delegate Broadcast when default fixture categories were changed */
+	FSimpleMulticastDelegate& GetOnDefaultFixtureCategoriesChanged()
+	{
+		return OnDefaultFixtureCategoriesChanged;
+	}
+
 #if WITH_EDITOR
 	/** Returns true if Fixture Patches should receive DMX in Editor */
 	bool ShouldAllFixturePatchesReceiveDMXInEditor() const { return bAllFixturePatchesReceiveDMXInEditor; }
@@ -171,6 +177,9 @@ private:
 	/** Broadcast when default attributets changed */
 	FSimpleMulticastDelegate OnDefaultAttributesChanged;
 
+	/** Broadcast when default fixture categories changed */
+	FSimpleMulticastDelegate OnDefaultFixtureCategoriesChanged;
+
 #if WITH_EDITORONLY_DATA
 	/** If true, all fixture patches receive DMX in Editor. This overrides the fixture patches 'Receive DMX In Editor' property. */
 	UPROPERTY(Config, EditAnywhere, Category = "DMX|Communication Settings", Meta = (DisplayName = "All Fixture Patches receive DMX in Editor"))
diff --git a/Engine/Plugins/VirtualProduction/DMX/DMXProtocol/Source/DMXProtocol/Public/DMXProtocolTypes.h b/Engine/Plugins/VirtualProduction/DMX/DMXProtocol/Source/DMXProtocol/Public/DMXProtocolTypes.h
--- a/Engine/Plugins/VirtualProduction/DMX/DMXProtocol/Source/DMXProtocol/Public/DMXProtocolTypes.h
+++ b/Engine/Plugins/VirtualProduction/DMX/DMXProtocol/Source/DMXProtocol/Public/DMXProtocolTypes.h
@@ -6,7 +6,6 @@
 #include "DMXProtocolConstants.h"
 #include "DMXProtocolLog.h"
 #include "DMXProtocolMacros.h"
-#include "DMXNameListItem.h"
 
 #include "Dom/JsonObject.h"
 #include "HAL/CriticalSection.h"
@@ -99,13 +98,10 @@ enum class EDMXFixtureSignalFormat : uint8
 /** A DMX protocol as a name that can be displayed in UI. The protocol is directly accessible via GetProtocol */
 USTRUCT(BlueprintType, Category = "DMX")
 struct DMXPROTOCOL_API FDMXProtocolName
-	: public FDMXNameListItem
 {
 public:
 	GENERATED_BODY()
 
-	DECLARE_DMX_NAMELISTITEM_STATICS(false)
-
 	FDMXProtocolName();
 
 	/** Construct from a protocol name */
@@ -118,7 +114,7 @@ public:
 	IDMXProtocolPtr GetProtocol() const;
 
 	/** IsValid member accessor */
-	bool IsValid() const { return IsValid(Name); }
+	bool IsValid() const { return !Name.IsNone(); }
 
 	//~ FName operators
 	operator FName&() { return Name; }
@@ -129,7 +125,12 @@ public:
 	operator const IDMXProtocolPtr() const { return GetProtocol(); }
 
 	/** Bool (is valid) operator */
-	operator bool() const { return IsValid(Name); }
+	UE_DEPRECATED(5.1, "Please use IsValid() instead.")
+	operator bool() const { return !Name.IsNone(); }
+
+	/** The Protocol Name */
+	UPROPERTY(EditAnywhere, Category = "DMX")
+	FName Name;
 
 	//~ Comparison operators
 	bool operator==(const FDMXProtocolName& Other) const { return Name == Other.Name; }
@@ -138,34 +139,65 @@ public:
 	bool operator!=(const IDMXProtocolPtr& Other) const { return GetProtocol().Get() != Other.Get(); }
 	bool operator==(const FName& Other) const { return Name == Other; }
 	bool operator!=(const FName& Other) const { return !Name.IsEqual(Other); }
+
+	//////////////////////////////////////////////////////////////////
+	// Deprecated members originating from deprecated FDMXNameListItem
+	UE_DEPRECATED(5.1, "Instead please use IDMXProtocol::GetProtocolNames")
+	static TArray<FName> GetPossibleValues();
+
+	// Deprecated without replacement. Always false.
+	UE_DEPRECATED(5.1, "Obsolete. The value never can and never should be none in any case.")
+	static const bool bCanBeNone;
+
+	// Deprecated without replacement. Protocols cannot be enabled or disabled after engine startup.
+	UE_DEPRECATED(5.1, "Obsolete. Protocols cannot be enabled or disabled after engine startup.")
+	static FSimpleMulticastDelegate OnValuesChanged;
 };
 
 /** Category of a fixture */
 USTRUCT(BlueprintType, Category = "DMX")
 struct DMXPROTOCOL_API FDMXFixtureCategory
-	: public FDMXNameListItem
 {
 public:
 	GENERATED_BODY()
 
-	DECLARE_DMX_NAMELISTITEM_STATICS(false)
-
 	static FName GetFirstValue();
 
 	FDMXFixtureCategory();
 
 	explicit FDMXFixtureCategory(const FName& InName);
+	
+	/** The Protocol Name */
+	UPROPERTY(EditAnywhere, Category = "DMX")
+	FName Name;
+
+	/** Returns the predefined values */
+	static TArray<FName> GetPredefinedValues();
 
 	//~ FName operators
 	operator FName&() { return Name; }
 	operator const FName&() const { return Name; }
 
 	/** Bool (is valid) operator */
-	operator bool() const { return IsValid(Name); }
+	UE_DEPRECATED(5.1, "Please use IsValid() instead.")
+	operator bool() const { return !Name.IsNone(); }
 
 	//~ Comparison operators
 	bool operator==(const FDMXFixtureCategory& Other) const { return Name == Other.Name; }
 	bool operator==(const FName& Other) const { return Name == Other; }
+
+	//////////////////////////////////////////////////////////////////
+	// Deprecated members originating from deprecated FDMXNameListItem
+	UE_DEPRECATED(5.1, "Please use GetPredefinedValues() instead")
+	static TArray<FName> GetPossibleValues();
+
+	// Deprecated without replacement. Always false
+	UE_DEPRECATED(5.1, "Obsolete. The value never can and never should be none in any case.")
+	static const bool bCanBeNone;
+
+	// Deprecated without replacement. Protocols cannot be enabled or disabled after engine startup.
+	UE_DEPRECATED(5.1, "Obsolete. To listen to default fixture category changes, please refer to UDMXProtocolSettings::GetOnDefaultFixtureCategoriesChanged().")
+	static FSimpleMulticastDelegate OnValuesChanged;
 };
 
 UCLASS()
@@ -188,6 +220,3 @@ public:
 	static FName Conv_DMXFixtureCategoryToName(const FDMXFixtureCategory& InFixtureCategory);
 };
 
-
-
-
