DMX - Don't tick destroyed DMX components marked RF_Transactional

diff --git a/Engine/Plugins/VirtualProduction/DMX/DMXEngine/Source/DMXRuntime/Private/Game/DMXComponent.cpp b/Engine/Plugins/VirtualProduction/DMX/DMXEngine/Source/DMXRuntime/Private/Game/DMXComponent.cpp
--- a/Engine/Plugins/VirtualProduction/DMX/DMXEngine/Source/DMXRuntime/Private/Game/DMXComponent.cpp
+++ b/Engine/Plugins/VirtualProduction/DMX/DMXEngine/Source/DMXRuntime/Private/Game/DMXComponent.cpp
@@ -114,13 +114,9 @@ void UDMXComponent::UpdateTickEnabled()
 	FEditorScriptExecutionGuard ScriptGuard;
 #endif
 
-	if (!OnDMXComponentTick.IsBound())
-	{
-		SetComponentTickEnabled(false);
-		return;
-	}
-
-	if (!bReceiveDMXFromPatch)
+	if (IsBeingDestroyed() ||
+		!OnDMXComponentTick.IsBound() ||
+		!bReceiveDMXFromPatch)
 	{
 		SetComponentTickEnabled(false);
 		return;
