DMX - Fix Attributes use the 'Name' instead of the 'Attribute' attribute as their name

diff --git a/Engine/Plugins/VirtualProduction/DMX/DMXEngine/Source/DMXEditor/Private/DMXInitializeFixtureTypeFromGDTFHelper.cpp b/Engine/Plugins/VirtualProduction/DMX/DMXEngine/Source/DMXEditor/Private/DMXInitializeFixtureTypeFromGDTFHelper.cpp
--- a/Engine/Plugins/VirtualProduction/DMX/DMXEngine/Source/DMXEditor/Private/DMXInitializeFixtureTypeFromGDTFHelper.cpp
+++ b/Engine/Plugins/VirtualProduction/DMX/DMXEngine/Source/DMXEditor/Private/DMXInitializeFixtureTypeFromGDTFHelper.cpp
@@ -157,18 +157,11 @@ namespace UE::DMX::DMXInitializeFixtureTypeFromGDTFHelper::Private
 				return false;
 			}
 
-			constexpr TCHAR ChannelFunctionTag[] = TEXT("ChannelFunction");
-			const FXmlNode* ChannelFunctionNode = LogicalChannelNode->FindChildNode(ChannelFunctionTag);
-			if (!ChannelFunctionNode)
-			{
-				return false;
-			}
-
 			// Name
-			constexpr TCHAR NameTag[] = TEXT("Name");
-			FString Name = ChannelFunctionNode->GetAttribute(NameTag);
+			constexpr TCHAR AttributeTag[] = TEXT("Attribute");
+			FString AttributeName = LogicalChannelNode->GetAttribute(AttributeTag);
 			
-			OutName = Name;
+			OutName = AttributeName;
 			return true;
 		}
 
