DMX - Duplicating a fixture from the Fixture Patch now keeps selection correctly

diff --git a/Engine/Plugins/VirtualProduction/DMX/DMXEngine/Source/DMXEditor/Private/Widgets/FixturePatch/SDMXMVRFixtureList.cpp b/Engine/Plugins/VirtualProduction/DMX/DMXEngine/Source/DMXEditor/Private/Widgets/FixturePatch/SDMXMVRFixtureList.cpp
--- a/Engine/Plugins/VirtualProduction/DMX/DMXEngine/Source/DMXEditor/Private/Widgets/FixturePatch/SDMXMVRFixtureList.cpp
+++ b/Engine/Plugins/VirtualProduction/DMX/DMXEngine/Source/DMXEditor/Private/Widgets/FixturePatch/SDMXMVRFixtureList.cpp
@@ -391,12 +391,6 @@ void SDMXMVRFixtureList::Construct(const FArguments& InArgs, TWeakPtr<FDMXEditor
 
 	RegisterCommands();
 	RefreshList();
-
-	// Make an initial selection, as if the user clicked it
-	if (ListSource.Num() > 0)
-	{
-		ListView->SetSelection(ListSource[0], ESelectInfo::OnMouseClick);
-	}
 }
 
 void SDMXMVRFixtureList::RequestListRefresh()
@@ -1099,7 +1093,6 @@ void SDMXMVRFixtureList::OnPasteItems()
 		}
 
 		RequestListRefresh();
-		AdoptSelectionFromFixturePatchSharedData();
 	}
 }
 
@@ -1118,7 +1111,6 @@ void SDMXMVRFixtureList::OnDuplicateItems()
 	FDMXMVRFixtureListItem::DuplicateItems(WeakDMXEditor, SelectedItems);
 
 	RequestListRefresh();
-	AdoptSelectionFromFixturePatchSharedData();
 }
 
 bool SDMXMVRFixtureList::CanDeleteItems() const
