DMX - MVR Import now properly generates fixture types from GDTF

diff --git a/Engine/Plugins/VirtualProduction/DMX/DMXEngine/Source/DMXEditor/Private/Factories/DMXLibraryFromMVRFactory.cpp b/Engine/Plugins/VirtualProduction/DMX/DMXEngine/Source/DMXEditor/Private/Factories/DMXLibraryFromMVRFactory.cpp
--- a/Engine/Plugins/VirtualProduction/DMX/DMXEngine/Source/DMXEditor/Private/Factories/DMXLibraryFromMVRFactory.cpp
+++ b/Engine/Plugins/VirtualProduction/DMX/DMXEngine/Source/DMXEditor/Private/Factories/DMXLibraryFromMVRFactory.cpp
@@ -6,6 +6,7 @@
 #include "DMXEditorModule.h"
 #include "DMXEditorSettings.h"
 #include "DMXEditorUtils.h"
+#include "DMXInitializeFixtureTypeFromGDTFHelper.h"
 #include "DMXZipper.h"
 #include "Factories/DMXGDTFFactory.h"
 #include "Library/DMXEntityFixturePatch.h"
@@ -254,7 +255,12 @@ void UDMXLibraryFromMVRFactory::InitDMXLibrary(UDMXLibrary* DMXLibrary, const TA
 		FixtureTypeConstructionParams.ParentDMXLibrary = DMXLibrary;
 
 		UDMXEntityFixtureType* FixtureType = UDMXEntityFixtureType::CreateFixtureTypeInLibrary(FixtureTypeConstructionParams, FPaths::GetBaseFilename(GDTFFilename));
-		FixtureType->SetModesFromDMXImport(GDTF);
+		const bool bAdvancedImportSuccess = FDMXInitializeFixtureTypeFromGDTFHelper::GenerateModesFromGDTF(*FixtureType, *GDTF);
+		if (!bAdvancedImportSuccess)
+		{
+			UE_LOG(LogDMXEditor, Warning, TEXT("Failed to initialize Fixture Type '%s', falling back to legacy method that doesn't support matrix fixtures."), *FixtureType->GetName());
+			FixtureType->SetModesFromDMXImport(GDTF);
+		}
 
 		GDTFSpecToFixtureTypeMap.Add(GDTFFilename, FixtureType);
 	}
