DMX - Fix fixture patches show their unique object name instead of their beautiful name

diff --git a/Engine/Plugins/VirtualProduction/DMX/DMXEngine/Source/DMXEditor/Private/Widgets/FixturePatch/DMXMVRFixtureListItem.cpp b/Engine/Plugins/VirtualProduction/DMX/DMXEngine/Source/DMXEditor/Private/Widgets/FixturePatch/DMXMVRFixtureListItem.cpp
--- a/Engine/Plugins/VirtualProduction/DMX/DMXEngine/Source/DMXEditor/Private/Widgets/FixturePatch/DMXMVRFixtureListItem.cpp
+++ b/Engine/Plugins/VirtualProduction/DMX/DMXEngine/Source/DMXEditor/Private/Widgets/FixturePatch/DMXMVRFixtureListItem.cpp
@@ -81,6 +81,7 @@ void FDMXMVRFixtureListItem::SetFixturePatchName(const FString& InDesiredName, F
 	{
 		if (FixturePatch->Name == InDesiredName)
 		{
+			OutNewName = InDesiredName;
 			return;
 		}
 
diff --git a/Engine/Plugins/VirtualProduction/DMX/DMXEngine/Source/DMXRuntime/Private/Library/DMXEntity.cpp b/Engine/Plugins/VirtualProduction/DMX/DMXEngine/Source/DMXRuntime/Private/Library/DMXEntity.cpp
--- a/Engine/Plugins/VirtualProduction/DMX/DMXEngine/Source/DMXRuntime/Private/Library/DMXEntity.cpp
+++ b/Engine/Plugins/VirtualProduction/DMX/DMXEngine/Source/DMXRuntime/Private/Library/DMXEntity.cpp
@@ -48,7 +48,7 @@ void UDMXEntity::SetName(const FString& InNewName)
 {
 	const FName UniqueName = MakeUniqueObjectName(GetOuter(), UDMXEntity::StaticClass(), *InNewName);
 	Rename(*UniqueName.ToString(), GetOuter());
-	Name = UniqueName.ToString();
+	Name = InNewName;
 }
 
 void UDMXEntity::SetParentLibrary(UDMXLibrary* InParent)
