DMX - Fix MVR generated in MA is not imported properly

diff --git a/Engine/Plugins/VirtualProduction/DMX/DMXEngine/Source/DMXRuntime/Private/MVR/Types/DMXMVRFixtureNode.cpp b/Engine/Plugins/VirtualProduction/DMX/DMXEngine/Source/DMXRuntime/Private/MVR/Types/DMXMVRFixtureNode.cpp
--- a/Engine/Plugins/VirtualProduction/DMX/DMXEngine/Source/DMXRuntime/Private/MVR/Types/DMXMVRFixtureNode.cpp
+++ b/Engine/Plugins/VirtualProduction/DMX/DMXEngine/Source/DMXRuntime/Private/MVR/Types/DMXMVRFixtureNode.cpp
@@ -7,6 +7,7 @@
 #include "DMXRuntimeUtils.h"
 
 #include "XmlNode.h"
+#include "Misc/Paths.h"
 
 
 namespace UE::DMXRuntime::DMXMVRFixture::Private
@@ -60,6 +61,13 @@ namespace UE::DMXRuntime::DMXMVRFixture::Private
 		if (GDTFSpecNode)
 		{
 			OutGDTFSpec = GDTFSpecNode->GetContent();
+
+			// Append the gdtf extension if it's missing.
+			if (FPaths::GetExtension(OutGDTFSpec) != TEXT(".gdtf"))
+			{
+				OutGDTFSpec += TEXT(".gdtf");
+			}
+
 			return true;
 		}
 
