DMX - Fix MVR Fixture Node's GDTFSpec may add one .gdtf extension too many in some cases

diff --git a/Engine/Plugins/VirtualProduction/DMX/DMXEngine/Source/DMXRuntime/Private/MVR/Types/DMXMVRFixtureNode.cpp b/Engine/Plugins/VirtualProduction/DMX/DMXEngine/Source/DMXRuntime/Private/MVR/Types/DMXMVRFixtureNode.cpp
--- a/Engine/Plugins/VirtualProduction/DMX/DMXEngine/Source/DMXRuntime/Private/MVR/Types/DMXMVRFixtureNode.cpp
+++ b/Engine/Plugins/VirtualProduction/DMX/DMXEngine/Source/DMXRuntime/Private/MVR/Types/DMXMVRFixtureNode.cpp
@@ -63,7 +63,7 @@ namespace UE::DMXRuntime::DMXMVRFixture::Private
 			OutGDTFSpec = GDTFSpecNode->GetContent();
 
 			// Append the gdtf extension if it's missing.
-			if (FPaths::GetExtension(OutGDTFSpec) != TEXT(".gdtf"))
+			if (FPaths::GetExtension(OutGDTFSpec) != TEXT("gdtf"))
 			{
 				OutGDTFSpec += TEXT(".gdtf");
 			}
