DMX - MVR import does not set the GDTF on the fixture types correctly

diff --git a/Engine/Plugins/VirtualProduction/DMX/DMXEngine/Source/DMXEditor/Private/Factories/DMXLibraryFromMVRFactory.cpp b/Engine/Plugins/VirtualProduction/DMX/DMXEngine/Source/DMXEditor/Private/Factories/DMXLibraryFromMVRFactory.cpp
--- a/Engine/Plugins/VirtualProduction/DMX/DMXEngine/Source/DMXEditor/Private/Factories/DMXLibraryFromMVRFactory.cpp
+++ b/Engine/Plugins/VirtualProduction/DMX/DMXEngine/Source/DMXEditor/Private/Factories/DMXLibraryFromMVRFactory.cpp
@@ -261,6 +261,7 @@ void UDMXLibraryFromMVRFactory::InitDMXLibrary(UDMXLibrary* DMXLibrary, const TA
 			UE_LOG(LogDMXEditor, Warning, TEXT("Failed to initialize Fixture Type '%s', falling back to legacy method that doesn't support matrix fixtures."), *FixtureType->GetName());
 			FixtureType->SetModesFromDMXImport(GDTF);
 		}
+		FixtureType->DMXImport = GDTF;
 
 		GDTFSpecToFixtureTypeMap.Add(GDTFFilename, FixtureType);
 	}
