DMX - Fix shortcut are not affecting fixture selected from the right hand panel of Fixture Patch, improve performance when duplicating fixture patches

diff --git a/Engine/Plugins/VirtualProduction/DMX/DMXEngine/Source/DMXEditor/Private/Tabs/DMXEditorTabFactories.cpp b/Engine/Plugins/VirtualProduction/DMX/DMXEngine/Source/DMXEditor/Private/Tabs/DMXEditorTabFactories.cpp
--- a/Engine/Plugins/VirtualProduction/DMX/DMXEngine/Source/DMXEditor/Private/Tabs/DMXEditorTabFactories.cpp
+++ b/Engine/Plugins/VirtualProduction/DMX/DMXEngine/Source/DMXEditor/Private/Tabs/DMXEditorTabFactories.cpp
@@ -42,12 +42,7 @@ TSharedRef<SWidget> FDMXLibraryEditorTabSummoner::CreateTabBody(const FWorkflowT
 {
 	TSharedPtr<FDMXEditor> DMXEditor = StaticCastSharedPtr<FDMXEditor>(HostingApp.Pin());
 
-	return SNew(SBox)
-		.HAlign(HAlign_Fill)
-		.VAlign(VAlign_Fill)
-		[
-			DMXEditor->GetDMXLibraryEditorTab()
-		];
+	return DMXEditor->GetDMXLibraryEditorTab();
 }
 
 FText FDMXLibraryEditorTabSummoner::GetTabToolTipText(const FWorkflowTabSpawnInfo& Info) const
@@ -71,12 +66,7 @@ TSharedRef<SWidget> FDMXEditorFixtureTypesSummoner::CreateTabBody(const FWorkflo
 {
 	TSharedPtr<FDMXEditor> DMXEditor = StaticCastSharedPtr<FDMXEditor>(HostingApp.Pin());
 
-	return SNew(SBox)
-		.HAlign(HAlign_Fill)
-		.VAlign(VAlign_Fill)
-		[
-			DMXEditor->GetFixtureTypeEditor()
-		];
+	return DMXEditor->GetFixtureTypeEditor();
 }
 
 FDMXEditorFixturePatchSummoner::FDMXEditorFixturePatchSummoner(TSharedPtr<FAssetEditorToolkit> InHostingApp)
@@ -95,12 +85,7 @@ TSharedRef<SWidget> FDMXEditorFixturePatchSummoner::CreateTabBody(const FWorkflo
 {
 	TSharedPtr<FDMXEditor> DMXEditor = StaticCastSharedPtr<FDMXEditor>(HostingApp.Pin());
 
-	return SNew(SBox)
-		.HAlign(HAlign_Fill)
-		.VAlign(VAlign_Fill)
-		[
-			DMXEditor->GetFixturePatchEditor()
-		];
+	return DMXEditor->GetFixturePatchEditor();
 }
 
 #undef LOCTEXT_NAMESPACE
diff --git a/Engine/Plugins/VirtualProduction/DMX/DMXEngine/Source/DMXEditor/Private/Widgets/FixturePatch/DMXMVRFixtureListItem.cpp b/Engine/Plugins/VirtualProduction/DMX/DMXEngine/Source/DMXEditor/Private/Widgets/FixturePatch/DMXMVRFixtureListItem.cpp
--- a/Engine/Plugins/VirtualProduction/DMX/DMXEngine/Source/DMXEditor/Private/Widgets/FixturePatch/DMXMVRFixtureListItem.cpp
+++ b/Engine/Plugins/VirtualProduction/DMX/DMXEngine/Source/DMXEditor/Private/Widgets/FixturePatch/DMXMVRFixtureListItem.cpp
@@ -10,6 +10,7 @@
 #include "MVR/Types/DMXMVRFixtureNode.h"
 
 #include "ScopedTransaction.h"
+#include "Algo/MinElement.h"
 
 
 #define LOCTEXT_NAMESPACE "DMXMVRFixtureListItem"
@@ -408,7 +409,8 @@ void FDMXMVRFixtureListItem::DuplicateFixturePatchesInternal(TWeakPtr<FDMXEditor
 			FixtureTypeConstructionParams.Modes = FixtureTypeOfPatchToDuplicate->Modes;
 			FixtureTypeConstructionParams.ParentDMXLibrary = DMXLibrary;
 
-			UDMXEntityFixtureType::CreateFixtureTypeInLibrary(FixtureTypeConstructionParams, FixtureTypeOfPatchToDuplicate->Name);
+			constexpr bool bMarkLibraryDirty = false;
+			UDMXEntityFixtureType::CreateFixtureTypeInLibrary(FixtureTypeConstructionParams, FixtureTypeOfPatchToDuplicate->Name, bMarkLibraryDirty);
 		}
 
 		// Duplicate the Fixture Patch
@@ -425,7 +427,8 @@ void FDMXMVRFixtureListItem::DuplicateFixturePatchesInternal(TWeakPtr<FDMXEditor
 		ConstructionParams.UniverseID = Universe;
 		ConstructionParams.StartingAddress = Address;
 
-		UDMXEntityFixturePatch* NewFixturePatch = UDMXEntityFixturePatch::CreateFixturePatchInLibrary(ConstructionParams, FixturePatchToDuplicate->Name);
+		constexpr bool bMarkLibraryDirty = false;
+		UDMXEntityFixturePatch* NewFixturePatch = UDMXEntityFixturePatch::CreateFixturePatchInLibrary(ConstructionParams, FixturePatchToDuplicate->Name, bMarkLibraryDirty);
 		NewFixturePatches.Add(NewFixturePatch);
 
 		Address += ChannelSpan;
@@ -435,6 +438,15 @@ void FDMXMVRFixtureListItem::DuplicateFixturePatchesInternal(TWeakPtr<FDMXEditor
 
 	if (NewFixturePatches.Num() > 0)
 	{
+		// Select universe and fixture patch
+		const TWeakObjectPtr<UDMXEntityFixturePatch>* FirstFixturePatchPtr = Algo::MinElementBy(NewFixturePatches, [](TWeakObjectPtr<UDMXEntityFixturePatch> FixturePatch)
+			{
+				return FixturePatch->GetUniverseID();
+			});
+		if (FirstFixturePatchPtr->IsValid())
+		{
+			FixturePatchSharedData->SelectUniverse(FirstFixturePatchPtr->Get()->GetUniverseID());
+		}
 		FixturePatchSharedData->SelectFixturePatches(NewFixturePatches);
 	}
 }
diff --git a/Engine/Plugins/VirtualProduction/DMX/DMXEngine/Source/DMXEditor/Private/Widgets/FixturePatch/SDMXFixturePatchEditor.cpp b/Engine/Plugins/VirtualProduction/DMX/DMXEngine/Source/DMXEditor/Private/Widgets/FixturePatch/SDMXFixturePatchEditor.cpp
--- a/Engine/Plugins/VirtualProduction/DMX/DMXEngine/Source/DMXEditor/Private/Widgets/FixturePatch/SDMXFixturePatchEditor.cpp
+++ b/Engine/Plugins/VirtualProduction/DMX/DMXEngine/Source/DMXEditor/Private/Widgets/FixturePatch/SDMXFixturePatchEditor.cpp
@@ -29,7 +29,6 @@ void SDMXFixturePatchEditor::Construct(const FArguments& InArgs)
 	FixturePatchSharedData = DMXEditorPtr.Pin()->GetFixturePatchSharedData();
 
 	SetCanTick(false);
-	bCanSupportFocus = false;
 
 	FixturePatchDetailsView = GenerateFixturePatchDetailsView();
 
@@ -76,6 +75,12 @@ void SDMXFixturePatchEditor::Construct(const FArguments& InArgs)
 	FixturePatchSharedData->OnFixturePatchSelectionChanged.AddSP(this, &SDMXFixturePatchEditor::OnFixturePatchesSelected);
 }
 
+FReply SDMXFixturePatchEditor::OnKeyDown(const FGeometry& MyGeometry, const FKeyEvent& InKeyEvent)
+{
+	MVRFixtureList->ProcessCommandBindings(InKeyEvent);
+	return FReply::Handled();
+}
+
 void SDMXFixturePatchEditor::RequestRenameOnNewEntity(const UDMXEntity* InEntity, ESelectInfo::Type SelectionType)
 {
 	if (MVRFixtureList.IsValid())
diff --git a/Engine/Plugins/VirtualProduction/DMX/DMXEngine/Source/DMXEditor/Private/Widgets/FixturePatch/SDMXFixturePatchEditor.h b/Engine/Plugins/VirtualProduction/DMX/DMXEngine/Source/DMXEditor/Private/Widgets/FixturePatch/SDMXFixturePatchEditor.h
--- a/Engine/Plugins/VirtualProduction/DMX/DMXEngine/Source/DMXEditor/Private/Widgets/FixturePatch/SDMXFixturePatchEditor.h
+++ b/Engine/Plugins/VirtualProduction/DMX/DMXEngine/Source/DMXEditor/Private/Widgets/FixturePatch/SDMXFixturePatchEditor.h
@@ -35,12 +35,17 @@ public:
 	void Construct(const FArguments& InArgs);
 
 public:
-	/** Begin SDMXEntityEditorTab interface */
+	//~ Begin SWidget interface
+	virtual bool SupportsKeyboardFocus() const override { return true; }
+	virtual FReply OnKeyDown(const FGeometry& MyGeometry, const FKeyEvent& InKeyEvent) override;
+	//~ End SWidget interface
+
+	// Begin SDMXEntityEditorTab interface
 	void RequestRenameOnNewEntity(const UDMXEntity* InEntity, ESelectInfo::Type SelectionType);
 	void SelectEntity(UDMXEntity* InEntity, ESelectInfo::Type InSelectionType = ESelectInfo::Type::Direct);
 	void SelectEntities(const TArray<UDMXEntity*>& InEntities, ESelectInfo::Type SelectionType = ESelectInfo::Type::Direct);
 	TArray<UDMXEntity*> GetSelectedEntities() const;
-	/** ~End SDMXEntityEditorTab interface */
+	// ~End SDMXEntityEditorTab interface 
 
 private:
 	/** Selects the patch */
diff --git a/Engine/Plugins/VirtualProduction/DMX/DMXEngine/Source/DMXEditor/Private/Widgets/FixturePatch/SDMXFixturePatchFragment.cpp b/Engine/Plugins/VirtualProduction/DMX/DMXEngine/Source/DMXEditor/Private/Widgets/FixturePatch/SDMXFixturePatchFragment.cpp
--- a/Engine/Plugins/VirtualProduction/DMX/DMXEngine/Source/DMXEditor/Private/Widgets/FixturePatch/SDMXFixturePatchFragment.cpp
+++ b/Engine/Plugins/VirtualProduction/DMX/DMXEngine/Source/DMXEditor/Private/Widgets/FixturePatch/SDMXFixturePatchFragment.cpp
@@ -174,7 +174,6 @@ void SDMXFixturePatchFragment::UpdateIsTopmost()
 	if (FixturePatchNodeGroup.Num() > 1)
 	{
 		UpdateFixturePatchNameText();
-		UE_LOG(LogTemp, Warning, TEXT("%s: I'm top!"), *FixturePatchNameText.ToString());
 	}
 }
 
diff --git a/Engine/Plugins/VirtualProduction/DMX/DMXEngine/Source/DMXEditor/Private/Widgets/FixturePatch/SDMXFixturePatcher.cpp b/Engine/Plugins/VirtualProduction/DMX/DMXEngine/Source/DMXEditor/Private/Widgets/FixturePatch/SDMXFixturePatcher.cpp
--- a/Engine/Plugins/VirtualProduction/DMX/DMXEngine/Source/DMXEditor/Private/Widgets/FixturePatch/SDMXFixturePatcher.cpp
+++ b/Engine/Plugins/VirtualProduction/DMX/DMXEngine/Source/DMXEditor/Private/Widgets/FixturePatch/SDMXFixturePatcher.cpp
@@ -162,12 +162,6 @@ void SDMXFixturePatcher::Construct(const FArguments& InArgs)
 		SharedData->SelectUniverse(Patches[0]->GetUniverseID());
 	}		
 
-	// Bind to tabs being switched
-	FGlobalTabmanager::Get()->OnActiveTabChanged_Subscribe(FOnActiveTabChanged::FDelegate::CreateSP(this, &SDMXFixturePatcher::OnActiveTabChanged));
-
-
-	GEditor->RegisterForUndo(this);
-
 	ShowSelectedUniverse();
 
 }
@@ -183,28 +177,18 @@ void SDMXFixturePatcher::RefreshFromLibrary()
 	RefreshFixturePatchState = EDMXRefreshFixturePatcherState::RefreshFromLibrary;
 }
 
-void SDMXFixturePatcher::OnActiveTabChanged(TSharedPtr<SDockTab> PreviouslyActive, TSharedPtr<SDockTab> NewlyActivated)
-{
-	if (!NewlyActivated.IsValid() ||
-		NewlyActivated->GetLayoutIdentifier().TabType == FDMXEditorTabNames::DMXFixturePatchEditor)
-	{
-		RefreshFromLibrary();
-	}
-}
-
 void SDMXFixturePatcher::Tick(const FGeometry& AllottedGeometry, const double InCurrentTime, const float InDeltaTime)
 {
 	// Refresh if requested
 	if (RefreshFixturePatchState != EDMXRefreshFixturePatcherState::NoRefreshRequested)
 	{
-		const bool bForceReconstructWidget = RefreshFixturePatchState == EDMXRefreshFixturePatcherState::RefreshFromLibrary;
 		if (IsUniverseSelectionEnabled())
 		{
-			ShowSelectedUniverse(bForceReconstructWidget);
+			ShowSelectedUniverse();
 		}
 		else
 		{
-			ShowAllPatchedUniverses(bForceReconstructWidget);
+			ShowAllPatchedUniverses();
 		}
 
 		RefreshFixturePatchState = EDMXRefreshFixturePatcherState::NoRefreshRequested;
@@ -480,43 +464,40 @@ void SDMXFixturePatcher::OnFixturePatchSelectionChanged()
 		return;
 	}
 
-	if (IsUniverseSelectionEnabled())
+	const TSharedPtr<SDMXPatchedUniverse>* const UniverseWidgetPtr = PatchedUniversesByID.Find(SelectedUniverse);
+	if (IsUniverseSelectionEnabled() || !UniverseWidgetPtr)
 	{
-		constexpr bool bReconstructWidget = true;
-		ShowSelectedUniverse(bReconstructWidget);
+		// Refresh if all universes are displayed, or if the Universe was empty previously and now needs to be displayed.
+		ShowSelectedUniverse();
 	}
 	else
 	{
-		const TSharedPtr<SDMXPatchedUniverse>* const UniverseWidgetPtr = PatchedUniversesByID.Find(SelectedUniverse);
-		if (UniverseWidgetPtr)
-		{
-			PatchedUniverseScrollBox->ScrollDescendantIntoView(*UniverseWidgetPtr);
-		}
+		PatchedUniverseScrollBox->ScrollDescendantIntoView(*UniverseWidgetPtr);
 	}
 }
 
 void SDMXFixturePatcher::OnUniverseSelectionChanged()
 {
-	const int32 SelectedUniverse = SharedData->GetSelectedUniverse();
-	if (PatchedUniversesByID.Contains(SelectedUniverse))
-	{
-		return;
-	}
-
 	if (IsUniverseSelectionEnabled())
 	{
-		ShowSelectedUniverse();
+		GEditor->GetTimerManager()->SetTimerForNextTick(FTimerDelegate::CreateLambda([this]()
+			{
+				ShowSelectedUniverse();
+			}));
 	}
-	else if(!PatchedUniversesByID.Contains(SharedData->GetSelectedUniverse()))
+	else 
 	{
-		ShowAllPatchedUniverses();
-	}
+		GEditor->GetTimerManager()->SetTimerForNextTick(FTimerDelegate::CreateLambda([this]()
+			{
+				ShowAllPatchedUniverses();
 
-	if (!ensureMsgf(PatchedUniversesByID.Contains(SelectedUniverse), TEXT("Tried to scroll universe into view after displaying it, but universe doesn't exist. This is not expected.")))
-	{
-		return;
+				const int32 SelectedUniverse = SharedData->GetSelectedUniverse();
+				if (PatchedUniversesByID.Contains(SelectedUniverse))
+				{
+					PatchedUniverseScrollBox->ScrollDescendantIntoView(PatchedUniversesByID[SelectedUniverse]);
+				}
+			}));
 	}
-	PatchedUniverseScrollBox->ScrollDescendantIntoView(PatchedUniversesByID[SelectedUniverse]);
 }
 
 void SDMXFixturePatcher::SelectUniverse(int32 NewUniverseID)
@@ -530,70 +511,26 @@ int32 SDMXFixturePatcher::GetSelectedUniverse() const
 	return UniverseToSetNextTick != INDEX_NONE ? UniverseToSetNextTick : SharedData->GetSelectedUniverse();
 }
 
-void SDMXFixturePatcher::ShowSelectedUniverse(bool bForceReconstructWidget)
+void SDMXFixturePatcher::ShowSelectedUniverse()
 {
-	int32 SelectedUniverseID = GetSelectedUniverse();
-
-	/** Select a patched universe if the selected Universe doesn't contain any universes */
-	UDMXLibrary* DMXLibrary = GetDMXLibrary();
-	if (!DMXLibrary)
-	{
-		return;
-	}
-	TArray<UDMXEntityFixturePatch*> FixturePatches = DMXLibrary->GetEntitiesTypeCast<UDMXEntityFixturePatch>();
-	const bool bIsSelectedUniversePatched = FixturePatches.ContainsByPredicate([SelectedUniverseID](const UDMXEntityFixturePatch* FixturePatch)
-		{
-			return FixturePatch->GetUniverseID() == SelectedUniverseID;
-		});
-	if (!bIsSelectedUniversePatched && FixturePatches.Num() > 0)
-	{
-		FixturePatches.Sort([](const UDMXEntityFixturePatch& FixturePatchA, const UDMXEntityFixturePatch& FixturePatchB)
-			{
-				return FixturePatchA.GetUniverseID() < FixturePatchB.GetUniverseID();
-			});
-		SelectedUniverseID = FixturePatches[0]->GetUniverseID();
-	}
-
-	// Force reconstruct if requested
-	if (bForceReconstructWidget)
-	{
-		PatchedUniverseScrollBox->ClearChildren();
-		PatchedUniversesByID.Reset();
-	}
+	PatchedUniverseScrollBox->ClearChildren();
+	PatchedUniversesByID.Reset();
+	
+	const int32 SelectedUniverseID = GetSelectedUniverse();
+	const TSharedRef<SDMXPatchedUniverse> NewPatchedUniverse =
+		SNew(SDMXPatchedUniverse)
+		.DMXEditor(DMXEditorPtr)
+		.UniverseID(SelectedUniverseID)
+		.OnDragEnterChannel(this, &SDMXFixturePatcher::OnDragEnterChannel)
+		.OnDropOntoChannel(this, &SDMXFixturePatcher::OnDropOntoChannel);
 
-	// Create a new patched universe if required
-	if (PatchedUniversesByID.Num() != 1)
-	{
-		const TSharedRef<SDMXPatchedUniverse> NewPatchedUniverse =
-			SNew(SDMXPatchedUniverse)
-			.DMXEditor(DMXEditorPtr)
-			.UniverseID(SelectedUniverseID)
-			.OnDragEnterChannel(this, &SDMXFixturePatcher::OnDragEnterChannel)
-			.OnDropOntoChannel(this, &SDMXFixturePatcher::OnDropOntoChannel);
-
-		PatchedUniverseScrollBox->AddSlot()
-			.Padding(FMargin(0.0f, 3.0f, 0.0f, 12.0f))			
-			[
-				NewPatchedUniverse
-			];
+	PatchedUniverseScrollBox->AddSlot()
+		.Padding(FMargin(0.0f, 3.0f, 0.0f, 12.0f))			
+		[
+			NewPatchedUniverse
+		];
 
-		PatchedUniversesByID.Add(SelectedUniverseID, NewPatchedUniverse);
-	}
-	else
-	{
-		// Update the single, existing universe instance
-		int32 OldUniverseID = -1;
-		for (TPair<int32, TSharedPtr<SDMXPatchedUniverse>> UniverseByID : PatchedUniversesByID)
-		{
-			OldUniverseID = UniverseByID.Key;
-			break;			
-		}
-		check(OldUniverseID != -1);
-				
-		TSharedPtr<SDMXPatchedUniverse> Universe = PatchedUniversesByID.FindAndRemoveChecked(OldUniverseID);
-		PatchedUniversesByID.Add(SelectedUniverseID, Universe);
-		Universe->SetUniverseID(SelectedUniverseID);
-	}	
+	PatchedUniversesByID.Add(SelectedUniverseID, NewPatchedUniverse);
 }
 
 void SDMXFixturePatcher::ShowAllPatchedUniverses(bool bForceReconstructWidget)
@@ -644,11 +581,6 @@ void SDMXFixturePatcher::ShowAllPatchedUniverses(bool bForceReconstructWidget)
 				PatchedUniversesByID.Remove(UniverseByIDKvp.Key);
 				PatchedUniverseScrollBox->RemoveSlot(UniverseByIDKvp.Value.ToSharedRef());
 			}
-			else
-			{
-				// Update universe widgets with patches
-				UniverseByIDKvp.Value->SetUniverseID(UniverseByIDKvp.Key);
-			}
 		}
 
 		// Show last patched universe +1 for convenience of adding patches to a new universe
@@ -689,18 +621,17 @@ void SDMXFixturePatcher::AddUniverse(int32 UniverseID)
 
 void SDMXFixturePatcher::OnToggleDisplayAllUniverses(ECheckBoxState CheckboxState)
 {
-	constexpr bool bForceReconstructWidget = true;
 	const int32 SelectedUniverse = SharedData->GetSelectedUniverse();
 
 	switch (ShowAllUniversesCheckBox->GetCheckedState())
 	{
 	case ECheckBoxState::Checked:
-		ShowAllPatchedUniverses(bForceReconstructWidget);
+		ShowAllPatchedUniverses();
 		return;
 
 	case ECheckBoxState::Unchecked:
 		SelectUniverse(SelectedUniverse);
-		ShowSelectedUniverse(bForceReconstructWidget);
+		ShowSelectedUniverse();
 		return;
 
 	case ECheckBoxState::Undetermined:
diff --git a/Engine/Plugins/VirtualProduction/DMX/DMXEngine/Source/DMXEditor/Private/Widgets/FixturePatch/SDMXFixturePatcher.h b/Engine/Plugins/VirtualProduction/DMX/DMXEngine/Source/DMXEditor/Private/Widgets/FixturePatch/SDMXFixturePatcher.h
--- a/Engine/Plugins/VirtualProduction/DMX/DMXEngine/Source/DMXEditor/Private/Widgets/FixturePatch/SDMXFixturePatcher.h
+++ b/Engine/Plugins/VirtualProduction/DMX/DMXEngine/Source/DMXEditor/Private/Widgets/FixturePatch/SDMXFixturePatcher.h
@@ -28,7 +28,7 @@ enum class ECheckBoxState : uint8;
 /** Widget to assign fixture patches to universes and their channels */
 class SDMXFixturePatcher
 	: public SCompoundWidget
-	, public FEditorUndoClient
+	, public FSelfRegisteringEditorUndoClient
 {
 public:
 	SLATE_BEGIN_ARGS(SDMXFixturePatcher)
@@ -49,9 +49,6 @@ public:
 	void RefreshFromLibrary();
 
 protected:
-	/** Called when the active tab changed */
-	void OnActiveTabChanged(TSharedPtr<SDockTab> PreviouslyActive, TSharedPtr<SDockTab> NewlyActivated);
-
 	// Begin SWidget Interface
 	virtual void Tick(const FGeometry& AllottedGeometry, const double InCurrentTime, const float InDeltaTime) override;
 	virtual FReply OnDrop(const FGeometry& MyGeometry, const FDragDropEvent& DragDropEvent) override;
@@ -99,7 +96,7 @@ protected:
 	int32 GetSelectedUniverse() const;
 
 	/** Shows the selected universe only */
-	void ShowSelectedUniverse(bool bForceReconstructWidget = false);
+	void ShowSelectedUniverse();
 
 	/** Shows all universes */
 	void ShowAllPatchedUniverses(bool bForceReconstructWidgets = false);
diff --git a/Engine/Plugins/VirtualProduction/DMX/DMXEngine/Source/DMXEditor/Private/Widgets/FixturePatch/SDMXMVRFixtureList.cpp b/Engine/Plugins/VirtualProduction/DMX/DMXEngine/Source/DMXEditor/Private/Widgets/FixturePatch/SDMXMVRFixtureList.cpp
--- a/Engine/Plugins/VirtualProduction/DMX/DMXEngine/Source/DMXEditor/Private/Widgets/FixturePatch/SDMXMVRFixtureList.cpp
+++ b/Engine/Plugins/VirtualProduction/DMX/DMXEngine/Source/DMXEditor/Private/Widgets/FixturePatch/SDMXMVRFixtureList.cpp
@@ -385,7 +385,6 @@ void SDMXMVRFixtureList::Construct(const FArguments& InArgs, TWeakPtr<FDMXEditor
 			.OnGenerateRow(this, &SDMXMVRFixtureList::OnGenerateRow)
 			.OnSelectionChanged(this, &SDMXMVRFixtureList::OnSelectionChanged)
 			.OnContextMenuOpening(this, &SDMXMVRFixtureList::OnContextMenuOpening)
-			.ReturnFocusToSelection(false)
 		]
 	];
 
@@ -426,9 +425,9 @@ void SDMXMVRFixtureList::EnterFixturePatchNameEditingMode()
 			}
 		}
 	}
-}
+}	
 
-FReply SDMXMVRFixtureList::OnKeyDown(const FGeometry& MyGeometry, const FKeyEvent& InKeyEvent)
+FReply SDMXMVRFixtureList::ProcessCommandBindings(const FKeyEvent& InKeyEvent)
 {
 	if (CommandList->ProcessCommandBindings(InKeyEvent))
 	{
@@ -438,6 +437,11 @@ FReply SDMXMVRFixtureList::OnKeyDown(const FGeometry& MyGeometry, const FKeyEven
 	return FReply::Unhandled();
 }
 
+FReply SDMXMVRFixtureList::OnKeyDown(const FGeometry& MyGeometry, const FKeyEvent& InKeyEvent)
+{
+	return ProcessCommandBindings(InKeyEvent);
+}
+
 void SDMXMVRFixtureList::OnSearchChanged()
 {
 	RequestListRefresh();
@@ -629,6 +633,11 @@ void SDMXMVRFixtureList::OnFixturePatchSharedDataSelectedFixturePatches()
 			ListView->ClearSelection();
 		}
 	}
+
+	GEditor->GetTimerManager()->SetTimerForNextTick(FTimerDelegate::CreateLambda([this]()
+		{
+			FSlateApplication::Get().SetKeyboardFocus(AsShared());
+		}));
 }
 
 void SDMXMVRFixtureList::OnFixturePatchSharedDataSelectedUniverse()
diff --git a/Engine/Plugins/VirtualProduction/DMX/DMXEngine/Source/DMXEditor/Private/Widgets/FixturePatch/SDMXMVRFixtureList.h b/Engine/Plugins/VirtualProduction/DMX/DMXEngine/Source/DMXEditor/Private/Widgets/FixturePatch/SDMXMVRFixtureList.h
--- a/Engine/Plugins/VirtualProduction/DMX/DMXEngine/Source/DMXEditor/Private/Widgets/FixturePatch/SDMXMVRFixtureList.h
+++ b/Engine/Plugins/VirtualProduction/DMX/DMXEngine/Source/DMXEditor/Private/Widgets/FixturePatch/SDMXMVRFixtureList.h
@@ -69,8 +69,12 @@ public:
 	/** If a single Row is selected, enters Editing Mode for the Fixture Patch Name of the selected Row */
 	void EnterFixturePatchNameEditingMode();
 
+	/** Processes command bindings for the Key Event */
+	FReply ProcessCommandBindings(const FKeyEvent& InKeyEvent);
+
 protected:
 	//~ Begin SWidget Interface
+	virtual bool SupportsKeyboardFocus() const override { return true; }
 	virtual FReply OnKeyDown(const FGeometry& MyGeometry, const FKeyEvent& InKeyEvent) override;
 	//~ End SWidget Interface
 
diff --git a/Engine/Plugins/VirtualProduction/DMX/DMXEngine/Source/DMXEditor/Private/Widgets/FixturePatch/SDMXPatchedUniverse.cpp b/Engine/Plugins/VirtualProduction/DMX/DMXEngine/Source/DMXEditor/Private/Widgets/FixturePatch/SDMXPatchedUniverse.cpp
--- a/Engine/Plugins/VirtualProduction/DMX/DMXEngine/Source/DMXEditor/Private/Widgets/FixturePatch/SDMXPatchedUniverse.cpp
+++ b/Engine/Plugins/VirtualProduction/DMX/DMXEngine/Source/DMXEditor/Private/Widgets/FixturePatch/SDMXPatchedUniverse.cpp
@@ -96,7 +96,7 @@ void SDMXPatchedUniverse::Construct(const FArguments& InArgs)
 
 	CreateChannelConnectors();
 
-	SetUniverseID(UniverseID);
+	SetUniverseIDInternal(UniverseID);
 
 	UpdateZOrderOfNodes();
 	UpdatePatchedUniverseReachability();
@@ -112,75 +112,12 @@ END_SLATE_FUNCTION_BUILD_OPTIMIZATION
 
 void SDMXPatchedUniverse::RequestRefresh()
 {
-	if (RequestRefreshTimerHandle.IsValid())
-	{
-		return;
-	}
-
-	RequestRefreshTimerHandle =
-		GEditor->GetTimerManager()->SetTimerForNextTick(
-			FTimerDelegate::CreateLambda([this]()
-				{
-					RefreshInternal();
-				})
-		);
-}
-
-void SDMXPatchedUniverse::SetUniverseID(int32 NewUniverseID)
-{
-	check(NewUniverseID > -1);
-
-	// Find patches in new universe
-	UDMXLibrary* Library = GetDMXLibrary();
-	if (Library)
-	{
-		// Unpatch all nodes
-		TArray<TSharedPtr<FDMXFixturePatchNode>> CachedPatchedNodes = PatchedNodes;
-		for (const TSharedPtr<FDMXFixturePatchNode>& Node : CachedPatchedNodes)
-		{
-			Unpatch(Node);
-		}
-		check(PatchedNodes.Num() == 0);
-
-		// Update what to draw
-		UniverseID = NewUniverseID;
-
-		UpdatePatchedUniverseReachability();
-
-		TArray<UDMXEntityFixturePatch*> PatchesInUniverse;
-		Library->ForEachEntityOfType<UDMXEntityFixturePatch>([&](UDMXEntityFixturePatch* Patch) 
-		{
-			if (Patch->GetUniverseID() == UniverseID)
-			{
-				PatchesInUniverse.Add(Patch);
-			}
-		});
-		PatchesInUniverse.Sort([](const UDMXEntityFixturePatch& FixturePatchA, const UDMXEntityFixturePatch& FixturePatchB)
-			{
-				return FixturePatchA.GetEndingChannel() <= FixturePatchB.GetEndingChannel();
-			});
-
-		// Add fixture patches to the grid
-		for (UDMXEntityFixturePatch* FixturePatch : PatchesInUniverse)
-		{
-			TSharedPtr<FDMXFixturePatchNode> Node = FindPatchNode(FixturePatch);
-			if (!Node.IsValid())
+	GEditor->GetTimerManager()->SetTimerForNextTick(
+		FTimerDelegate::CreateLambda([this]()
 			{
-				Node = FDMXFixturePatchNode::Create(DMXEditorPtr, FixturePatch);
-			}
-			check(Node.IsValid());
-
-			Patch(Node, FixturePatch->GetStartingChannel(), false);
-		}
-
-		// Update the channel connectors' Universe ID
-		for (const TSharedPtr<SDMXChannelConnector>& Connector : ChannelConnectors)
-		{
-			Connector->SetUniverseID(NewUniverseID);
-		}
-
-		RequestRefresh();
-	}
+				RefreshInternal();
+			})
+	);
 }
 
 void SDMXPatchedUniverse::SetShowUniverseName(bool bShow)
@@ -392,15 +329,68 @@ void SDMXPatchedUniverse::OnFixturePatchChanged(const UDMXEntityFixturePatch* Fi
 	else if (!NodePtr && UniverseID == FixturePatch->GetUniverseID())
 	{
 		SharedData->SelectUniverse(UniverseID);
-		SetUniverseID(UniverseID);
 	}
+}
 
-	RequestRefresh();
+void SDMXPatchedUniverse::SetUniverseIDInternal(int32 NewUniverseID)
+{
+	check(NewUniverseID > -1);
+
+	// Find patches in new universe
+	UDMXLibrary* Library = GetDMXLibrary();
+	if (Library)
+	{
+		// Unpatch all nodes
+		TArray<TSharedPtr<FDMXFixturePatchNode>> CachedPatchedNodes = PatchedNodes;
+		for (const TSharedPtr<FDMXFixturePatchNode>& Node : CachedPatchedNodes)
+		{
+			Unpatch(Node);
+		}
+		check(PatchedNodes.Num() == 0);
+
+		// Update what to draw
+		UniverseID = NewUniverseID;
+
+		UpdatePatchedUniverseReachability();
+
+		TArray<UDMXEntityFixturePatch*> PatchesInUniverse;
+		Library->ForEachEntityOfType<UDMXEntityFixturePatch>([&](UDMXEntityFixturePatch* Patch)
+			{
+				if (Patch->GetUniverseID() == UniverseID)
+				{
+					PatchesInUniverse.Add(Patch);
+				}
+			});
+		PatchesInUniverse.Sort([](const UDMXEntityFixturePatch& FixturePatchA, const UDMXEntityFixturePatch& FixturePatchB)
+			{
+				return FixturePatchA.GetEndingChannel() <= FixturePatchB.GetEndingChannel();
+			});
+
+		// Add fixture patches to the grid
+		for (UDMXEntityFixturePatch* FixturePatch : PatchesInUniverse)
+		{
+			TSharedPtr<FDMXFixturePatchNode> Node = FindPatchNode(FixturePatch);
+			if (!Node.IsValid())
+			{
+				Node = FDMXFixturePatchNode::Create(DMXEditorPtr, FixturePatch);
+			}
+			check(Node.IsValid());
+
+			Patch(Node, FixturePatch->GetStartingChannel(), false);
+		}
+
+		// Update the channel connectors' Universe ID
+		for (const TSharedPtr<SDMXChannelConnector>& Connector : ChannelConnectors)
+		{
+			Connector->SetUniverseID(NewUniverseID);
+		}
+
+		RefreshInternal();
+	}
 }
 
 void SDMXPatchedUniverse::RefreshInternal()
 {
-	RequestRefreshTimerHandle.Invalidate();
 	UpdateZOrderOfNodes();
 
 	// Remove old widgets from the grid
@@ -478,17 +468,9 @@ FReply SDMXPatchedUniverse::HandleOnMouseButtonDownOnChannel(uint32 Channel, con
 				SharedData->SelectFixturePatch(Patch);
 			}
 
-			// SelectFixturePatch calls SetKeyboardFocus at some point; the call to SetKeyboardFocus is not enough for Slate: 
-			// We must also return the widget using FReply.
-			FReply Reply = FReply::Handled();
-			TSharedPtr<SWidget> WidgetThatWasFocusedBy_SelectFixturePatch = FSlateApplication::Get().GetKeyboardFocusedWidget();
-			if(WidgetThatWasFocusedBy_SelectFixturePatch.IsValid())
-			{
-				Reply.SetUserFocus(WidgetThatWasFocusedBy_SelectFixturePatch.ToSharedRef());
-			}
 			if (ensureMsgf(ChannelConnectors.IsValidIndex(Channel - 1), TEXT("Trying to drag Fixture Patch, but the dragged channel is not valid.")))
 			{
-				return Reply.DetectDrag(ChannelConnectors[Channel - 1].ToSharedRef(), EKeys::LeftMouseButton);
+				return FReply::Handled().DetectDrag(ChannelConnectors[Channel - 1].ToSharedRef(), EKeys::LeftMouseButton);
 			}
 		}
 	}
diff --git a/Engine/Plugins/VirtualProduction/DMX/DMXEngine/Source/DMXEditor/Private/Widgets/FixturePatch/SDMXPatchedUniverse.h b/Engine/Plugins/VirtualProduction/DMX/DMXEngine/Source/DMXEditor/Private/Widgets/FixturePatch/SDMXPatchedUniverse.h
--- a/Engine/Plugins/VirtualProduction/DMX/DMXEngine/Source/DMXEditor/Private/Widgets/FixturePatch/SDMXPatchedUniverse.h
+++ b/Engine/Plugins/VirtualProduction/DMX/DMXEngine/Source/DMXEditor/Private/Widgets/FixturePatch/SDMXPatchedUniverse.h
@@ -76,12 +76,6 @@ public:
 
 	/** Requests to refresh the Widget on the next tick */
 	void RequestRefresh();
-	
-	/** 
-	 * Shows specified Universe ID in the widget. 
-	 * If the same Universe ID is set anew, refreshes the widget. 
-	 */
-	void SetUniverseID(int32 NewUniverseID);
 
 	/** 
 	 * Patches the node.
@@ -122,6 +116,9 @@ private:
 	/** Called when a Fixture Patch changed */
 	void OnFixturePatchChanged(const UDMXEntityFixturePatch* FixturePatch);
 
+	/** Shows the universe directly */
+	void SetUniverseIDInternal(int32 NewUniverseID);
+
 	/** Refreshes the widget directly */
 	void RefreshInternal();
 
@@ -203,9 +200,6 @@ private:
 	/** The Fixture Patch Widgets that are currently being displalyed */
 	TArray<TSharedPtr<SDMXFixturePatchFragment>> FixturePatchWidgets;
 
-	/** Timer handle for the Request Refresh Timer */
-	FTimerHandle RequestRefreshTimerHandle;
-
 	/** Shared data for fixture patch editors */
 	TSharedPtr<FDMXFixturePatchSharedData> SharedData;
 
