DMX - Add a reimport option for MVR

diff --git a/Engine/Plugins/VirtualProduction/DMX/DMXEngine/Source/DMXEditor/Private/Commands/DMXEditorCommands.cpp b/Engine/Plugins/VirtualProduction/DMX/DMXEngine/Source/DMXEditor/Private/Commands/DMXEditorCommands.cpp
--- a/Engine/Plugins/VirtualProduction/DMX/DMXEngine/Source/DMXEditor/Private/Commands/DMXEditorCommands.cpp
+++ b/Engine/Plugins/VirtualProduction/DMX/DMXEngine/Source/DMXEditor/Private/Commands/DMXEditorCommands.cpp
@@ -12,7 +12,8 @@ void FDMXEditorCommands::RegisterCommands()
 {
 	UI_COMMAND(GoToDocumentation, "View Documentation", "View documentation about DMX editor", EUserInterfaceActionType::Button, FInputChord());
 
-	UI_COMMAND(ExportDMXLibrary, "Export", "Exports the DMX Library", EUserInterfaceActionType::Button, FInputChord());
+	UI_COMMAND(ImportDMXLibrary, "Import", "Import the DMX Library from an MVR File", EUserInterfaceActionType::Button, FInputChord());
+	UI_COMMAND(ExportDMXLibrary, "Export", "Exports the DMX Library to an MVR File", EUserInterfaceActionType::Button, FInputChord());
 
 	UI_COMMAND(AddNewEntityFixtureType, "New Fixture Type", "Creates a new Fixture Type in this library", EUserInterfaceActionType::Button, FInputChord());
 	UI_COMMAND(AddNewEntityFixturePatch, "Add Fixture", "Creates a new Fixture Patch in this library", EUserInterfaceActionType::Button, FInputChord());
diff --git a/Engine/Plugins/VirtualProduction/DMX/DMXEngine/Source/DMXEditor/Private/DMXEditor.cpp b/Engine/Plugins/VirtualProduction/DMX/DMXEngine/Source/DMXEditor/Private/DMXEditor.cpp
--- a/Engine/Plugins/VirtualProduction/DMX/DMXEngine/Source/DMXEditor/Private/DMXEditor.cpp
+++ b/Engine/Plugins/VirtualProduction/DMX/DMXEngine/Source/DMXEditor/Private/DMXEditor.cpp
@@ -31,7 +31,9 @@
 #include "DesktopPlatformModule.h"
 #include "IDesktopPlatform.h"
 #include "ScopedTransaction.h"
+#include "Utils.h"
 #include "Framework/Notifications/NotificationManager.h"
+#include "Misc/MessageDialog.h"
 #include "Modules/ModuleManager.h"
 #include "Styling/AppStyle.h"
 #include "Widgets/Notifications/SNotificationList.h"
@@ -125,6 +127,60 @@ UDMXLibrary* FDMXEditor::GetDMXLibrary() const
 	return Cast<UDMXLibrary>(GetEditingObject());
 }
 
+void FDMXEditor::ImportDMXLibrary() const
+{
+	UDMXLibrary* DMXLibrary = GetDMXLibrary();
+	if (!DMXLibrary)
+	{
+		return;
+	}
+
+	IDesktopPlatform* DesktopPlatform = FDesktopPlatformModule::Get();
+	if (!DesktopPlatform)
+	{
+		return;
+	}
+
+	UDMXEditorSettings* DMXEditorSettings = GetMutableDefault<UDMXEditorSettings>();
+	if (!DMXEditorSettings)
+	{
+		return;
+	}
+
+	const FString LastMVRImportPath = DMXEditorSettings->LastMVRImportPath;
+	const FString DefaultPath = FPaths::DirectoryExists(LastMVRImportPath) ? LastMVRImportPath : FPaths::ProjectSavedDir();
+	
+	if (!DMXLibrary->GetEntities().IsEmpty())
+	{
+		const FText MessageText = LOCTEXT("MVRImportDialog", "DMX Library already contains data. Importing the MVR will clear existing data. Do you want to proceed?");
+		if (FMessageDialog::Open(EAppMsgType::YesNo, MessageText) == EAppReturnType::No)
+		{
+			return;
+		}
+	}
+
+	TArray<FString> OpenFilenames;
+	DesktopPlatform->OpenFileDialog(
+		FSlateApplication::Get().FindBestParentWindowHandleForDialogs(nullptr),
+		LOCTEXT("ImportMVR", "Import MVR").ToString(),
+		DefaultPath,
+		TEXT(""),
+		TEXT("My Virtual Rig (*.mvr)|*.mvr"),
+		EFileDialogFlags::None,
+		OpenFilenames);
+
+	if (OpenFilenames.IsEmpty())
+	{
+		return;
+	}
+
+	if (ImportObject<UDMXLibrary>(DMXLibrary->GetOuter(), DMXLibrary->GetFName(), RF_Public | RF_Standalone, *OpenFilenames[0], nullptr))
+	{
+		DMXEditorSettings->LastMVRImportPath = FPaths::GetPath(OpenFilenames[0]);
+		DMXEditorSettings->SaveConfig();
+	}
+}
+
 void FDMXEditor::ExportDMXLibrary() const
 {
 	UDMXLibrary* DMXLibrary = GetDMXLibrary();
@@ -139,14 +195,8 @@ void FDMXEditor::ExportDMXLibrary() const
 		UDMXEditorSettings* DMXEditorSettings = GetMutableDefault<UDMXEditorSettings>();
 		check(DMXEditorSettings);
 
-		const FString DefaultPath = [DMXEditorSettings]()
-		{
-			if (FPaths::DirectoryExists(DMXEditorSettings->LastMVRExportPath))
-			{
-				return DMXEditorSettings->LastMVRExportPath;
-			}
-			return FPaths::ProjectSavedDir();
-		}();
+		const FString LastMVRExportPath = DMXEditorSettings->LastMVRExportPath;
+		const FString DefaultPath = FPaths::DirectoryExists(LastMVRExportPath) ? LastMVRExportPath : FPaths::ProjectSavedDir();
 
 		TArray<FString> SaveFilenames;
 		const bool bSaveFile = DesktopPlatform->SaveFileDialog(
@@ -167,8 +217,8 @@ void FDMXEditor::ExportDMXLibrary() const
 		FDMXMVRExporter::Export(DMXLibrary, SaveFilenames[0], ErrorReason);
 		if (ErrorReason.IsEmpty())
 		{
-			const FString SavePath = FPaths::GetPath(SaveFilenames[0]);
-			DMXEditorSettings->LastMVRExportPath = SavePath;
+			DMXEditorSettings->LastMVRExportPath = FPaths::GetPath(SaveFilenames[0]);
+			DMXEditorSettings->SaveConfig();
 
 			FNotificationInfo NotificationInfo(FText::Format(LOCTEXT("ExportDMXLibraryAsMVRSuccessNotification", "Successfully exported MVR to {0}."), FText::FromString(SaveFilenames[0])));
 			NotificationInfo.ExpireDuration = 5.f;
@@ -202,6 +252,10 @@ void FDMXEditor::CreateDefaultCommands()
 {
 	FDMXEditorCommands::Register();
 
+	ToolkitCommands->MapAction(
+		FDMXEditorCommands::Get().ImportDMXLibrary,
+		FExecuteAction::CreateSP(this, &FDMXEditor::ImportDMXLibrary)
+	);
 	ToolkitCommands->MapAction(
 		FDMXEditorCommands::Get().ExportDMXLibrary,
 		FExecuteAction::CreateSP(this, &FDMXEditor::ExportDMXLibrary)
diff --git a/Engine/Plugins/VirtualProduction/DMX/DMXEngine/Source/DMXEditor/Private/DMXEditorStyle.cpp b/Engine/Plugins/VirtualProduction/DMX/DMXEngine/Source/DMXEditor/Private/DMXEditorStyle.cpp
--- a/Engine/Plugins/VirtualProduction/DMX/DMXEngine/Source/DMXEditor/Private/DMXEditorStyle.cpp
+++ b/Engine/Plugins/VirtualProduction/DMX/DMXEngine/Source/DMXEditor/Private/DMXEditorStyle.cpp
@@ -79,6 +79,7 @@ FDMXEditorStyle::FDMXEditorStyle()
 
 		Set("Icons.DMXLibrarySettings", new CORE_IMAGE_BRUSH_SVG("Starship/Common/settings", Icon16x16));
 
+		Set("Icons.DMXLibraryToolbar.Import", new CORE_IMAGE_BRUSH_SVG("Starship/Common/import", Icon20x20));
 		Set("Icons.DMXLibraryToolbar.Export", new EDITOR_IMAGE_BRUSH_SVG("Starship/Common/Export", Icon20x20));
 
 		Set("Icons.FixtureType", new IMAGE_BRUSH_SVG("FixtureType_16", Icon16x16));
diff --git a/Engine/Plugins/VirtualProduction/DMX/DMXEngine/Source/DMXEditor/Private/Factories/DMXLibraryFromMVRFactory.cpp b/Engine/Plugins/VirtualProduction/DMX/DMXEngine/Source/DMXEditor/Private/Factories/DMXLibraryFromMVRFactory.cpp
--- a/Engine/Plugins/VirtualProduction/DMX/DMXEngine/Source/DMXEditor/Private/Factories/DMXLibraryFromMVRFactory.cpp
+++ b/Engine/Plugins/VirtualProduction/DMX/DMXEngine/Source/DMXEditor/Private/Factories/DMXLibraryFromMVRFactory.cpp
@@ -29,8 +29,7 @@
 
 #define LOCTEXT_NAMESPACE "DMXLibraryFromMVRFactory"
 
-const FName UDMXLibraryFromMVRFactory::MVRFileExtension = "MVR";
-const FName UDMXLibraryFromMVRFactory::GDTFFileExtension = "GDTF";
+const FString UDMXLibraryFromMVRFactory::MVRFileExtension = TEXT("MVR");
 
 UDMXLibraryFromMVRFactory::UDMXLibraryFromMVRFactory()
 {
@@ -55,7 +54,7 @@ UObject* UDMXLibraryFromMVRFactory::FactoryCreateFile(UClass* InClass, UObject*
 	check(DMXEditorSettings);
 	DMXEditorSettings->LastMVRImportPath = FPaths::GetPath(InFilename);
 
-	UDMXLibrary* NewDMXLibrary = CreateDMXLibraryAsset(Parent, Flags, InFilename);
+	UDMXLibrary* NewDMXLibrary = CreateDMXLibraryAsset(Parent, InName, Flags, InFilename);
 	if (!NewDMXLibrary)
 	{
 		UE_LOG(LogDMXEditor, Error, TEXT("Failed to create DMX Library for MVR '%s'."), *InFilename);
@@ -100,8 +99,7 @@ UObject* UDMXLibraryFromMVRFactory::FactoryCreateFile(UClass* InClass, UObject*
 bool UDMXLibraryFromMVRFactory::FactoryCanImport(const FString& Filename)
 {
 	const FString TargetExtension = FPaths::GetExtension(Filename);
-
-	if (TargetExtension.Equals(UDMXLibraryFromMVRFactory::MVRFileExtension.ToString(), ESearchCase::IgnoreCase))
+	if (TargetExtension.Equals(UDMXLibraryFromMVRFactory::MVRFileExtension, ESearchCase::IgnoreCase))
 	{
 		return true;
 	}
@@ -109,29 +107,78 @@ bool UDMXLibraryFromMVRFactory::FactoryCanImport(const FString& Filename)
 	return false;
 }
 
-UDMXLibrary* UDMXLibraryFromMVRFactory::CreateDMXLibraryAsset(UObject* Parent, EObjectFlags Flags, const FString& InFilename)
+bool UDMXLibraryFromMVRFactory::CanReimport(UObject* Obj, TArray<FString>& OutFilenames)
 {
-	constexpr bool bRemovePathFromDesiredName = true;
-	const FString BaseFileName = FPaths::GetBaseFilename(InFilename, bRemovePathFromDesiredName);
-	const FString PackageName = Parent->GetName();
-	FString AssetName = ObjectTools::SanitizeObjectName(BaseFileName);
-	if (PackageName.Contains(AssetName))
+	UDMXMVRAssetImportData* MVRAssetImportData = GetMVRAssetImportData(Obj);
+	if (!MVRAssetImportData)
+	{
+		return false;
+	}
+	
+	const FString SourceFilename = MVRAssetImportData->GetFilePathAndName();
+	OutFilenames.Add(SourceFilename);
+	return true;
+}
+
+void UDMXLibraryFromMVRFactory::SetReimportPaths(UObject* Obj, const TArray<FString>& NewReimportPaths)
+{
+	UDMXMVRAssetImportData* MVRAssetImportData = GetMVRAssetImportData(Obj);
+	if (!ensureMsgf(MVRAssetImportData, TEXT("Invalid MVR Asset Import Data for General Scene Description %s.")))
+	{
+		return;
+	}
+
+	if (!ensureMsgf(NewReimportPaths.Num() == 1, TEXT("Unexpected, more than one or no new reimport path when reimporting DMX Library from MVR.")))
 	{
-		AssetName = AssetName + TEXT("_DMXLibrary");
+		return;
 	}
 
-	FString DMXLibraryAssetName;
-	FAssetToolsModule& AssetToolsModule = FModuleManager::LoadModuleChecked<FAssetToolsModule>("AssetTools");
-	AssetToolsModule.Get().CreateUniqueAssetName(PackageName / AssetName, TEXT(""), DMXLibraryPackageName, DMXLibraryAssetName);
+	MVRAssetImportData->SetSourceFile(NewReimportPaths[0]);
+}
 
-	UPackage* Package = CreatePackage(*DMXLibraryPackageName);
-	check(Package);
-	Package->FullyLoad();
+EReimportResult::Type UDMXLibraryFromMVRFactory::Reimport(UObject* Obj)
+{
+	UDMXMVRAssetImportData* MVRAssetImportData = GetMVRAssetImportData(Obj);
+	if (!ensureMsgf(MVRAssetImportData, TEXT("Invalid MVR Asset Import Data for General Scene Description.")))
+	{
+		return EReimportResult::Failed;
+	}
+	const FString SourceFilename = MVRAssetImportData->GetFilePathAndName();
 
+	UDMXLibrary* DMXLibrary = Cast<UDMXLibrary>(Obj);
+	if (!ensureMsgf(DMXLibrary, TEXT("Invalid DMX Library when trying to reimport DMX Library.")))
+	{
+		return EReimportResult::Failed;
+	}
+	if (!DMXLibrary->GetEntities().IsEmpty())
+	{
+		const FText MessageText = LOCTEXT("MVRReimportDialog", "DMX Library already contains data. Reimporting MVR will clear existing data. Do you want to proceed?");
+		if (FMessageDialog::Open(EAppMsgType::YesNo, MessageText) == EAppReturnType::No)
+		{
+			return EReimportResult::Cancelled;
+		}
+	}
+
+	bool bOutCanceled = false;
+	if (ImportObject(Obj->GetClass(), Obj->GetOuter(), *Obj->GetName(), RF_Public | RF_Standalone, SourceFilename, nullptr, bOutCanceled))
+	{
+		return EReimportResult::Succeeded;
+	}
+
+	return bOutCanceled ? EReimportResult::Cancelled : EReimportResult::Failed;
+}
+
+int32 UDMXLibraryFromMVRFactory::GetPriority() const
+{
+	return ImportPriority;
+}
+
+UDMXLibrary* UDMXLibraryFromMVRFactory::CreateDMXLibraryAsset(UObject* Parent, const FName& Name, EObjectFlags Flags, const FString& InFilename)
+{
 	UImportSubsystem* ImportSubsystem = GEditor->GetEditorSubsystem<UImportSubsystem>();
-	ImportSubsystem->BroadcastAssetPreImport(this, UDMXLibrary::StaticClass(), Parent, *DMXLibraryAssetName, *MVRFileExtension.ToString());
+	ImportSubsystem->BroadcastAssetPreImport(this, UDMXLibrary::StaticClass(), Parent, Name, *MVRFileExtension);
 
-	UDMXLibrary* NewDMXLibrary = NewObject<UDMXLibrary>(Package, *DMXLibraryAssetName, Flags | RF_Public);
+	UDMXLibrary* NewDMXLibrary = NewObject<UDMXLibrary>(Parent, Name, Flags | RF_Public);
 	if (!NewDMXLibrary)
 	{
 		return nullptr;
@@ -144,7 +191,7 @@ UDMXLibrary* UDMXLibraryFromMVRFactory::CreateDMXLibraryAsset(UObject* Parent, E
 
 TArray<UDMXImportGDTF*> UDMXLibraryFromMVRFactory::CreateGDTFAssets(UObject* Parent, EObjectFlags Flags, const TSharedRef<FDMXZipper>& Zip, const UDMXMVRGeneralSceneDescription& GeneralSceneDescription)
 {
-	const FString Path = FPaths::GetPath(DMXLibraryPackageName) + TEXT("/GDTFs");
+	const FString Path = FPaths::GetPath(Parent->GetName()) + TEXT("/GDTFs");
 
 	TArray<FAssetData> ExistingGDTFAssets;
 	FAssetRegistryModule& AssetRegistryModule = FModuleManager::LoadModuleChecked<FAssetRegistryModule>(TEXT("AssetRegistry"));
@@ -316,4 +363,21 @@ void UDMXLibraryFromMVRFactory::InitDMXLibrary(UDMXLibrary* DMXLibrary, const TA
 	DMXLibrary->UpdateGeneralSceneDescription();
 }
 
+UDMXMVRAssetImportData* UDMXLibraryFromMVRFactory::GetMVRAssetImportData(UObject* DMXLibraryObject) const
+{
+	UDMXLibrary* DMXLibrary = Cast<UDMXLibrary>(DMXLibraryObject);
+	if (!DMXLibrary)
+	{
+		return nullptr;
+	}
+
+	UDMXMVRGeneralSceneDescription* GeneralSceneDescription = DMXLibrary->GetLazyGeneralSceneDescription();
+	if (!ensureMsgf(GeneralSceneDescription, TEXT("Invalid General Scene Description in DMX Library.")))
+	{
+		return nullptr;
+	}
+
+	return GeneralSceneDescription->GetMVRAssetImportData();
+}
+
 #undef LOCTEXT_NAMESPACE
diff --git a/Engine/Plugins/VirtualProduction/DMX/DMXEngine/Source/DMXEditor/Private/Toolbars/DMXEditorToolbar.cpp b/Engine/Plugins/VirtualProduction/DMX/DMXEngine/Source/DMXEditor/Private/Toolbars/DMXEditorToolbar.cpp
--- a/Engine/Plugins/VirtualProduction/DMX/DMXEngine/Source/DMXEditor/Private/Toolbars/DMXEditorToolbar.cpp
+++ b/Engine/Plugins/VirtualProduction/DMX/DMXEngine/Source/DMXEditor/Private/Toolbars/DMXEditorToolbar.cpp
@@ -26,21 +26,25 @@ FSlateIcon FDMXEditorToolbar::GetStatusImage() const
 	return FSlateIcon(FAppStyle::GetAppStyleSetName(), "Kismet.Status.Good");
 }
 
-FText FDMXEditorToolbar::GetStatusTooltip() const
-{
-	return LOCTEXT("Default_Status", "Good to go");
-}
-
 void FDMXEditorToolbar::FillDMXLibraryToolbar(FToolBarBuilder& ToolbarBuilder)
 {
 	ToolbarBuilder.BeginSection("DMXLibraryToolbar");
+	{
+		ToolbarBuilder.AddToolBarButton(
+			FDMXEditorCommands::Get().ImportDMXLibrary,
+			NAME_None,
+			TAttribute<FText>(),
+			TAttribute<FText>(),
+			FSlateIcon(FDMXEditorStyle::Get().GetStyleSetName(), "Icons.DMXLibraryToolbar.Import"));
+
 		ToolbarBuilder.AddToolBarButton(
-		FDMXEditorCommands::Get().ExportDMXLibrary,
-		NAME_None,
-		TAttribute<FText>(),
-		TAttribute<FText>(),
-		FSlateIcon(FDMXEditorStyle::Get().GetStyleSetName(), "Icons.DMXLibraryToolbar.Export"));
+			FDMXEditorCommands::Get().ExportDMXLibrary,
+			NAME_None,
+			TAttribute<FText>(),
+			TAttribute<FText>(),
+			FSlateIcon(FDMXEditorStyle::Get().GetStyleSetName(), "Icons.DMXLibraryToolbar.Export"));
+	}
 	ToolbarBuilder.EndSection();
 }
 
-#undef LOCTEXT_NAMESPACE
\ No newline at end of file
+#undef LOCTEXT_NAMESPACE
diff --git a/Engine/Plugins/VirtualProduction/DMX/DMXEngine/Source/DMXEditor/Private/Toolbars/DMXEditorToolbar.h b/Engine/Plugins/VirtualProduction/DMX/DMXEngine/Source/DMXEditor/Private/Toolbars/DMXEditorToolbar.h
--- a/Engine/Plugins/VirtualProduction/DMX/DMXEngine/Source/DMXEditor/Private/Toolbars/DMXEditorToolbar.h
+++ b/Engine/Plugins/VirtualProduction/DMX/DMXEngine/Source/DMXEditor/Private/Toolbars/DMXEditorToolbar.h
@@ -24,9 +24,6 @@ public:
 	/** Returns the current status icon for the blueprint being edited */
 	FSlateIcon GetStatusImage() const;
 
-	/** Returns the current status as text for the blueprint being edited */
-	FText GetStatusTooltip() const;
-
 private:
 	void FillDMXLibraryToolbar(FToolBarBuilder& ToolbarBuilder);
 
diff --git a/Engine/Plugins/VirtualProduction/DMX/DMXEngine/Source/DMXEditor/Public/Commands/DMXEditorCommands.h b/Engine/Plugins/VirtualProduction/DMX/DMXEngine/Source/DMXEditor/Public/Commands/DMXEditorCommands.h
--- a/Engine/Plugins/VirtualProduction/DMX/DMXEngine/Source/DMXEditor/Public/Commands/DMXEditorCommands.h
+++ b/Engine/Plugins/VirtualProduction/DMX/DMXEngine/Source/DMXEditor/Public/Commands/DMXEditorCommands.h
@@ -20,6 +20,7 @@ public:
 	TSharedPtr<FUICommandInfo> GoToDocumentation;
 
 	// DMX Library related
+	TSharedPtr<FUICommandInfo> ImportDMXLibrary;
 	TSharedPtr<FUICommandInfo> ExportDMXLibrary;
 
 	TSharedPtr<FUICommandInfo> AddNewEntityFixtureType;
diff --git a/Engine/Plugins/VirtualProduction/DMX/DMXEngine/Source/DMXEditor/Public/DMXEditor.h b/Engine/Plugins/VirtualProduction/DMX/DMXEngine/Source/DMXEditor/Public/DMXEditor.h
--- a/Engine/Plugins/VirtualProduction/DMX/DMXEngine/Source/DMXEditor/Public/DMXEditor.h
+++ b/Engine/Plugins/VirtualProduction/DMX/DMXEngine/Source/DMXEditor/Public/DMXEditor.h
@@ -89,7 +89,8 @@ public:
 	TSharedPtr<FDMXEditorToolbar> GetToolbarBuilder() { return Toolbar; }
 
 	void RegisterToolbarTab(const TSharedRef<class FTabManager>& TabManager);
-	
+
+
 protected:
 	/** Creates the widgets that go into the tabs (note: does not create the tabs themselves) **/
 	virtual void CreateDefaultTabContents(UDMXLibrary* DMXLibrary);
@@ -104,7 +105,10 @@ protected:
 	virtual void RegisterApplicationModes(UDMXLibrary* DMXLibrary, bool bShouldOpenInDefaultsMode, bool bNewlyCreated = false);
 
 private:
-	/** Exports the DMX Library */
+	/** Imports the DMX Library from an MVR File */
+	void ImportDMXLibrary() const;
+
+	/** Exports the DMX Library as MVR File */
 	void ExportDMXLibrary() const;
 
 	//~ Generate Editor widgets for tabs
diff --git a/Engine/Plugins/VirtualProduction/DMX/DMXEngine/Source/DMXEditor/Public/Factories/DMXLibraryFromMVRFactory.h b/Engine/Plugins/VirtualProduction/DMX/DMXEngine/Source/DMXEditor/Public/Factories/DMXLibraryFromMVRFactory.h
--- a/Engine/Plugins/VirtualProduction/DMX/DMXEngine/Source/DMXEditor/Public/Factories/DMXLibraryFromMVRFactory.h
+++ b/Engine/Plugins/VirtualProduction/DMX/DMXEngine/Source/DMXEditor/Public/Factories/DMXLibraryFromMVRFactory.h
@@ -4,17 +4,21 @@
 
 #include "Factories/Factory.h"
 
+#include "EditorReimportHandler.h"
+
 #include "DMXLibraryFromMVRFactory.generated.h"
 
 class FDMXZipper;
 class UDMXImportGDTF;
 class UDMXLibrary;
+class UDMXMVRAssetImportData;
 class UDMXMVRGeneralSceneDescription;
 
 
 UCLASS()
 class DMXEDITOR_API UDMXLibraryFromMVRFactory 
 	: public UFactory
+	, public FReimportHandler
 {
 	GENERATED_BODY()
 
@@ -27,14 +31,19 @@ public:
 	//~ End UFactory Interface	
 
 	/** File extention for MVR files */
-	static const FName MVRFileExtension;
+	static const FString MVRFileExtension;
 
-	/** File extention for GDTF files */
-	static const FName GDTFFileExtension;
+protected:
+	//~ Begin FReimportHandler Interface
+	virtual bool CanReimport(UObject* Obj, TArray<FString>& OutFilenames) override;
+	virtual void SetReimportPaths(UObject* Obj, const TArray<FString>& NewReimportPaths) override;
+	virtual EReimportResult::Type Reimport(UObject* Obj) override;
+	virtual int32 GetPriority() const override;
+	//~ End FReimportHandler Interface
 
 private:
 	/** Creates a DMX Library asset. Returns nullptr if the library could not be created */
-	UDMXLibrary* CreateDMXLibraryAsset(UObject* Parent, EObjectFlags Flags, const FString& InFilename);
+	UDMXLibrary* CreateDMXLibraryAsset(UObject* Parent, const FName& Name, EObjectFlags Flags, const FString& InFilename);
 
 	/** Creates GDTF assets from the MVR */
 	TArray<UDMXImportGDTF*> CreateGDTFAssets(UObject* Parent, EObjectFlags Flags, const TSharedRef<FDMXZipper>& Zip, const UDMXMVRGeneralSceneDescription& GeneralSceneDescription);
@@ -42,6 +51,6 @@ private:
 	/** Initializes the DMX Library from the General Scene Description and GDTF assets */
 	void InitDMXLibrary(UDMXLibrary* DMXLibrary, const TArray<UDMXImportGDTF*>& GDTFAssets, UDMXMVRGeneralSceneDescription* GeneralSceneDescription) const;
 
-	/** The DMX Library Package Name, initialized when the DMX Library Asset is created */
-	FString DMXLibraryPackageName;
+	/** Returns the MVR Asset Import Data for the Object, or nullptr if it cannot be obtained */
+	UDMXMVRAssetImportData* GetMVRAssetImportData(UObject* DMXLibraryObject) const;
 };
