DMX - Remove unacceptable words from Output Console

diff --git a/Engine/Plugins/VirtualProduction/DMX/DMXEngine/Source/DMXEditor/Private/Widgets/OutputConsole/SDMXOutputFaderList.cpp b/Engine/Plugins/VirtualProduction/DMX/DMXEngine/Source/DMXEditor/Private/Widgets/OutputConsole/SDMXOutputFaderList.cpp
--- a/Engine/Plugins/VirtualProduction/DMX/DMXEngine/Source/DMXEditor/Private/Widgets/OutputConsole/SDMXOutputFaderList.cpp
+++ b/Engine/Plugins/VirtualProduction/DMX/DMXEngine/Source/DMXEditor/Private/Widgets/OutputConsole/SDMXOutputFaderList.cpp
@@ -23,7 +23,7 @@ namespace DMXOutputFaderList
 	static FLinearColor DefaultFillColor = FLinearColor::FromSRGBColor(FColor::FromHex("00aeef"));
 	static FLinearColor DefaultBackColor = FLinearColor::FromSRGBColor(FColor::FromHex("414042"));
 	static FLinearColor DefeaultForeColor = FLinearColor::FromSRGBColor(FColor::FromHex("d5d6d8"));
-	static FSpinBoxStyle MasterFaderStyle(FCoreStyle::Get().GetWidgetStyle<FSpinBoxStyle>("SpinBox"));
+	static FSpinBoxStyle PrimaryFaderStyle(FCoreStyle::Get().GetWidgetStyle<FSpinBoxStyle>("SpinBox"));
 };
 
 SDMXOutputFaderList::SDMXOutputFaderList()
@@ -54,7 +54,7 @@ void SDMXOutputFaderList::Construct(const FArguments& InArgs)
 	FSlateBrush ArrowsImage;
 	ArrowsImage.TintColor = FLinearColor::Transparent;
 
-	DMXOutputFaderList::MasterFaderStyle
+	DMXOutputFaderList::PrimaryFaderStyle
 		.SetActiveFillBrush(FillBrush)
 		.SetInactiveFillBrush(FillBrush)
 		.SetBackgroundBrush(BackBrush)
@@ -94,7 +94,7 @@ void SDMXOutputFaderList::Construct(const FArguments& InArgs)
 				.Orientation(Orient_Horizontal)
 			]
 
-			// Master Fader
+			// Primary Fader
 			+ SVerticalBox::Slot()
 			.VAlign(VAlign_Center)
 			.HAlign(HAlign_Fill)
@@ -116,7 +116,7 @@ void SDMXOutputFaderList::Construct(const FArguments& InArgs)
 				.Padding(4.0f, 0.0f, 0.0f, 0.0f)
 				[
 					SNew(STextBlock)
-					.Text(LOCTEXT("MasterFader", "Master Fader"))
+					.Text(LOCTEXT("PrimaryFaderLabel", "Primary Fader"))
 					.MinDesiredWidth(100)
 					.Justification(ETextJustify::Center)
 				]
@@ -133,13 +133,13 @@ void SDMXOutputFaderList::Construct(const FArguments& InArgs)
 						SNew(SBorder)
 						.BorderBackgroundColor(FLinearColor::Black)
 						[
-							SAssignNew(MasterFader, SSpinBox<uint8>)
+							SAssignNew(PrimaryFader, SSpinBox<uint8>)
 							.MinValue(0)
 							.MaxValue(100)
 							.MaxFractionalDigits(0)
 							.ContentPadding(FMargin(100.0f, 1.0f, 0.0f, 1.0f))
-							.Style(&DMXOutputFaderList::MasterFaderStyle)
-							.OnValueChanged(this, &SDMXOutputFaderList::HandleMasterFaderChanged)
+							.Style(&DMXOutputFaderList::PrimaryFaderStyle)
+							.OnValueChanged(this, &SDMXOutputFaderList::HandlePrimaryFaderChanged)
 							.MinDesiredWidth(100.0f)
 						]
 					]
@@ -672,7 +672,7 @@ FReply SDMXOutputFaderList::OnSortFadersClicked()
 	return FReply::Handled();
 }
 
-void SDMXOutputFaderList::HandleMasterFaderChanged(uint8 NewValue)
+void SDMXOutputFaderList::HandlePrimaryFaderChanged(uint8 NewValue)
 {
 	for (TSharedPtr<SDMXFader> Fader : Faders)
 	{
diff --git a/Engine/Plugins/VirtualProduction/DMX/DMXEngine/Source/DMXEditor/Private/Widgets/OutputConsole/SDMXOutputFaderList.h b/Engine/Plugins/VirtualProduction/DMX/DMXEngine/Source/DMXEditor/Private/Widgets/OutputConsole/SDMXOutputFaderList.h
--- a/Engine/Plugins/VirtualProduction/DMX/DMXEngine/Source/DMXEditor/Private/Widgets/OutputConsole/SDMXOutputFaderList.h
+++ b/Engine/Plugins/VirtualProduction/DMX/DMXEngine/Source/DMXEditor/Private/Widgets/OutputConsole/SDMXOutputFaderList.h
@@ -103,7 +103,7 @@ protected:
 	FReply OnSortFadersClicked();
 
 	/** Called when a fader send DMX check box state changes  */
-	void HandleMasterFaderChanged(uint8 NewValue);
+	void HandlePrimaryFaderChanged(uint8 NewValue);
 
 	/** Called when the add fader button is clicked */
 	FReply HandleAddFadersClicked();
@@ -117,8 +117,8 @@ protected:
 	/** Port selector */
 	TSharedPtr<SDMXOutputConsolePortSelector> PortSelector;
 
-	/** The master fader that controlls all faders */
-	TSharedPtr<SSpinBox<uint8>> MasterFader;
+	/** The primary fader that controlls all faders */
+	TSharedPtr<SSpinBox<uint8>> PrimaryFader;
 
 	/** The displayed fader widgets */
 	TArray<TSharedPtr<SDMXFader>> Faders;
