FName -> FSoftObjectPath and FTopLevelAssetPath refactoring.

diff --git a/Engine/Plugins/VirtualProduction/CameraCalibration/Source/CameraCalibrationEditor/Private/AssetEditor/SCameraCalibrationSteps.cpp b/Engine/Plugins/VirtualProduction/CameraCalibration/Source/CameraCalibrationEditor/Private/AssetEditor/SCameraCalibrationSteps.cpp
--- a/Engine/Plugins/VirtualProduction/CameraCalibration/Source/CameraCalibrationEditor/Private/AssetEditor/SCameraCalibrationSteps.cpp
+++ b/Engine/Plugins/VirtualProduction/CameraCalibration/Source/CameraCalibrationEditor/Private/AssetEditor/SCameraCalibrationSteps.cpp
@@ -225,7 +225,7 @@ TSharedRef<SWidget> SCameraCalibrationSteps::BuildCameraPickerWidget()
 			if (ACameraActor* Camera = CalibrationStepsController.Pin()->GetCamera())
 			{
 				FAssetData AssetData(Camera, true);
-				return AssetData.ObjectPath.ToString();
+				return AssetData.GetObjectPathString();
 			}
 
 			return TEXT("");
diff --git a/Engine/Plugins/VirtualProduction/CameraCalibration/Source/CameraCalibrationEditor/Private/Calibrators/CameraLensDistortionAlgoCheckerboard.cpp b/Engine/Plugins/VirtualProduction/CameraCalibration/Source/CameraCalibrationEditor/Private/Calibrators/CameraLensDistortionAlgoCheckerboard.cpp
--- a/Engine/Plugins/VirtualProduction/CameraCalibration/Source/CameraCalibrationEditor/Private/Calibrators/CameraLensDistortionAlgoCheckerboard.cpp
+++ b/Engine/Plugins/VirtualProduction/CameraCalibration/Source/CameraCalibrationEditor/Private/Calibrators/CameraLensDistortionAlgoCheckerboard.cpp
@@ -690,7 +690,7 @@ TSharedRef<SWidget> UCameraLensDistortionAlgoCheckerboard::BuildCalibrationDevic
 				if (AActor* TheCalibrator = GetCalibrator())
 				{
 					FAssetData AssetData(TheCalibrator, true);
-					return AssetData.ObjectPath.ToString();
+					return AssetData.GetObjectPathString();
 				}
 
 				return TEXT("");
diff --git a/Engine/Plugins/VirtualProduction/DMX/DMXEngine/Source/DMXEditor/Private/Widgets/SAssetPickerButton.cpp b/Engine/Plugins/VirtualProduction/DMX/DMXEngine/Source/DMXEditor/Private/Widgets/SAssetPickerButton.cpp
--- a/Engine/Plugins/VirtualProduction/DMX/DMXEngine/Source/DMXEditor/Private/Widgets/SAssetPickerButton.cpp
+++ b/Engine/Plugins/VirtualProduction/DMX/DMXEngine/Source/DMXEditor/Private/Widgets/SAssetPickerButton.cpp
@@ -270,7 +270,7 @@ const FAssetData& SAssetPickerButton::GetAssetData() const
 {
 	if (UObject* Object = CurrentAssetValue.Get().Get())
 	{
-		if (!Object->GetPathName().Equals(CachedAssetData.ObjectPath.ToString(), ESearchCase::CaseSensitive))
+		if (FSoftObjectPath(Object) != CachedAssetData.GetSoftObjectPath())
 		{
 			// This always uses the exact object pointed at
 			CachedAssetData = FAssetData(Object, true);
diff --git a/Engine/Plugins/VirtualProduction/DMX/DMXEngine/Source/DMXRuntime/Private/MVR/DMXMVRFixtureActorLibrary.cpp b/Engine/Plugins/VirtualProduction/DMX/DMXEngine/Source/DMXRuntime/Private/MVR/DMXMVRFixtureActorLibrary.cpp
--- a/Engine/Plugins/VirtualProduction/DMX/DMXEngine/Source/DMXRuntime/Private/MVR/DMXMVRFixtureActorLibrary.cpp
+++ b/Engine/Plugins/VirtualProduction/DMX/DMXEngine/Source/DMXRuntime/Private/MVR/DMXMVRFixtureActorLibrary.cpp
@@ -20,7 +20,7 @@ private:
 		: public TSharedFromThis<FDMXVRAssetNode>
 	{
 	public:
-		FDMXVRAssetNode(FName InClassPath, FString InClassName)
+		FDMXVRAssetNode(FTopLevelAssetPath InClassPath, FString InClassName)
 			: ClassName(InClassName)
 			, ClassPath(InClassPath)
 		{}
@@ -28,7 +28,7 @@ private:
 		TSharedPtr<FDMXVRAssetNode> Parent;
 		TArray<TSharedRef<FDMXVRAssetNode>> Children;
 
-		void AppendClassPathsOfSelfAndChildren(TArray<FName>& OutChildClassPaths) const
+		void AppendClassPathsOfSelfAndChildren(TArray<FTopLevelAssetPath>& OutChildClassPaths) const
 		{
 			OutChildClassPaths.Add(ClassPath);
 			for (const TSharedRef<FDMXVRAssetNode>& ChildNode : Children)
@@ -38,14 +38,14 @@ private:
 		}
 
 		FString ClassName;
-		FName ClassPath;
-		FName ParentClassPath;
-		FName BlueprintAssetPath;
+		FTopLevelAssetPath ClassPath;
+		FTopLevelAssetPath ParentClassPath;
+		FSoftObjectPath BlueprintAssetPath;
 		TArray<FString> ImplementedInterfaces;
 	};
 
 public:
-	static void GetClassPathsWithInterface(const FString& InInterfaceClassPathName, TArray<FName>& OutClassPaths)
+	static void GetClassPathsWithInterface(const FString& InInterfaceClassPathName, TArray<FTopLevelAssetPath>& OutClassPaths)
 	{
 		FDMXMVRFixtureActorAssetHierarchy MVRAssetHierarchy;
 
@@ -57,7 +57,7 @@ public:
 	}
 
 private:
-	void GetClassPathsWithInterface(const FString& InInterfaceClassPathName, TArray<FName>& OutClassPaths, const TSharedRef<FDMXVRAssetNode>& StartNode)
+	void GetClassPathsWithInterface(const FString& InInterfaceClassPathName, TArray<FTopLevelAssetPath>& OutClassPaths, const TSharedRef<FDMXVRAssetNode>& StartNode)
 	{
 		if (StartNode->ImplementedInterfaces.Contains(InInterfaceClassPathName))
 		{
@@ -77,7 +77,7 @@ private:
 		// Fetch all classes from AssetRegistry blueprint data (which covers unloaded classes), and in-memory UClasses.
 		// Create a node for each one with unioned data from the AssetRegistry or UClass for that class.
 		// Set parent/child pointers to create a tree, and store this tree in this->ObjectClassRoot
-		TMap<FName, TSharedPtr<FDMXVRAssetNode>> ClassPathToNode;
+		TMap<FTopLevelAssetPath, TSharedPtr<FDMXVRAssetNode>> ClassPathToNode;
 
 		// Create a node for every Blueprint class listed in the AssetRegistry and set the Blueprint fields
 		// Retrieve all blueprint classes 
@@ -93,12 +93,12 @@ private:
 		FString ClassPathString;
 		for (FAssetData& AssetData : BlueprintList)
 		{
-			FName ClassPath;
+			FTopLevelAssetPath ClassPath;
 			if (AssetData.GetTagValue(FBlueprintTags::GeneratedClassPath, ClassPathString))
 			{
-				ClassPath = FName(*FPackageName::ExportTextPathToObjectPath(ClassPathString));
+				ClassPath = FTopLevelAssetPath(FPackageName::ExportTextPathToObjectPath(ClassPathString));
 			}
-			if (ClassPath.IsNone())
+			if (ClassPath.IsNull())
 			{
 				continue;
 			}
@@ -117,7 +117,7 @@ private:
 		BuildHierarchyFromLoadedClassees(ObjectClassRoot, ClassPathToNode);
 
 		// Set the parent and child pointers
-		for (TPair<FName, TSharedPtr<FDMXVRAssetNode>>& KVPair : ClassPathToNode)
+		for (TPair<FTopLevelAssetPath, TSharedPtr<FDMXVRAssetNode>>& KVPair : ClassPathToNode)
 		{
 			TSharedPtr<FDMXVRAssetNode>& Node = KVPair.Value;
 			if (Node == ObjectClassRoot)
@@ -126,7 +126,7 @@ private:
 				continue;
 			}
 			TSharedPtr<FDMXVRAssetNode>* ParentNodePtr = nullptr;
-			if (!Node->ParentClassPath.IsNone())
+			if (!Node->ParentClassPath.IsNull())
 			{
 				ParentNodePtr = ClassPathToNode.Find(Node->ParentClassPath);
 			}
@@ -140,7 +140,7 @@ private:
 		}
 	}
 
-	void BuildHierarchyFromLoadedClassees(TSharedPtr<FDMXVRAssetNode>& OutRootNode, TMap<FName, TSharedPtr<FDMXVRAssetNode>>& InOutClassPathToNode)
+	void BuildHierarchyFromLoadedClassees(TSharedPtr<FDMXVRAssetNode>& OutRootNode, TMap<FTopLevelAssetPath, TSharedPtr<FDMXVRAssetNode>>& InOutClassPathToNode)
 	{
 		for (TObjectIterator<UClass> ClassIt; ClassIt; ++ClassIt)
 		{
@@ -150,14 +150,14 @@ private:
 			{
 				continue;
 			}
-			TSharedPtr<FDMXVRAssetNode>& Node = InOutClassPathToNode.FindOrAdd(FName(*CurrentClass->GetPathName()));
+			TSharedPtr<FDMXVRAssetNode>& Node = InOutClassPathToNode.FindOrAdd(FTopLevelAssetPath(CurrentClass->GetPathName()));
 			if (!Node.IsValid())
 			{
-				Node = MakeShared<FDMXVRAssetNode>(*CurrentClass->GetPathName(), CurrentClass->GetPathName());
+				Node = MakeShared<FDMXVRAssetNode>(FTopLevelAssetPath(CurrentClass->GetPathName()), CurrentClass->GetPathName());
 			}
 			SetClassFields(Node, *CurrentClass);
 		}
-		TSharedPtr<FDMXVRAssetNode>* ExistingRoot = InOutClassPathToNode.Find(FName(*AActor::StaticClass()->GetPathName()));
+		TSharedPtr<FDMXVRAssetNode>* ExistingRoot = InOutClassPathToNode.Find(FTopLevelAssetPath(AActor::StaticClass()->GetPathName()));
 		check(ExistingRoot && ExistingRoot->IsValid());
 		OutRootNode = *ExistingRoot;
 	}
@@ -165,15 +165,15 @@ private:
 	void SetClassFields(TSharedPtr<FDMXVRAssetNode>& InOutClassNode, UClass& Class)
 	{
 		// Fields that can als be set from FAssetData
-		if (InOutClassNode->ClassPath.IsNone())
+		if (InOutClassNode->ClassPath.IsNull())
 		{
-			InOutClassNode->ClassPath = FName(*Class.GetPathName());
+			InOutClassNode->ClassPath = FTopLevelAssetPath(Class.GetPathName());
 		}
-		if (InOutClassNode->ParentClassPath.IsNone())
+		if (InOutClassNode->ParentClassPath.IsNull())
 		{
 			if (Class.GetSuperClass())
 			{
-				InOutClassNode->ParentClassPath = FName(*Class.GetSuperClass()->GetPathName());
+				InOutClassNode->ParentClassPath = FTopLevelAssetPath(Class.GetSuperClass()->GetPathName());
 			}
 		}
 
@@ -183,7 +183,7 @@ private:
 		}
 	}
 
-	TSharedPtr<FDMXVRAssetNode> FindNodeByGeneratedClassPath(const TSharedRef<FDMXVRAssetNode>& InRootNode, FName InGeneratedClassPath)
+	TSharedPtr<FDMXVRAssetNode> FindNodeByGeneratedClassPath(const TSharedRef<FDMXVRAssetNode>& InRootNode, FTopLevelAssetPath InGeneratedClassPath)
 	{
 		if (InRootNode->ClassPath == InGeneratedClassPath)
 		{
@@ -208,30 +208,30 @@ private:
 
 	void SetAssetDataFields(TSharedPtr<FDMXVRAssetNode>& InOutMXVRAssetNode, const FAssetData& InAssetData)
 	{
-		if (InOutMXVRAssetNode->ClassPath.IsNone())
+		if (InOutMXVRAssetNode->ClassPath.IsNull())
 		{
 			FString GeneratedClassPath;
 			UClass* AssetClass = InAssetData.GetClass();
 			if (AssetClass && AssetClass->IsChildOf(UBlueprintGeneratedClass::StaticClass()))
 			{
-				InOutMXVRAssetNode->ClassPath = InAssetData.ObjectPath;
+				InOutMXVRAssetNode->ClassPath = InAssetData.GetSoftObjectPath().GetAssetPath();
 			}
 			else if (InAssetData.GetTagValue(FBlueprintTags::GeneratedClassPath, GeneratedClassPath))
 			{
-				InOutMXVRAssetNode->ClassPath = FName(*FPackageName::ExportTextPathToObjectPath(GeneratedClassPath));
+				InOutMXVRAssetNode->ClassPath = FTopLevelAssetPath(FPackageName::ExportTextPathToObjectPath(GeneratedClassPath));
 			}
 		}
-		if (InOutMXVRAssetNode->ParentClassPath.IsNone())
+		if (InOutMXVRAssetNode->ParentClassPath.IsNull())
 		{
 			FString ParentClassPathString;
 			if (InAssetData.GetTagValue(FBlueprintTags::ParentClassPath, ParentClassPathString))
 			{
-				InOutMXVRAssetNode->ParentClassPath = FName(*FPackageName::ExportTextPathToObjectPath(ParentClassPathString));
+				InOutMXVRAssetNode->ParentClassPath = FTopLevelAssetPath(FPackageName::ExportTextPathToObjectPath(ParentClassPathString));
 			}
 		}
 
 		// Blueprint-specific fields
-		InOutMXVRAssetNode->BlueprintAssetPath = InAssetData.ObjectPath;
+		InOutMXVRAssetNode->BlueprintAssetPath = InAssetData.GetSoftObjectPath();
 
 		// Get interface class paths.
 		TArray<FString> ImplementedInterfaces;
@@ -285,9 +285,9 @@ private:
 
 FDMXMVRFixtureActorLibrary::FDMXMVRFixtureActorLibrary()
 {
-	TArray<FName> ClassPaths;
+	TArray<FTopLevelAssetPath> ClassPaths;
 	FDMXMVRFixtureActorAssetHierarchy::GetClassPathsWithInterface(UDMXMVRFixtureActorInterface::StaticClass()->GetPathName(), ClassPaths);
-	ClassPaths.RemoveAll([](const FName& ClassPath)
+	ClassPaths.RemoveAll([](const FTopLevelAssetPath& ClassPath)
 		{
 			return
 				ClassPath.ToString().Contains(TEXT(".SKEL_")) ||
@@ -296,7 +296,7 @@ FDMXMVRFixtureActorLibrary::FDMXMVRFixtureActorLibrary()
 
 	FPreviewScene PreviewScene;
 	TArray<UClass*> MVRFixtureActorClasses;
-	for (const FName& ClassPath : ClassPaths)
+	for (const FTopLevelAssetPath& ClassPath : ClassPaths)
 	{
 		UClass* Class = LoadClass<UObject>(nullptr, *ClassPath.ToString());
 		AActor* Actor = PreviewScene.GetWorld()->SpawnActor<AActor>(Class);
diff --git a/Engine/Plugins/VirtualProduction/LevelSnapshots/Source/LevelShapshotFilters/Private/Builtin/ActorInMapFilter.cpp b/Engine/Plugins/VirtualProduction/LevelSnapshots/Source/LevelShapshotFilters/Private/Builtin/ActorInMapFilter.cpp
--- a/Engine/Plugins/VirtualProduction/LevelSnapshots/Source/LevelShapshotFilters/Private/Builtin/ActorInMapFilter.cpp
+++ b/Engine/Plugins/VirtualProduction/LevelSnapshots/Source/LevelShapshotFilters/Private/Builtin/ActorInMapFilter.cpp
@@ -10,7 +10,7 @@ namespace UE::LevelSnapshotsFilters::Private
 {
 	static bool IsActorInMap(const FSoftObjectPath& Actor, const FSoftObjectPath& MapNameToCheck)
 	{
-		return Actor.GetAssetPathName() == MapNameToCheck.GetAssetPathName();
+		return Actor.GetAssetPath() == MapNameToCheck.GetAssetPath();
 	}
 
 	static EFilterResult::Type IsActorAllowed(const FSoftObjectPath& Actor, const TArray<TSoftObjectPtr<UWorld>>& AllowedLevels)
diff --git a/Engine/Plugins/VirtualProduction/LevelSnapshots/Source/LevelSnapshots/Private/Data/Util/SnapshotUtil.cpp b/Engine/Plugins/VirtualProduction/LevelSnapshots/Source/LevelSnapshots/Private/Data/Util/SnapshotUtil.cpp
--- a/Engine/Plugins/VirtualProduction/LevelSnapshots/Source/LevelSnapshots/Private/Data/Util/SnapshotUtil.cpp
+++ b/Engine/Plugins/VirtualProduction/LevelSnapshots/Source/LevelSnapshots/Private/Data/Util/SnapshotUtil.cpp
@@ -53,7 +53,7 @@ TOptional<FSoftObjectPath> UE::LevelSnapshots::Private::ExtractActorFromPath(con
 		OriginalObjectPath
 		:
 		// Converts /Game/MapName.MapName:PersistentLevel.StaticMeshActor_42.StaticMeshComponent to /Game/MapName.MapName:PersistentLevel.StaticMeshActor_42
-		FSoftObjectPath(OriginalObjectPath.GetAssetPathName(), SubPathString.Left(DotAfterActorNameIndex));
+		FSoftObjectPath(OriginalObjectPath.GetAssetPath(), SubPathString.Left(DotAfterActorNameIndex));
 }
 
 bool UE::LevelSnapshots::Private::IsPathToWorldObject(const FSoftObjectPath& OriginalObjectPath)
@@ -113,5 +113,5 @@ FSoftObjectPath UE::LevelSnapshots::Private::SetActorInPath(AActor* NewActor, co
 	const FSoftObjectPath PathToNewActor(NewActor);
 	// PersistentLevel.StaticMeshActor_42.StaticMeshComponent becomes .StaticMeshComponent
 	const FString PathAfterOriginalActor = SubPathString.Right(SubPathString.Len() - DotAfterActorNameIndex); 
-	return FSoftObjectPath(PathToNewActor.GetAssetPathName(), PathToNewActor.GetSubPathString() + PathAfterOriginalActor);
+	return FSoftObjectPath(PathToNewActor.GetAssetPath(), PathToNewActor.GetSubPathString() + PathAfterOriginalActor);
 }
\ No newline at end of file
diff --git a/Engine/Plugins/VirtualProduction/LevelSnapshots/Source/LevelSnapshotsEditor/Private/Util/TakeSnapshotUtil.cpp b/Engine/Plugins/VirtualProduction/LevelSnapshots/Source/LevelSnapshotsEditor/Private/Util/TakeSnapshotUtil.cpp
--- a/Engine/Plugins/VirtualProduction/LevelSnapshots/Source/LevelSnapshotsEditor/Private/Util/TakeSnapshotUtil.cpp
+++ b/Engine/Plugins/VirtualProduction/LevelSnapshots/Source/LevelSnapshotsEditor/Private/Util/TakeSnapshotUtil.cpp
@@ -96,7 +96,7 @@ ULevelSnapshot* SnapshotEditor::TakeLevelSnapshotAndSaveToDisk(UWorld* World, co
     if (bProceedWithSave && !bShouldCreateUniqueFileName)
     {
     	FAssetRegistryModule& AssetRegistryModule = FModuleManager::LoadModuleChecked<FAssetRegistryModule>("AssetRegistry");
-    	FAssetData ExistingAsset = AssetRegistryModule.Get().GetAssetByObjectPath(FName(*(PackageName + "." + FileName)));
+    	FAssetData ExistingAsset = AssetRegistryModule.Get().GetAssetByObjectPath(FSoftObjectPath((PackageName + "." + FileName)));
     	if (ExistingAsset.IsValid() && ExistingAsset.AssetClassPath == ULevelSnapshot::StaticClass()->GetClassPathName())
     	{
     		EAppReturnType::Type ShouldReplace = FMessageDialog::Open( EAppMsgType::YesNo, FText::Format(LOCTEXT("ReplaceAssetMessage", "{0} already exists. Do you want to replace it?"), FText::FromString(FileName)) );
diff --git a/Engine/Plugins/VirtualProduction/LevelSnapshots/Source/LevelSnapshotsEditor/Private/Views/Input/SLevelSnapshotsEditorContextPicker.cpp b/Engine/Plugins/VirtualProduction/LevelSnapshots/Source/LevelSnapshotsEditor/Private/Views/Input/SLevelSnapshotsEditorContextPicker.cpp
--- a/Engine/Plugins/VirtualProduction/LevelSnapshots/Source/LevelSnapshotsEditor/Private/Views/Input/SLevelSnapshotsEditorContextPicker.cpp
+++ b/Engine/Plugins/VirtualProduction/LevelSnapshots/Source/LevelSnapshotsEditor/Private/Views/Input/SLevelSnapshotsEditorContextPicker.cpp
@@ -80,7 +80,7 @@ UWorld* SLevelSnapshotsEditorContextPicker::GetSelectedWorld() const
 {
 	const FAssetRegistryModule& AssetRegistryModule = FModuleManager::LoadModuleChecked<FAssetRegistryModule>(TEXT("AssetRegistry"));
 
-	return Cast<UWorld>(AssetRegistryModule.Get().GetAssetByObjectPath(FName(SelectedWorldPath.GetAssetPathString())).GetAsset());
+	return Cast<UWorld>(AssetRegistryModule.Get().GetAssetByObjectPath(SelectedWorldPath.GetWithoutSubPath()).GetAsset());
 }
 
 FSoftObjectPath SLevelSnapshotsEditorContextPicker::GetSelectedWorldSoftPath() const
diff --git a/Engine/Plugins/VirtualProduction/Switchboard/Source/SwitchboardEditor/Private/DisplayClusterRootActorReferenceDetailCustomization.cpp b/Engine/Plugins/VirtualProduction/Switchboard/Source/SwitchboardEditor/Private/DisplayClusterRootActorReferenceDetailCustomization.cpp
--- a/Engine/Plugins/VirtualProduction/Switchboard/Source/SwitchboardEditor/Private/DisplayClusterRootActorReferenceDetailCustomization.cpp
+++ b/Engine/Plugins/VirtualProduction/Switchboard/Source/SwitchboardEditor/Private/DisplayClusterRootActorReferenceDetailCustomization.cpp
@@ -105,7 +105,7 @@ void FDisplayClusterRootActorReferenceDetailCustomization::CustomizeHeader(TShar
 				if (DCRAReference && DCRAReference->DCRA.IsValid())
 				{
 					FAssetData AssetData(DCRAReference->DCRA.Get(), true);
-					return AssetData.ObjectPath.ToString();
+					return AssetData.GetObjectPathString();
 				}
 
 				return TEXT("");
