DMX - Fix MVR fixtures get detached and can't be reattached after changing their type

diff --git a/Engine/Plugins/VirtualProduction/DMX/DMXEngine/Source/DMXRuntime/Private/MVR/DMXMVRSceneActor.cpp b/Engine/Plugins/VirtualProduction/DMX/DMXEngine/Source/DMXRuntime/Private/MVR/DMXMVRSceneActor.cpp
--- a/Engine/Plugins/VirtualProduction/DMX/DMXEngine/Source/DMXRuntime/Private/MVR/DMXMVRSceneActor.cpp
+++ b/Engine/Plugins/VirtualProduction/DMX/DMXEngine/Source/DMXRuntime/Private/MVR/DMXMVRSceneActor.cpp
@@ -347,6 +347,7 @@ AActor* ADMXMVRSceneActor::SpawnMVRActor(const TSubclassOf<AActor>&ActorClass, U
 #endif 
 
 	// Attach, set MVR Fixture UUID, set Fixture Patch, remember as a Related Actor
+	RootComponentOfChildActor->SetMobility(EComponentMobility::Movable);
 	RootComponentOfChildActor->AttachToComponent(MVRSceneRoot, FAttachmentTransformRules::KeepRelativeTransform);
 	const FGuid& MVRFixtureUUID = FixturePatch->GetMVRFixtureUUID();
 	UDMXMVRAssetUserData::SetMVRAssetUserDataValueForKey(*NewFixtureActor, UDMXMVRAssetUserData::MVRFixtureUUIDMetaDataKey, MVRFixtureUUID.ToString());
