DMX - Fix MVR reimport creates new patches instead of using existing ones where feasible

diff --git a/Engine/Plugins/VirtualProduction/DMX/DMXEngine/Source/DMXEditor/Private/Factories/DMXLibraryFromMVRFactory.cpp b/Engine/Plugins/VirtualProduction/DMX/DMXEngine/Source/DMXEditor/Private/Factories/DMXLibraryFromMVRFactory.cpp
--- a/Engine/Plugins/VirtualProduction/DMX/DMXEngine/Source/DMXEditor/Private/Factories/DMXLibraryFromMVRFactory.cpp
+++ b/Engine/Plugins/VirtualProduction/DMX/DMXEngine/Source/DMXEditor/Private/Factories/DMXLibraryFromMVRFactory.cpp
@@ -68,25 +68,13 @@ UObject* UDMXLibraryFromMVRFactory::FactoryCreateFile(UClass* InClass, UObject*
 		return nullptr;
 	}
 
-	TArray64<uint8> XMLData;
-	if (!Zip->GetFileContent(TEXT("GeneralSceneDescription.xml"), XMLData))
+	UDMXMVRGeneralSceneDescription* GeneralSceneDescription = CreateGeneralSsceneDescription(*NewDMXLibrary, Zip);
+	if (!GeneralSceneDescription)
 	{
 		UE_LOG(LogDMXEditor, Error, TEXT("Cannot read General Scene Description from MVR '%s'. File is not a valid MVR."), *InFilename);
 		return nullptr;
 	}
 
-	// MVR implicitly adpots UTF-8 encoding of Xml Files by adopting the GDT standard (DIN-15800).
-	// Content is NOT null-terminated; we need to specify lengths here.
-	const FUTF8ToTCHAR TCHARData(reinterpret_cast<const ANSICHAR*>(XMLData.GetData()), XMLData.Num());
-
-	const TSharedRef<FXmlFile> GeneralSceneDescriptionXml = MakeShared<FXmlFile>(FString(TCHARData.Length(), TCHARData.Get()), EConstructMethod::ConstructFromBuffer);
-
-	const FName GeneralSceneDescriptionName = FName(GetName() + TEXT("_MVRGeneralSceneDescription"));
-	UDMXMVRGeneralSceneDescription* GeneralSceneDescription = UDMXMVRGeneralSceneDescription::CreateFromXmlFile(GeneralSceneDescriptionXml, NewDMXLibrary, GeneralSceneDescriptionName);
-	if (!GeneralSceneDescription)
-	{
-		return nullptr;
-	}
 	UDMXMVRAssetImportData* MVRAssetImportData = GeneralSceneDescription->GetMVRAssetImportData();
 	MVRAssetImportData->SetSourceFile(InFilename);
 
@@ -150,6 +138,7 @@ EReimportResult::Type UDMXLibraryFromMVRFactory::Reimport(UObject* Obj)
 	{
 		return EReimportResult::Failed;
 	}
+
 	if (!DMXLibrary->GetEntities().IsEmpty())
 	{
 		const FText MessageText = LOCTEXT("MVRReimportDialog", "DMX Library already contains data. Reimporting MVR will clear existing data. Do you want to proceed?");
@@ -159,13 +148,24 @@ EReimportResult::Type UDMXLibraryFromMVRFactory::Reimport(UObject* Obj)
 		}
 	}
 
-	bool bOutCanceled = false;
-	if (ImportObject(Obj->GetClass(), Obj->GetOuter(), *Obj->GetName(), RF_Public | RF_Standalone, SourceFilename, nullptr, bOutCanceled))
+	const TSharedRef<FDMXZipper> Zip = MakeShared<FDMXZipper>();
+	if (!Zip->LoadFromFile(SourceFilename))
+	{
+		UE_LOG(LogDMXEditor, Error, TEXT("Cannot read MVR '%s'. File is not a valid MVR."), *SourceFilename);
+		return EReimportResult::Failed;
+	}
+
+	UDMXMVRGeneralSceneDescription* GeneralSceneDescription = CreateGeneralSsceneDescription(*DMXLibrary, Zip);
+	if (!GeneralSceneDescription)
 	{
-		return EReimportResult::Succeeded;
+		UE_LOG(LogDMXEditor, Error, TEXT("Cannot read General Scene Description from MVR '%s'. File is not a valid MVR."), *SourceFilename);
+		return EReimportResult::Failed;
 	}
 
-	return bOutCanceled ? EReimportResult::Cancelled : EReimportResult::Failed;
+	TArray<UDMXImportGDTF*> GDTFs = CreateGDTFAssets(DMXLibrary->GetOuter(), RF_NoFlags, Zip, *GeneralSceneDescription);
+	InitDMXLibrary(DMXLibrary, GDTFs, GeneralSceneDescription);
+
+	return EReimportResult::Succeeded;
 }
 
 int32 UDMXLibraryFromMVRFactory::GetPriority() const
@@ -189,6 +189,26 @@ UDMXLibrary* UDMXLibraryFromMVRFactory::CreateDMXLibraryAsset(UObject* Parent, c
 	return NewDMXLibrary;
 }
 
+UDMXMVRGeneralSceneDescription* UDMXLibraryFromMVRFactory::CreateGeneralSsceneDescription(UDMXLibrary& DMXLibrary, const TSharedRef<FDMXZipper>& Zip) const
+{
+	TArray64<uint8> XMLData;
+	if (!Zip->GetFileContent(TEXT("GeneralSceneDescription.xml"), XMLData))
+	{
+		return nullptr;
+	}
+
+	// MVR implicitly adpots UTF-8 encoding of Xml Files by adopting the GDTF standard (DIN-15800).
+	// Content is NOT null-terminated; we need to specify lengths here.
+	const FUTF8ToTCHAR TCHARData(reinterpret_cast<const ANSICHAR*>(XMLData.GetData()), XMLData.Num());
+
+	const TSharedRef<FXmlFile> GeneralSceneDescriptionXml = MakeShared<FXmlFile>(FString(TCHARData.Length(), TCHARData.Get()), EConstructMethod::ConstructFromBuffer);
+
+	const FName GeneralSceneDescriptionName = FName(GetName() + TEXT("_MVRGeneralSceneDescription"));
+	UDMXMVRGeneralSceneDescription* GeneralSceneDescription = UDMXMVRGeneralSceneDescription::CreateFromXmlFile(GeneralSceneDescriptionXml, &DMXLibrary, GeneralSceneDescriptionName);
+	
+	return GeneralSceneDescription;
+}
+
 TArray<UDMXImportGDTF*> UDMXLibraryFromMVRFactory::CreateGDTFAssets(UObject* Parent, EObjectFlags Flags, const TSharedRef<FDMXZipper>& Zip, const UDMXMVRGeneralSceneDescription& GeneralSceneDescription)
 {
 	const FString Path = FPaths::GetPath(Parent->GetName()) + TEXT("/GDTFs");
@@ -284,7 +304,8 @@ void UDMXLibraryFromMVRFactory::InitDMXLibrary(UDMXLibrary* DMXLibrary, const TA
 	}
 	DMXLibrary->SetMVRGeneralSceneDescription(GeneralSceneDescription);
 
-	// Create a Fixture Type for each GDTF
+	// Get or create a Fixture Type for each GDTF
+	TArray<UDMXEntityFixtureType*> FixtureTypes = DMXLibrary->GetEntitiesTypeCast<UDMXEntityFixtureType>();
 	TMap<FString, UDMXEntityFixtureType*> GDTFSpecToFixtureTypeMap;
 	TMap<UDMXEntityFixtureType*, FLinearColor> FixtureTypeToColorMap;
 	for (UDMXImportGDTF* GDTF : GDTFAssets)
@@ -294,28 +315,55 @@ void UDMXLibraryFromMVRFactory::InitDMXLibrary(UDMXLibrary* DMXLibrary, const TA
 		{
 			continue;
 		}
-
 		const FString GDTFSourceFilename = GDTFAssetImportData->GetFilePathAndName();
 		const FString GDTFFilename = FPaths::GetCleanFilename(GDTFSourceFilename);
-		FDMXEntityFixtureTypeConstructionParams FixtureTypeConstructionParams;
-		FixtureTypeConstructionParams.DMXCategory = FDMXFixtureCategory(FDMXFixtureCategory::GetFirstValue());
-		FixtureTypeConstructionParams.ParentDMXLibrary = DMXLibrary;
 
-		UDMXEntityFixtureType* FixtureType = UDMXEntityFixtureType::CreateFixtureTypeInLibrary(FixtureTypeConstructionParams, FPaths::GetBaseFilename(GDTFFilename));
-		const bool bAdvancedImportSuccess = FDMXInitializeFixtureTypeFromGDTFHelper::GenerateModesFromGDTF(*FixtureType, *GDTF);
-		if (!bAdvancedImportSuccess)
+		UDMXEntityFixtureType** ExistingFixtureTypePtr = Algo::FindByPredicate(FixtureTypes, [GDTF](const UDMXEntityFixtureType* FixtureType)
+			{
+				return FixtureType->DMXImport == GDTF;
+			});
+		if (ExistingFixtureTypePtr)
 		{
-			UE_LOG(LogDMXEditor, Warning, TEXT("Failed to initialize Fixture Type '%s', falling back to legacy method that doesn't support matrix fixtures."), *FixtureType->GetName());
-			FixtureType->SetModesFromDMXImport(GDTF);
+			GDTFSpecToFixtureTypeMap.Add(GDTFFilename, *ExistingFixtureTypePtr);
 		}
-		FixtureType->DMXImport = GDTF;
+		else
+		{
+			FDMXEntityFixtureTypeConstructionParams FixtureTypeConstructionParams;
+			FixtureTypeConstructionParams.DMXCategory = FDMXFixtureCategory(FDMXFixtureCategory::GetFirstValue());
+			FixtureTypeConstructionParams.ParentDMXLibrary = DMXLibrary;
 
-		GDTFSpecToFixtureTypeMap.Add(GDTFFilename, FixtureType);
+			UDMXEntityFixtureType* NewFixtureType = UDMXEntityFixtureType::CreateFixtureTypeInLibrary(FixtureTypeConstructionParams, FPaths::GetBaseFilename(GDTFFilename));
+			const bool bAdvancedImportSuccess = FDMXInitializeFixtureTypeFromGDTFHelper::GenerateModesFromGDTF(*NewFixtureType, *GDTF);
+			if (!bAdvancedImportSuccess)
+			{
+				UE_LOG(LogDMXEditor, Warning, TEXT("Failed to initialize Fixture Type '%s', falling back to legacy method that doesn't support matrix fixtures."), *NewFixtureType->GetName());
+				NewFixtureType->SetModesFromDMXImport(GDTF);
+			}
+			NewFixtureType->DMXImport = GDTF;
+
+			GDTFSpecToFixtureTypeMap.Add(GDTFFilename, NewFixtureType);
+		}
 	}
 
-	// Create Fixture Patches for the MVR Fixtures
-	TArray< UDMXMVRFixtureNode*> FixtureNodes;
+	// Remove Fixture Patches that have no corresponding Fixture Node in the General Scene Description
+	TArray<UDMXMVRFixtureNode*> FixtureNodes;
 	GeneralSceneDescription->GetFixtureNodes(FixtureNodes);
+	TArray<UDMXEntityFixturePatch*> FixturePatches = DMXLibrary->GetEntitiesTypeCast<UDMXEntityFixturePatch>();
+	for (UDMXEntityFixturePatch* FixturePatch : FixturePatches)
+	{
+		UDMXMVRFixtureNode** CorrespondingFixtureNodePtr = Algo::FindByPredicate(FixtureNodes, [FixturePatch](const UDMXMVRFixtureNode* FixtureNode)
+			{
+				return FixturePatch->GetMVRFixtureUUID() == FixtureNode->UUID;
+			});
+		if (!CorrespondingFixtureNodePtr)
+		{
+			checkNoEntry();
+			UDMXEntityFixturePatch::RemoveFixturePatchFromLibrary(FDMXEntityFixturePatchRef(FixturePatch));
+		}
+	}
+	FixturePatches = DMXLibrary->GetEntitiesTypeCast<UDMXEntityFixturePatch>();
+
+	// Update or create new Fixture Patches for the MVR Fixtures
 	for (const UDMXMVRFixtureNode* FixtureNode : FixtureNodes)
 	{
 		if (!GDTFSpecToFixtureTypeMap.Contains(FixtureNode->GDTFSpec))
@@ -344,15 +392,34 @@ void UDMXLibraryFromMVRFactory::InitDMXLibrary(UDMXLibrary* DMXLibrary, const TA
 
 		if (ActiveModeIndex != INDEX_NONE)
 		{
-			FDMXEntityFixturePatchConstructionParams FixturePatchConstructionParams;
-			FixturePatchConstructionParams.ActiveMode = ActiveModeIndex;
-			FixturePatchConstructionParams.FixtureTypeRef = FDMXEntityFixtureTypeRef(FixtureType);
-			FixturePatchConstructionParams.UniverseID = FixtureNode->GetUniverseID();
-			FixturePatchConstructionParams.StartingAddress = FixtureNode->GetStartingChannel();
-			FixturePatchConstructionParams.MVRFixtureUUID = FixtureNode->UUID;
-
-			UDMXEntityFixturePatch* FixturePatch = UDMXEntityFixturePatch::CreateFixturePatchInLibrary(FixturePatchConstructionParams, FixtureNode->Name);
-			FixturePatch->EditorColor = FixtureTypeToColorMap.FindChecked(FixtureType);
+			UDMXEntityFixturePatch* const* ExistingFixturePatchPtr = Algo::FindByPredicate(FixturePatches, [FixtureNode](const UDMXEntityFixturePatch* FixturePatch)
+				{
+					return FixturePatch->GetMVRFixtureUUID() == FixtureNode->UUID;
+				});
+			if (ExistingFixturePatchPtr)
+			{			
+				// Update the existing patch (matching MVR UUID) if possible
+				UDMXEntityFixturePatch& ExistingFixturePatch = **ExistingFixturePatchPtr;
+
+				ExistingFixturePatch.SetFixtureType(FixtureType);
+				ExistingFixturePatch.SetActiveModeIndex(ActiveModeIndex);
+				ExistingFixturePatch.SetUniverseID(FixtureNode->GetUniverseID());
+				ExistingFixturePatch.SetStartingChannel(FixtureNode->GetStartingChannel());
+				ExistingFixturePatch.EditorColor = FixtureTypeToColorMap.FindChecked(FixtureType);
+			}
+			else
+			{
+				// Create a new patch
+				FDMXEntityFixturePatchConstructionParams FixturePatchConstructionParams;
+				FixturePatchConstructionParams.ActiveMode = ActiveModeIndex;
+				FixturePatchConstructionParams.FixtureTypeRef = FDMXEntityFixtureTypeRef(FixtureType);
+				FixturePatchConstructionParams.UniverseID = FixtureNode->GetUniverseID();
+				FixturePatchConstructionParams.StartingAddress = FixtureNode->GetStartingChannel();
+				FixturePatchConstructionParams.MVRFixtureUUID = FixtureNode->UUID;
+
+				UDMXEntityFixturePatch* FixturePatch = UDMXEntityFixturePatch::CreateFixturePatchInLibrary(FixturePatchConstructionParams, FixtureNode->Name);
+				FixturePatch->EditorColor = FixtureTypeToColorMap.FindChecked(FixtureType);
+			}
 		}
 		else
 		{
diff --git a/Engine/Plugins/VirtualProduction/DMX/DMXEngine/Source/DMXEditor/Private/Widgets/FixturePatch/SDMXFixturePatchEditor.cpp b/Engine/Plugins/VirtualProduction/DMX/DMXEngine/Source/DMXEditor/Private/Widgets/FixturePatch/SDMXFixturePatchEditor.cpp
--- a/Engine/Plugins/VirtualProduction/DMX/DMXEngine/Source/DMXEditor/Private/Widgets/FixturePatch/SDMXFixturePatchEditor.cpp
+++ b/Engine/Plugins/VirtualProduction/DMX/DMXEngine/Source/DMXEditor/Private/Widgets/FixturePatch/SDMXFixturePatchEditor.cpp
@@ -77,8 +77,7 @@ void SDMXFixturePatchEditor::Construct(const FArguments& InArgs)
 
 FReply SDMXFixturePatchEditor::OnKeyDown(const FGeometry& MyGeometry, const FKeyEvent& InKeyEvent)
 {
-	MVRFixtureList->ProcessCommandBindings(InKeyEvent);
-	return FReply::Handled();
+	return MVRFixtureList->ProcessCommandBindings(InKeyEvent);
 }
 
 void SDMXFixturePatchEditor::RequestRenameOnNewEntity(const UDMXEntity* InEntity, ESelectInfo::Type SelectionType)
diff --git a/Engine/Plugins/VirtualProduction/DMX/DMXEngine/Source/DMXEditor/Public/Factories/DMXLibraryFromMVRFactory.h b/Engine/Plugins/VirtualProduction/DMX/DMXEngine/Source/DMXEditor/Public/Factories/DMXLibraryFromMVRFactory.h
--- a/Engine/Plugins/VirtualProduction/DMX/DMXEngine/Source/DMXEditor/Public/Factories/DMXLibraryFromMVRFactory.h
+++ b/Engine/Plugins/VirtualProduction/DMX/DMXEngine/Source/DMXEditor/Public/Factories/DMXLibraryFromMVRFactory.h
@@ -45,10 +45,13 @@ private:
 	/** Creates a DMX Library asset. Returns nullptr if the library could not be created */
 	UDMXLibrary* CreateDMXLibraryAsset(UObject* Parent, const FName& Name, EObjectFlags Flags, const FString& InFilename);
 
+	/** Creates a General Scene Description from the MVR zip. Returns the General Scene Description or nullptr if not a valid MVR zip. */
+	UDMXMVRGeneralSceneDescription* CreateGeneralSsceneDescription(UDMXLibrary& DMXLibrary, const TSharedRef<FDMXZipper>& Zip) const;
+
 	/** Creates GDTF assets from the MVR */
 	TArray<UDMXImportGDTF*> CreateGDTFAssets(UObject* Parent, EObjectFlags Flags, const TSharedRef<FDMXZipper>& Zip, const UDMXMVRGeneralSceneDescription& GeneralSceneDescription);
 
-	/** Initializes the DMX Library from the General Scene Description and GDTF assets */
+	/** Initializes the DMX Library from the General Scene Description and GDTF assets. Updates the library in case it already contains patches */
 	void InitDMXLibrary(UDMXLibrary* DMXLibrary, const TArray<UDMXImportGDTF*>& GDTFAssets, UDMXMVRGeneralSceneDescription* GeneralSceneDescription) const;
 
 	/** Returns the MVR Asset Import Data for the Object, or nullptr if it cannot be obtained */
diff --git a/Engine/Plugins/VirtualProduction/DMX/DMXEngine/Source/DMXRuntime/Private/Library/DMXLibrary.cpp b/Engine/Plugins/VirtualProduction/DMX/DMXEngine/Source/DMXRuntime/Private/Library/DMXLibrary.cpp
--- a/Engine/Plugins/VirtualProduction/DMX/DMXEngine/Source/DMXRuntime/Private/Library/DMXLibrary.cpp
+++ b/Engine/Plugins/VirtualProduction/DMX/DMXEngine/Source/DMXRuntime/Private/Library/DMXLibrary.cpp
@@ -569,13 +569,6 @@ void UDMXLibrary::UpdatePorts()
 
 void UDMXLibrary::SetMVRGeneralSceneDescription(UDMXMVRGeneralSceneDescription* NewGeneralSceneDescription)
 {
-	// Remove all Fixture Patches
-	TArray<UDMXEntityFixturePatch*> FixturePatches = GetEntitiesTypeCast<UDMXEntityFixturePatch>();
-	for (UDMXEntityFixturePatch* FixturePatch : FixturePatches)
-	{
-		UDMXEntityFixturePatch::RemoveFixturePatchFromLibrary(FDMXEntityFixturePatchRef(FixturePatch));
-	}
-
 	GeneralSceneDescription = NewGeneralSceneDescription;
 }
 
