DMX - Add sort options for all columns in the MVR Fixture List

diff --git a/Engine/Plugins/VirtualProduction/DMX/DMXEngine/Source/DMXEditor/Private/Widgets/FixturePatch/SDMXMVRFixtureList.cpp b/Engine/Plugins/VirtualProduction/DMX/DMXEngine/Source/DMXEditor/Private/Widgets/FixturePatch/SDMXMVRFixtureList.cpp
--- a/Engine/Plugins/VirtualProduction/DMX/DMXEngine/Source/DMXEditor/Private/Widgets/FixturePatch/SDMXMVRFixtureList.cpp
+++ b/Engine/Plugins/VirtualProduction/DMX/DMXEngine/Source/DMXEditor/Private/Widgets/FixturePatch/SDMXMVRFixtureList.cpp
@@ -16,6 +16,7 @@
 
 #include "ScopedTransaction.h"
 #include "TimerManager.h"
+#include "Algo/Sort.h"
 #include "Framework/Commands/GenericCommands.h"
 #include "Framework/Commands/UICommandList.h"
 #include "Framework/MultiBox/MultiBoxBuilder.h"
@@ -746,6 +747,8 @@ TSharedRef<SHeaderRow> SDMXMVRFixtureList::GenerateHeaderRow()
 	HeaderRow->AddColumn(
 		SHeaderRow::FColumn::FArguments()
 		.ColumnId(FDMXMVRFixtureListCollumnIDs::FixturePatchName)
+		.SortMode(this, &SDMXMVRFixtureList::GetColumnSortMode, FDMXMVRFixtureListCollumnIDs::FixturePatchName)
+		.OnSort(this, &SDMXMVRFixtureList::SortByColumnID)
 		.DefaultLabel(LOCTEXT("FixturePatchNameColumnLabel", "Fixture Patch"))
 		.FillWidth(0.2f)
 	);
@@ -769,6 +772,8 @@ TSharedRef<SHeaderRow> SDMXMVRFixtureList::GenerateHeaderRow()
 	HeaderRow->AddColumn(
 		SHeaderRow::FColumn::FArguments()
 		.ColumnId(FDMXMVRFixtureListCollumnIDs::MVRFixtureName)
+		.SortMode(this, &SDMXMVRFixtureList::GetColumnSortMode, FDMXMVRFixtureListCollumnIDs::MVRFixtureName)
+		.OnSort(this, &SDMXMVRFixtureList::SortByColumnID)
 		.DefaultLabel(LOCTEXT("NameColumnLabel", "Fixture Name"))
 		.FillWidth(0.2f)
 	);
@@ -776,6 +781,8 @@ TSharedRef<SHeaderRow> SDMXMVRFixtureList::GenerateHeaderRow()
 	HeaderRow->AddColumn(
 		SHeaderRow::FColumn::FArguments()
 		.ColumnId(FDMXMVRFixtureListCollumnIDs::FixtureType)
+		.SortMode(this, &SDMXMVRFixtureList::GetColumnSortMode, FDMXMVRFixtureListCollumnIDs::FixtureType)
+		.OnSort(this, &SDMXMVRFixtureList::SortByColumnID)
 		.DefaultLabel(LOCTEXT("FixtureTypeColumnLabel", "FixtureType"))
 		.FillWidth(0.2f)
 	);
@@ -783,6 +790,8 @@ TSharedRef<SHeaderRow> SDMXMVRFixtureList::GenerateHeaderRow()
 	HeaderRow->AddColumn(
 		SHeaderRow::FColumn::FArguments()
 		.ColumnId(FDMXMVRFixtureListCollumnIDs::Mode)
+		.SortMode(this, &SDMXMVRFixtureList::GetColumnSortMode, FDMXMVRFixtureListCollumnIDs::Mode)
+		.OnSort(this, &SDMXMVRFixtureList::SortByColumnID)
 		.DefaultLabel(LOCTEXT("ModeColumnLabel", "Mode"))
 		.FillWidth(0.2f)
 	);
@@ -886,50 +895,92 @@ void SDMXMVRFixtureList::SortByColumnID(const EColumnSortPriority::Type SortPrio
 	SortMode = InSortMode;
 	SortedByColumnID = ColumnId;
 
-	if (ColumnId == FDMXMVRFixtureListCollumnIDs::FixtureID)
+	const bool bAscending = InSortMode == EColumnSortMode::Ascending ? true : false;
+
+	if (ColumnId == FDMXMVRFixtureListCollumnIDs::FixturePatchName)
 	{
-		if (InSortMode == EColumnSortMode::Descending)
-		{
-			ListSource.Sort([](const TSharedPtr<FDMXMVRFixtureListItem>& ItemA, const TSharedPtr<FDMXMVRFixtureListItem>& ItemB)
-				{
-					return !(ItemA->GetFixtureID() <= ItemB->GetFixtureID());
-				});
-		}
-		else if (InSortMode == EColumnSortMode::Ascending)
-		{
-			ListSource.Sort([](const TSharedPtr<FDMXMVRFixtureListItem>& ItemA, const TSharedPtr<FDMXMVRFixtureListItem>& ItemB)
-				{
-					return !(ItemA->GetFixtureID() >= ItemB->GetFixtureID());
-				});
-		}
+		Algo::Sort(ListSource, [bAscending](const TSharedPtr<FDMXMVRFixtureListItem>& ItemA, const TSharedPtr<FDMXMVRFixtureListItem>& ItemB)
+			{
+				const FString FixturePatchNameA = ItemA->GetFixturePatchName();
+				const FString FixturePatchNameB = ItemB->GetFixturePatchName();
+
+				const bool bIsGreater = FixturePatchNameA >= FixturePatchNameB;
+				return bAscending ? !bIsGreater : bIsGreater;
+			});
 	}
-	else if(ColumnId == FDMXMVRFixtureListCollumnIDs::Patch)
+	else if (ColumnId == FDMXMVRFixtureListCollumnIDs::FixtureID)
 	{
-		if (InSortMode == EColumnSortMode::Descending)
-		{
-			ListSource.StableSort([](const TSharedPtr<FDMXMVRFixtureListItem>& ItemA, const TSharedPtr<FDMXMVRFixtureListItem>& ItemB)
+		Algo::Sort(ListSource, [bAscending](const TSharedPtr<FDMXMVRFixtureListItem>& ItemA, const TSharedPtr<FDMXMVRFixtureListItem>& ItemB)
+			{
+				bool bIsGreater = [ItemA, ItemB]()
 				{
-					UDMXEntityFixturePatch* FixturePatchA = ItemA->GetFixturePatch();
-					UDMXEntityFixturePatch* FixturePatchB = ItemB->GetFixturePatch();
+					const FString FixtureIDStringA = ItemA->GetFixtureID();
+					const FString FixtureIDStringB = ItemB->GetFixtureID();
 
-					const bool bIsUniverseIDLarger = ItemA->GetUniverse() > ItemB->GetUniverse();
-					const bool bIsSameUniverse = ItemA->GetUniverse() == ItemB->GetUniverse();
-					const bool bAreAddressesLarger = ItemA->GetAddress() > ItemB->GetAddress();
-					return bIsUniverseIDLarger || (bIsSameUniverse && bAreAddressesLarger);
-				});
-		}
-		else if (InSortMode == EColumnSortMode::Ascending)
-		{
-			ListSource.StableSort([](const TSharedPtr<FDMXMVRFixtureListItem>& ItemA, const TSharedPtr<FDMXMVRFixtureListItem>& ItemB)
-				{
-					UDMXEntityFixturePatch* FixturePatchA = ItemA->GetFixturePatch();
-					UDMXEntityFixturePatch* FixturePatchB = ItemB->GetFixturePatch();
-					const bool bIsUniverseIDSmaller = ItemA->GetUniverse() < ItemB->GetUniverse();
-					const bool bIsSameUniverse = ItemA->GetUniverse() == ItemB->GetUniverse();
-					const bool bAreAddressesSmaller = ItemA->GetAddress() < ItemB->GetAddress();
-					return bIsUniverseIDSmaller || (bIsSameUniverse && bAreAddressesSmaller);
-				});
-		}
+					int32 FixtureIDA = 0;
+					int32 FixtureIDB = 0;
+
+					const bool bCanParseA = LexTryParseString(FixtureIDA, *FixtureIDStringA);
+					const bool bCanParseB = LexTryParseString(FixtureIDB, *FixtureIDStringB);
+					
+					const bool bIsNumeric = bCanParseA && bCanParseB;
+					if (bIsNumeric)
+					{
+						return FixtureIDA >= FixtureIDB;
+					}
+					else
+					{
+						return FixtureIDStringA >= FixtureIDStringB;
+					}
+				}();
+
+				return bAscending ? !bIsGreater : bIsGreater;
+			});
+	}
+	else if (ColumnId == FDMXMVRFixtureListCollumnIDs::MVRFixtureName)
+	{
+		Algo::Sort(ListSource, [bAscending](const TSharedPtr<FDMXMVRFixtureListItem>& ItemA, const TSharedPtr<FDMXMVRFixtureListItem>& ItemB)
+			{
+				const FString MVRFixtureNameA = ItemA->GetMVRFixtureName();
+				const FString MVRFixtureNameB = ItemB->GetMVRFixtureName();
+
+				const bool bIsGreater = MVRFixtureNameA >= MVRFixtureNameB;
+				return bAscending ? !bIsGreater : bIsGreater;
+			});
+	}
+	else if (ColumnId == FDMXMVRFixtureListCollumnIDs::FixtureType)
+	{
+		Algo::Sort(ListSource, [bAscending](const TSharedPtr<FDMXMVRFixtureListItem>& ItemA, const TSharedPtr<FDMXMVRFixtureListItem>& ItemB)
+			{
+				const FString FixtureTypeA = ItemA->GetFixtureType()->Name;
+				const FString FixtureTypeB = ItemB->GetFixtureType()->Name;
+
+				const bool bIsGreater = FixtureTypeA >= FixtureTypeB;
+				return bAscending ? !bIsGreater : bIsGreater;
+			});
+	}
+	else if (ColumnId == FDMXMVRFixtureListCollumnIDs::Mode)
+	{
+		Algo::Sort(ListSource, [bAscending](const TSharedPtr<FDMXMVRFixtureListItem>& ItemA, const TSharedPtr<FDMXMVRFixtureListItem>& ItemB)
+			{
+				const bool bIsGreater = ItemA->GetModeIndex() >= ItemB->GetModeIndex();
+				return bAscending ? !bIsGreater : bIsGreater;
+			});
+	}
+	else if (ColumnId == FDMXMVRFixtureListCollumnIDs::Patch)
+	{
+		Algo::Sort(ListSource, [bAscending](const TSharedPtr<FDMXMVRFixtureListItem>& ItemA, const TSharedPtr<FDMXMVRFixtureListItem>& ItemB)
+			{
+				const UDMXEntityFixturePatch* FixturePatchA = ItemA->GetFixturePatch();
+				const UDMXEntityFixturePatch* FixturePatchB = ItemB->GetFixturePatch();
+
+				const bool bIsUniverseIDGreater = ItemA->GetUniverse() > ItemB->GetUniverse();
+				const bool bIsSameUniverse = ItemA->GetUniverse() == ItemB->GetUniverse();
+				const bool bAreAddressesGreater = ItemA->GetAddress() > ItemB->GetAddress();
+
+				const bool bIsGreater = bIsUniverseIDGreater || (bIsSameUniverse && bAreAddressesGreater);
+				return bAscending ? !bIsGreater : bIsGreater;
+			});
 	}
 
 	ListView->RequestListRefresh();
