DMX - Fix Typo in Patch Tool

diff --git a/Engine/Plugins/VirtualProduction/DMX/DMXEngine/Source/DMXEditor/Private/Widgets/PatchTool/SDMXPatchTool.cpp b/Engine/Plugins/VirtualProduction/DMX/DMXEngine/Source/DMXEditor/Private/Widgets/PatchTool/SDMXPatchTool.cpp
--- a/Engine/Plugins/VirtualProduction/DMX/DMXEngine/Source/DMXEditor/Private/Widgets/PatchTool/SDMXPatchTool.cpp
+++ b/Engine/Plugins/VirtualProduction/DMX/DMXEngine/Source/DMXEditor/Private/Widgets/PatchTool/SDMXPatchTool.cpp
@@ -57,7 +57,7 @@ void SDMXPatchTool::Construct(const FArguments& InArgs)
 				.MaxDesiredWidth(160.f)
 				[
 					SNew(STextBlock)
-					.Text(LOCTEXT("DMXLibraryComboboxLabel", "DXM Library"))
+					.Text(LOCTEXT("DMXLibraryComboboxLabel", "DMX Library"))
 					.Font(FCoreStyle::GetDefaultFontStyle("Regular", 14))
 				]
 			]
