DMX - Move Receive DMX in Editor boolean to top of advanced category

diff --git a/Engine/Plugins/VirtualProduction/DMX/DMXEngine/Source/DMXRuntime/Public/Game/DMXComponent.h b/Engine/Plugins/VirtualProduction/DMX/DMXEngine/Source/DMXRuntime/Public/Game/DMXComponent.h
--- a/Engine/Plugins/VirtualProduction/DMX/DMXEngine/Source/DMXRuntime/Public/Game/DMXComponent.h
+++ b/Engine/Plugins/VirtualProduction/DMX/DMXEngine/Source/DMXRuntime/Public/Game/DMXComponent.h
@@ -70,7 +70,7 @@ protected:
 	void UpdateTickEnabled();
 
 	/** If true, the component will receive DMX from the patch */
-	UPROPERTY(BlueprintReadOnly, EditAnywhere, AdvancedDisplay, Category = "DMX")
+	UPROPERTY(BlueprintReadOnly, EditAnywhere, AdvancedDisplay, Category = "DMX", Meta = (DisplayPriority = 0))
 	bool bReceiveDMXFromPatch;
 
 	// ~Begin UActorComponent interface
