DMX - Fix fixture patcher doesn't update when changing mode or channel properties in details

diff --git a/Engine/Plugins/VirtualProduction/DMX/DMXEngine/Source/DMXEditor/Private/Widgets/FixturePatch/DMXFixturePatchNode.cpp b/Engine/Plugins/VirtualProduction/DMX/DMXEngine/Source/DMXEditor/Private/Widgets/FixturePatch/DMXFixturePatchNode.cpp
--- a/Engine/Plugins/VirtualProduction/DMX/DMXEngine/Source/DMXEditor/Private/Widgets/FixturePatch/DMXFixturePatchNode.cpp
+++ b/Engine/Plugins/VirtualProduction/DMX/DMXEngine/Source/DMXEditor/Private/Widgets/FixturePatch/DMXFixturePatchNode.cpp
@@ -405,7 +405,7 @@ TArray<TSharedPtr<SDMXFixturePatchFragment>> FDMXFixturePatchNode::GenerateWidge
 bool FDMXFixturePatchNode::OccupiesChannels(int32 Channel, int32 Span) const
 {
 	check(Span != 0);
-	if ((Channel + Span <= StartingChannel) || (Channel >= StartingChannel + ChannelSpan))
+	if ((Channel + Span <= GetStartingChannel()) || (Channel >= GetStartingChannel() + GetChannelSpan()))
 	{
 		return false;
 	}
@@ -452,6 +452,13 @@ int32 FDMXFixturePatchNode::GetChannelSpan() const
 	return -1;
 }
 
+bool FDMXFixturePatchNode::NeedsUpdateGrid() const
+{
+	return 
+		StartingChannel != GetStartingChannel() ||
+		ChannelSpan != GetChannelSpan();
+}
+
 void FDMXFixturePatchNode::OnSelectionChanged()
 {
 	const TSharedPtr<FDMXEditor> DMXEditor = WeakDMXEditor.Pin();
diff --git a/Engine/Plugins/VirtualProduction/DMX/DMXEngine/Source/DMXEditor/Private/Widgets/FixturePatch/DMXFixturePatchNode.h b/Engine/Plugins/VirtualProduction/DMX/DMXEngine/Source/DMXEditor/Private/Widgets/FixturePatch/DMXFixturePatchNode.h
--- a/Engine/Plugins/VirtualProduction/DMX/DMXEngine/Source/DMXEditor/Private/Widgets/FixturePatch/DMXFixturePatchNode.h
+++ b/Engine/Plugins/VirtualProduction/DMX/DMXEngine/Source/DMXEditor/Private/Widgets/FixturePatch/DMXFixturePatchNode.h
@@ -51,10 +51,10 @@ public:
 	/** Returns the Universe ID the Node currently resides in. Returns a negative Value if not assigned to a Universe. */
 	int32 GetUniverseID() const;
 
-	/** Returns the Starting Channel the Node currently resides on */
+	/** Returns the Starting Channel of the Patch */
 	int32 GetStartingChannel() const;
 
-	/** Returns the Channel Span the Node currently occupies */
+	/** Returns the Channel Span of the Patch */
 	int32 GetChannelSpan() const;
 
 	/** Sets the ZOrder of the Node */
@@ -63,6 +63,9 @@ public:
 	/** Gets the ZOrder of the Node */
 	int32 GetZOrder() const { return ZOrder; }
 
+	/** Returns true if the patch was moved on the grid and new widgets need be generated. */
+	bool NeedsUpdateGrid() const;
+
 private:
 	/** Called when a patch node was selected */
 	void OnSelectionChanged();
@@ -70,10 +73,10 @@ private:
 	/** Universe the patch is assigned to */
 	TSharedPtr<SDMXPatchedUniverse> UniverseWidget;
 
-	/** Starting channel of the patch */	
+	/** Cached starting channel of the patch */	
 	int32 StartingChannel = 0;
 
-	/** Channel span of the patch */
+	/** Cached channel span of the patch */
 	int32 ChannelSpan = 0;
 
 	/** Last transacted Universe ID, required for propert undo/redo */
diff --git a/Engine/Plugins/VirtualProduction/DMX/DMXEngine/Source/DMXEditor/Private/Widgets/FixturePatch/SDMXPatchedUniverse.cpp b/Engine/Plugins/VirtualProduction/DMX/DMXEngine/Source/DMXEditor/Private/Widgets/FixturePatch/SDMXPatchedUniverse.cpp
--- a/Engine/Plugins/VirtualProduction/DMX/DMXEngine/Source/DMXEditor/Private/Widgets/FixturePatch/SDMXPatchedUniverse.cpp
+++ b/Engine/Plugins/VirtualProduction/DMX/DMXEngine/Source/DMXEditor/Private/Widgets/FixturePatch/SDMXPatchedUniverse.cpp
@@ -324,12 +324,18 @@ void SDMXPatchedUniverse::OnFixturePatchChanged(const UDMXEntityFixturePatch* Fi
 
 	if (NodePtr && FixturePatch && UniverseID != FixturePatch->GetUniverseID())
 	{
+		// Patch was moved out of this universe
 		Unpatch(*NodePtr);
 	}
 	else if (!NodePtr && UniverseID == FixturePatch->GetUniverseID())
 	{
+		// Patch was moved to this universe
 		SharedData->SelectUniverse(UniverseID);
 	}
+	else if (NodePtr && (*NodePtr)->NeedsUpdateGrid())
+	{
+		RequestRefresh();
+	}
 }
 
 void SDMXPatchedUniverse::SetUniverseIDInternal(int32 NewUniverseID)
