DMX - Fix Multiple warnings "Failed to read file ..." when opening the project with Datasmith MVR plugin enabled

diff --git a/Engine/Plugins/VirtualProduction/DMX/DMXEngine/Source/DMXEditor/Private/DMXEditorStyle.cpp b/Engine/Plugins/VirtualProduction/DMX/DMXEngine/Source/DMXEditor/Private/DMXEditorStyle.cpp
--- a/Engine/Plugins/VirtualProduction/DMX/DMXEngine/Source/DMXEditor/Private/DMXEditorStyle.cpp
+++ b/Engine/Plugins/VirtualProduction/DMX/DMXEngine/Source/DMXEditor/Private/DMXEditorStyle.cpp
@@ -135,10 +135,10 @@ FDMXEditorStyle::FDMXEditorStyle()
 		static const FLinearColor DefeaultFaderForeColor = FLinearColor::FromSRGBColor(FColor::FromHex("ffffff"));
 
 		Set("DMXEditor.OutputConsole.Fader", FSpinBoxStyle(FAppStyle::GetWidgetStyle<FSpinBoxStyle>("SpinBox"))
-			.SetBackgroundBrush(CORE_BOX_BRUSH("Slate/Common/Spinbox", FMargin(4.f / 16.f), DefaultFaderBackColor))
-			.SetHoveredBackgroundBrush(CORE_BOX_BRUSH("Slate/Common/Spinbox", FMargin(4.f / 16.f), DefaultFaderBackColor))
-			.SetActiveFillBrush(CORE_BOX_BRUSH("Slate/Common/Spinbox_Fill", FMargin(4.f / 16.f), DefaultFaderFillColor))
-			.SetInactiveFillBrush(CORE_BOX_BRUSH("Slate/Common/Spinbox_Fill", FMargin(4.f / 16.f), DefaultFaderFillColor))
+			.SetBackgroundBrush(CORE_BOX_BRUSH("Common/Spinbox", FMargin(4.f / 16.f), DefaultFaderBackColor))
+			.SetHoveredBackgroundBrush(CORE_BOX_BRUSH("Common/Spinbox", FMargin(4.f / 16.f), DefaultFaderBackColor))
+			.SetActiveFillBrush(CORE_BOX_BRUSH("Common/Spinbox_Fill", FMargin(4.f / 16.f), DefaultFaderFillColor))
+			.SetInactiveFillBrush(CORE_BOX_BRUSH("Common/Spinbox_Fill", FMargin(4.f / 16.f), DefaultFaderFillColor))
 			.SetForegroundColor(DefeaultFaderForeColor)
 			.SetArrowsImage(FSlateNoResource()));
 	}
@@ -146,11 +146,11 @@ FDMXEditorStyle::FDMXEditorStyle()
 	// Fixture List
 	{
 		Set("MVRFixtureList.Row", FTableRowStyle()
-			.SetEvenRowBackgroundBrush(IMAGE_BRUSH("Common/Selection", Icon8x8, FLinearColor(1.f, 1.f, 1.f, .2f)))
-			.SetOddRowBackgroundBrush(IMAGE_BRUSH("Common/Selection", Icon8x8, FLinearColor(1.f, 1.f, 1.f, .3f)))
-			.SetEvenRowBackgroundHoveredBrush(IMAGE_BRUSH("Common/Selection", Icon8x8, FLinearColor(1.f, 1.f, 1.f, .4f)))
-			.SetOddRowBackgroundHoveredBrush(IMAGE_BRUSH("Common/Selection", Icon8x8, FLinearColor(1.f, 1.f, 1.f, .5f)))
-			.SetSelectorFocusedBrush(BORDER_BRUSH("Common/Selector", FMargin(4.f / 16.f), FLinearColor(1.f, 1.f, 1.f, .8f)))
+			.SetEvenRowBackgroundBrush(CORE_IMAGE_BRUSH("Common/Selection", Icon8x8, FLinearColor(1.f, 1.f, 1.f, .2f)))
+			.SetOddRowBackgroundBrush(CORE_IMAGE_BRUSH("Common/Selection", Icon8x8, FLinearColor(1.f, 1.f, 1.f, .3f)))
+			.SetEvenRowBackgroundHoveredBrush(CORE_IMAGE_BRUSH("Common/Selection", Icon8x8, FLinearColor(1.f, 1.f, 1.f, .4f)))
+			.SetOddRowBackgroundHoveredBrush(CORE_IMAGE_BRUSH("Common/Selection", Icon8x8, FLinearColor(1.f, 1.f, 1.f, .5f)))
+			.SetSelectorFocusedBrush(CORE_BORDER_BRUSH("Common/Selector", FMargin(4.f / 16.f), FLinearColor(1.f, 1.f, 1.f, .8f)))
 		);
 	}
 
