DMX - Add nice borders for the Fixture Patcher

diff --git a/Engine/Plugins/VirtualProduction/DMX/DMXEngine/Source/DMXEditor/Private/DMXEditorStyle.cpp b/Engine/Plugins/VirtualProduction/DMX/DMXEngine/Source/DMXEditor/Private/DMXEditorStyle.cpp
--- a/Engine/Plugins/VirtualProduction/DMX/DMXEngine/Source/DMXEditor/Private/DMXEditorStyle.cpp
+++ b/Engine/Plugins/VirtualProduction/DMX/DMXEngine/Source/DMXEditor/Private/DMXEditorStyle.cpp
@@ -49,8 +49,12 @@ FDMXEditorStyle::FDMXEditorStyle()
 	// Border Brushes
 	{
 		Set("DMXEditor.RoundedPropertyBorder", new FSlateRoundedBoxBrush(FLinearColor::Black, 4.f));
-	}
 
+		Set("FixturePatcher.FragmentBorder.Normal", new BORDER_BRUSH("FixturePatch_Border_Normal_8x", 2.f / 8.f));
+		Set("FixturePatcher.FragmentBorder.L", new BORDER_BRUSH("FixturePatch_Border_L_8x", 2.f / 8.f));
+		Set("FixturePatcher.FragmentBorder.R", new BORDER_BRUSH("FixturePatch_Border_R_8x", 2.f / 8.f));
+		Set("FixturePatcher.FragmentBorder.TB", new BORDER_BRUSH("FixturePatch_Border_TB_8x", 2.f / 8.f));
+	}
 
 	// Fonts
 	{
diff --git a/Engine/Plugins/VirtualProduction/DMX/DMXEngine/Source/DMXEditor/Private/Widgets/FixturePatch/SDMXFixturePatchFragment.cpp b/Engine/Plugins/VirtualProduction/DMX/DMXEngine/Source/DMXEditor/Private/Widgets/FixturePatch/SDMXFixturePatchFragment.cpp
--- a/Engine/Plugins/VirtualProduction/DMX/DMXEngine/Source/DMXEditor/Private/Widgets/FixturePatch/SDMXFixturePatchFragment.cpp
+++ b/Engine/Plugins/VirtualProduction/DMX/DMXEngine/Source/DMXEditor/Private/Widgets/FixturePatch/SDMXFixturePatchFragment.cpp
@@ -46,14 +46,15 @@ void SDMXFixturePatchFragment::Construct(const FArguments& InArgs, const TShared
 	if (bIsText)
 	{
 		SAssignNew(ContentBorder, SBorder)
+			.BorderImage(FAppStyle::GetBrush("NoBorder"))
 			.HAlign(HAlign_Fill)
 			.VAlign(VAlign_Fill)
-			.Padding(FMargin(3.0f, 2.0f, 0.f, 1.0f))
 			.Visibility(EVisibility::SelfHitTestInvisible);
 	}
 	else
 	{
 		SAssignNew(ContentBorder, SBorder)
+			.Padding(1.f)
 			.HAlign(HAlign_Fill)
 			.VAlign(VAlign_Fill)
 			.Visibility(EVisibility::SelfHitTestInvisible)
@@ -91,13 +92,14 @@ void SDMXFixturePatchFragment::Construct(const FArguments& InArgs, const TShared
 				})
 			[
 				SNew(SBorder)
+				.Padding(1.f)
 				.HAlign(HAlign_Fill)
 				.VAlign(VAlign_Fill)
 				.Visibility(EVisibility::SelfHitTestInvisible)
 				.BorderImage(&GetBorderImage())
 				.BorderBackgroundColor_Lambda([this]()
 					{
-						return bHighlight ? FLinearColor(1.f, 1.f, 1.0f) : FLinearColor(.2f, .2f, .2f);
+						return bHighlight ? FLinearColor(.8f, .8f, .8f) : FLinearColor(.2f, .2f, .2f);
 					})
 				[	
 					ContentBorder.ToSharedRef()
@@ -215,19 +217,19 @@ const FSlateBrush& SDMXFixturePatchFragment::GetBorderImage() const
 	}
 	else if (bIsHead && bIsTail)
 	{
-		BorderImage = FCoreStyle::Get().GetBrush("Border");
+		BorderImage = FDMXEditorStyle::Get().GetBrush("FixturePatcher.FragmentBorder.Normal");
 	}
 	else if (bIsHead)
 	{
-		BorderImage = FAppStyle::GetBrush("Profiler.EventGraph.Border.L");
+		BorderImage = FDMXEditorStyle::Get().GetBrush("FixturePatcher.FragmentBorder.L");
 	}
 	else if (bIsTail)
 	{
-		BorderImage = FAppStyle::GetBrush("Profiler.EventGraph.Border.R");
+		BorderImage = FDMXEditorStyle::Get().GetBrush("FixturePatcher.FragmentBorder.R");
 	}
 	else
 	{
-		BorderImage = FAppStyle::GetBrush("Profiler.EventGraph.Border.TB");
+		BorderImage = FDMXEditorStyle::Get().GetBrush("FixturePatcher.FragmentBorder.TB");
 	}
 	check(BorderImage);
 
