DMX - Remove Fixture Name column, instead use Fixture Patch name as Fixture Name

diff --git a/Engine/Plugins/VirtualProduction/DMX/DMXEngine/Source/DMXEditor/Private/Widgets/FixturePatch/DMXMVRFixtureListItem.cpp b/Engine/Plugins/VirtualProduction/DMX/DMXEngine/Source/DMXEditor/Private/Widgets/FixturePatch/DMXMVRFixtureListItem.cpp
--- a/Engine/Plugins/VirtualProduction/DMX/DMXEngine/Source/DMXEditor/Private/Widgets/FixturePatch/DMXMVRFixtureListItem.cpp
+++ b/Engine/Plugins/VirtualProduction/DMX/DMXEngine/Source/DMXEditor/Private/Widgets/FixturePatch/DMXMVRFixtureListItem.cpp
@@ -38,12 +38,15 @@ FDMXMVRFixtureListItem::FDMXMVRFixtureListItem(TWeakPtr<FDMXEditor> InDMXEditor,
 			return FixturePatch->GetMVRFixtureUUID() == MVRFixtureUUID;
 		});
 
-	if (ensureAlwaysMsgf(FixturePatchPtr, TEXT("Trying to create an MVR Fixture List Item, but there's no corresponding Fixture Patch for the MVR Fixture UUID.")))
+	if (!ensureAlwaysMsgf(FixturePatchPtr, TEXT("Trying to create an MVR Fixture List Item, but there's no corresponding Fixture Patch for the MVR Fixture UUID.")))
 	{
-		WeakFixturePatch = *FixturePatchPtr;
+		return;
 	}
-	
+	WeakFixturePatch = *FixturePatchPtr;
+
+	// Keep the MVR Fixture Name sync with the Fixture Patch name
 	MVRFixtureNode = FindMVRFixture();
+	MVRFixtureNode->Name = WeakFixturePatch->Name;
 }
 
 const FGuid& FDMXMVRFixtureListItem::GetMVRUUID() const
@@ -90,6 +93,10 @@ void FDMXMVRFixtureListItem::SetFixturePatchName(const FString& InDesiredName, F
 		FixturePatch->PreEditChange(UDMXEntityFixturePatch::StaticClass()->FindPropertyByName(GET_MEMBER_NAME_CHECKED(UDMXEntity, Name)));
 
 		FixturePatch->SetName(InDesiredName);
+		OutNewName = FixturePatch->Name;
+
+		MVRFixtureNode = FindMVRFixture();
+		MVRFixtureNode->Name = FixturePatch->Name;
 
 		FixturePatch->PostEditChange();
 	}
@@ -112,27 +119,6 @@ void FDMXMVRFixtureListItem::SetFixtureID(int32 InFixtureID)
 
 		DMXLibrary->PostEditChange();
 	}
-}	
-
-const FString& FDMXMVRFixtureListItem::GetMVRFixtureName() const
-{
-	return MVRFixtureNode->Name;
-}
-
-bool FDMXMVRFixtureListItem::SetMVRFixtureName(const FString& Name)
-{
-	UDMXLibrary* DMXLibrary = GetDMXLibrary();
-	if (!Name.IsEmpty() && MVRFixtureNode && DMXLibrary)
-	{
-		const FScopedTransaction SetMVRFixtureNameTransaction(LOCTEXT("SetMVRFixtureNameTransaction", "Set MVR Fixture Name"));
-		DMXLibrary->PreEditChange(UDMXLibrary::StaticClass()->FindPropertyByName(UDMXLibrary::GetGeneralSceneDescriptionPropertyName()));
-
-		MVRFixtureNode->Name = Name;
-
-		DMXLibrary->PostEditChange();
-	}
-
-	return false;
 }
 
 UDMXEntityFixtureType* FDMXMVRFixtureListItem::GetFixtureType() const
diff --git a/Engine/Plugins/VirtualProduction/DMX/DMXEngine/Source/DMXEditor/Private/Widgets/FixturePatch/DMXMVRFixtureListItem.h b/Engine/Plugins/VirtualProduction/DMX/DMXEngine/Source/DMXEditor/Private/Widgets/FixturePatch/DMXMVRFixtureListItem.h
--- a/Engine/Plugins/VirtualProduction/DMX/DMXEngine/Source/DMXEditor/Private/Widgets/FixturePatch/DMXMVRFixtureListItem.h
+++ b/Engine/Plugins/VirtualProduction/DMX/DMXEngine/Source/DMXEditor/Private/Widgets/FixturePatch/DMXMVRFixtureListItem.h
@@ -43,12 +43,6 @@ public:
 	/** Sets the Fixture ID. Note, as other hard- and software, we allow to set integer values only. */
 	void SetFixtureID(int32 InFixtureID);
 
-	/** Returns the Name of the MVR Fixture */
-	const FString& GetMVRFixtureName() const;
-
-	/** Sets the Name of the MVR Fixture */
-	bool SetMVRFixtureName(const FString& Name);
-
 	/** Returns the Fixture Type of the MVR Fixture */
 	UDMXEntityFixtureType* GetFixtureType() const;
 	
diff --git a/Engine/Plugins/VirtualProduction/DMX/DMXEngine/Source/DMXEditor/Private/Widgets/FixturePatch/SDMXMVRFixtureList.cpp b/Engine/Plugins/VirtualProduction/DMX/DMXEngine/Source/DMXEditor/Private/Widgets/FixturePatch/SDMXMVRFixtureList.cpp
--- a/Engine/Plugins/VirtualProduction/DMX/DMXEngine/Source/DMXEditor/Private/Widgets/FixturePatch/SDMXMVRFixtureList.cpp
+++ b/Engine/Plugins/VirtualProduction/DMX/DMXEngine/Source/DMXEditor/Private/Widgets/FixturePatch/SDMXMVRFixtureList.cpp
@@ -302,7 +302,7 @@ private:
 	static FText MakeBeautifulItemText(const TSharedPtr<FDMXMVRFixtureListItem>& Item)
 	{
 		const FString AddressesString = FString::FromInt(Item->GetUniverse()) + TEXT(".") + FString::FromInt(Item->GetAddress());
-		const FString ItemNameString = TEXT("'") + Item->GetMVRFixtureName() + TEXT("'");
+		const FString ItemNameString = TEXT("'") + Item->GetFixturePatchName() + TEXT("'");
 		const FString BeautifulItemString = ItemNameString + TEXT(" (") + AddressesString + TEXT(")");;
 		return FText::FromString(BeautifulItemString);
 	}
@@ -315,7 +315,6 @@ private:
 const FName FDMXMVRFixtureListCollumnIDs::Status = "Status";
 const FName FDMXMVRFixtureListCollumnIDs::FixturePatchName = "FixturePatchName";
 const FName FDMXMVRFixtureListCollumnIDs::FixtureID = "FixtureID";
-const FName FDMXMVRFixtureListCollumnIDs::MVRFixtureName = "MVRFixtureName";
 const FName FDMXMVRFixtureListCollumnIDs::FixtureType = "FixtureType";
 const FName FDMXMVRFixtureListCollumnIDs::Mode = "Mode";
 const FName FDMXMVRFixtureListCollumnIDs::Patch = "Patch";
@@ -769,15 +768,6 @@ TSharedRef<SHeaderRow> SDMXMVRFixtureList::GenerateHeaderRow()
 		.FillWidth(0.1f)
 	);
 
-	HeaderRow->AddColumn(
-		SHeaderRow::FColumn::FArguments()
-		.ColumnId(FDMXMVRFixtureListCollumnIDs::MVRFixtureName)
-		.SortMode(this, &SDMXMVRFixtureList::GetColumnSortMode, FDMXMVRFixtureListCollumnIDs::MVRFixtureName)
-		.OnSort(this, &SDMXMVRFixtureList::SortByColumnID)
-		.DefaultLabel(LOCTEXT("NameColumnLabel", "Fixture Name"))
-		.FillWidth(0.2f)
-	);
-
 	HeaderRow->AddColumn(
 		SHeaderRow::FColumn::FArguments()
 		.ColumnId(FDMXMVRFixtureListCollumnIDs::FixtureType)
@@ -822,11 +812,7 @@ void SDMXMVRFixtureList::SaveHeaderRowSettings()
 			{
 				EditorSettings->MVRFixtureListSettings.FixtureIDColumnWidth = Column.Width.Get();
 			}
-			if (Column.ColumnId == FDMXMVRFixtureListCollumnIDs::MVRFixtureName)
-			{
-				EditorSettings->MVRFixtureListSettings.NameColumnWidth = Column.Width.Get();
-			}
-			else if (Column.ColumnId == FDMXMVRFixtureListCollumnIDs::FixtureType)
+			if (Column.ColumnId == FDMXMVRFixtureListCollumnIDs::FixtureType)
 			{
 				EditorSettings->MVRFixtureListSettings.FixtureTypeColumnWidth = Column.Width.Get();
 			}
@@ -854,12 +840,6 @@ void SDMXMVRFixtureList::RestoresHeaderRowSettings()
 			HeaderRow->SetColumnWidth(FDMXMVRFixtureListCollumnIDs::FixtureID, FixtureIDColumnWidth);
 		}
 
-		const float NameColumnWidth = EditorSettings->MVRFixtureListSettings.NameColumnWidth;
-		if (NameColumnWidth > 20.f)
-		{
-			HeaderRow->SetColumnWidth(FDMXMVRFixtureListCollumnIDs::MVRFixtureName, NameColumnWidth);
-		}
-
 		const float FixtureTypeColumnWidth = EditorSettings->MVRFixtureListSettings.FixtureTypeColumnWidth;
 		if (FixtureTypeColumnWidth > 20.f)
 		{
@@ -937,17 +917,6 @@ void SDMXMVRFixtureList::SortByColumnID(const EColumnSortPriority::Type SortPrio
 				return bAscending ? !bIsGreater : bIsGreater;
 			});
 	}
-	else if (ColumnId == FDMXMVRFixtureListCollumnIDs::MVRFixtureName)
-	{
-		Algo::Sort(ListSource, [bAscending](const TSharedPtr<FDMXMVRFixtureListItem>& ItemA, const TSharedPtr<FDMXMVRFixtureListItem>& ItemB)
-			{
-				const FString MVRFixtureNameA = ItemA->GetMVRFixtureName();
-				const FString MVRFixtureNameB = ItemB->GetMVRFixtureName();
-
-				const bool bIsGreater = MVRFixtureNameA >= MVRFixtureNameB;
-				return bAscending ? !bIsGreater : bIsGreater;
-			});
-	}
 	else if (ColumnId == FDMXMVRFixtureListCollumnIDs::FixtureType)
 	{
 		Algo::Sort(ListSource, [bAscending](const TSharedPtr<FDMXMVRFixtureListItem>& ItemA, const TSharedPtr<FDMXMVRFixtureListItem>& ItemB)
diff --git a/Engine/Plugins/VirtualProduction/DMX/DMXEngine/Source/DMXEditor/Private/Widgets/FixturePatch/SDMXMVRFixtureList.h b/Engine/Plugins/VirtualProduction/DMX/DMXEngine/Source/DMXEditor/Private/Widgets/FixturePatch/SDMXMVRFixtureList.h
--- a/Engine/Plugins/VirtualProduction/DMX/DMXEngine/Source/DMXEditor/Private/Widgets/FixturePatch/SDMXMVRFixtureList.h
+++ b/Engine/Plugins/VirtualProduction/DMX/DMXEngine/Source/DMXEditor/Private/Widgets/FixturePatch/SDMXMVRFixtureList.h
@@ -30,7 +30,6 @@ struct FDMXMVRFixtureListCollumnIDs
 	static const FName Status;
 	static const FName FixturePatchName;
 	static const FName FixtureID;
-	static const FName MVRFixtureName;
 	static const FName FixtureType;
 	static const FName Mode;
 	static const FName Patch;
diff --git a/Engine/Plugins/VirtualProduction/DMX/DMXEngine/Source/DMXEditor/Private/Widgets/FixturePatch/SDMXMVRFixtureListRow.cpp b/Engine/Plugins/VirtualProduction/DMX/DMXEngine/Source/DMXEditor/Private/Widgets/FixturePatch/SDMXMVRFixtureListRow.cpp
--- a/Engine/Plugins/VirtualProduction/DMX/DMXEngine/Source/DMXEditor/Private/Widgets/FixturePatch/SDMXMVRFixtureListRow.cpp
+++ b/Engine/Plugins/VirtualProduction/DMX/DMXEngine/Source/DMXEditor/Private/Widgets/FixturePatch/SDMXMVRFixtureListRow.cpp
@@ -384,10 +384,6 @@ TSharedRef<SWidget> SDMXMVRFixtureListRow::GenerateWidgetForColumn(const FName&
 	{
 		return GenerateFixtureIDRow();
 	}
-	else if (ColumnName == FDMXMVRFixtureListCollumnIDs::MVRFixtureName)
-	{
-		return GenerateMVRFixtureNameRow();
-	}
 	else if (ColumnName == FDMXMVRFixtureListCollumnIDs::FixtureType)
 	{
 		return GenerateFixtureTypeRow();
@@ -551,58 +547,6 @@ void SDMXMVRFixtureListRow::OnFixtureIDCommitted(const FText& InNewText, ETextCo
 	}
 }
 
-TSharedRef<SWidget> SDMXMVRFixtureListRow::GenerateMVRFixtureNameRow()
-{
-	return
-		SNew(SBorder)
-		.HAlign(HAlign_Fill)
-		.VAlign(VAlign_Center)
-		.Padding(4.f)
-		.BorderImage(FAppStyle::GetBrush("NoBorder"))
-		.OnMouseDoubleClick(this, &SDMXMVRFixtureListRow::OnMVRFixtureNameBorderDoubleClicked)
-		[
-			SNew(SBorder)
-			.HAlign(HAlign_Fill)
-			.VAlign(VAlign_Center)
-			.IsEnabled(true)
-			.BorderImage(FDMXEditorStyle::Get().GetBrush("DMXEditor.RoundedPropertyBorder"))
-			[
-				SAssignNew(NameTextBlock, SInlineEditableTextBlock)
-				.Text_Lambda([this]()
-				{
-					const FString& Name = Item->GetMVRFixtureName();
-					return FText::FromString(Name);
-				})
-				.Font(FCoreStyle::GetDefaultFontStyle("Regular", 10))
-				.OnTextCommitted(this, &SDMXMVRFixtureListRow::OnMVRFixtureNameCommitted)
-				.IsSelected(IsSelected)
-			]
-		];
-
-}
-
-FReply SDMXMVRFixtureListRow::OnMVRFixtureNameBorderDoubleClicked(const FGeometry& InMyGeometry, const FPointerEvent& InMouseEvent)
-{
-	if (InMouseEvent.GetEffectingButton() == EKeys::LeftMouseButton)
-	{
-		if (NameTextBlock.IsValid())
-		{
-			NameTextBlock->EnterEditingMode();
-		}
-	}
-
-	return FReply::Handled();
-}
-
-void SDMXMVRFixtureListRow::OnMVRFixtureNameCommitted(const FText& InNewText, ETextCommit::Type InTextCommit)
-{
-	const FString NameString = InNewText.ToString();
-	if (Item->SetMVRFixtureName(NameString))
-	{
-		NameTextBlock->SetText(FText::FromString(NameString));
-	}
-}
-
 TSharedRef<SWidget> SDMXMVRFixtureListRow::GenerateFixtureTypeRow()
 {
 	UDMXLibrary* DMXLibrary = Item->GetDMXLibrary();
diff --git a/Engine/Plugins/VirtualProduction/DMX/DMXEngine/Source/DMXEditor/Private/Widgets/FixturePatch/SDMXMVRFixtureListRow.h b/Engine/Plugins/VirtualProduction/DMX/DMXEngine/Source/DMXEditor/Private/Widgets/FixturePatch/SDMXMVRFixtureListRow.h
--- a/Engine/Plugins/VirtualProduction/DMX/DMXEngine/Source/DMXEditor/Private/Widgets/FixturePatch/SDMXMVRFixtureListRow.h
+++ b/Engine/Plugins/VirtualProduction/DMX/DMXEngine/Source/DMXEditor/Private/Widgets/FixturePatch/SDMXMVRFixtureListRow.h
@@ -68,15 +68,6 @@ private:
 	/** Called when a Fixture ID was committed */
 	void OnFixtureIDCommitted(const FText& InNewText, ETextCommit::Type InTextCommit);
 
-	/** Generates the row that displays the MVR Fixture Name */
-	TSharedRef<SWidget> GenerateMVRFixtureNameRow();
-
-	/** Called when the Name Border was double-clicked */
-	FReply OnMVRFixtureNameBorderDoubleClicked(const FGeometry& InMyGeometry, const FPointerEvent& InMouseEvent);
-
-	/** Called when the MVR Fixture Name was committed */
-	void OnMVRFixtureNameCommitted(const FText& InNewText, ETextCommit::Type InTextCommit);
-
 	/** Generates the row that displays the Fixture Type */
 	TSharedRef<SWidget> GenerateFixtureTypeRow();
 
diff --git a/Engine/Plugins/VirtualProduction/DMX/DMXEngine/Source/DMXEditor/Private/Widgets/FixturePatch/SDMXMVRFixtureListToolbar.cpp b/Engine/Plugins/VirtualProduction/DMX/DMXEngine/Source/DMXEditor/Private/Widgets/FixturePatch/SDMXMVRFixtureListToolbar.cpp
--- a/Engine/Plugins/VirtualProduction/DMX/DMXEngine/Source/DMXEditor/Private/Widgets/FixturePatch/SDMXMVRFixtureListToolbar.cpp
+++ b/Engine/Plugins/VirtualProduction/DMX/DMXEngine/Source/DMXEditor/Private/Widgets/FixturePatch/SDMXMVRFixtureListToolbar.cpp
@@ -257,9 +257,7 @@ TArray<TSharedPtr<FDMXMVRFixtureListItem>> SDMXMVRFixtureListToolbar::FilterItem
 
 	Result.RemoveAll([this](const TSharedPtr<FDMXMVRFixtureListItem>& Item)
 		{
-			return
-				!Item->GetFixturePatchName().Contains(SearchString) &&
-				!Item->GetMVRFixtureName().Contains(SearchString);
+			return !Item->GetFixturePatchName().Contains(SearchString);
 		});
 
 	return Result;
diff --git a/Engine/Plugins/VirtualProduction/DMX/DMXEngine/Source/DMXEditor/Public/DMXEditorSettings.h b/Engine/Plugins/VirtualProduction/DMX/DMXEngine/Source/DMXEditor/Public/DMXEditorSettings.h
--- a/Engine/Plugins/VirtualProduction/DMX/DMXEngine/Source/DMXEditor/Public/DMXEditorSettings.h
+++ b/Engine/Plugins/VirtualProduction/DMX/DMXEngine/Source/DMXEditor/Public/DMXEditorSettings.h
@@ -19,10 +19,6 @@ struct FDMXMVRFixtureListSettings
 	UPROPERTY()
 	float FixtureIDColumnWidth = 0.f;
 
-	/** Width of the Name column */
-	UPROPERTY()
-	float NameColumnWidth = 0.f;
-
 	/** With of the Fixture Type column */
 	UPROPERTY()
 	float FixtureTypeColumnWidth = 0.f;
