DMX - Rework DMXMVRFixtureActorInterface so that it does not require a DMX dependency

diff --git a/Engine/Plugins/VirtualProduction/DMX/DMXEngine/DMXEngine.uplugin b/Engine/Plugins/VirtualProduction/DMX/DMXEngine/DMXEngine.uplugin
--- a/Engine/Plugins/VirtualProduction/DMX/DMXEngine/DMXEngine.uplugin
+++ b/Engine/Plugins/VirtualProduction/DMX/DMXEngine/DMXEngine.uplugin
@@ -34,6 +34,10 @@
 			"Name": "DatasmithContent",
 			"Enabled": true
 		},
+		{
+			"Name": "DMXModularFeatures",
+			"Enabled": true
+		},
 		{
 			"Name": "DMXProtocol",
 			"Enabled": true
diff --git a/Engine/Plugins/VirtualProduction/DMX/DMXEngine/Source/DMXRuntime/DMXRuntime.Build.cs b/Engine/Plugins/VirtualProduction/DMX/DMXEngine/Source/DMXRuntime/DMXRuntime.Build.cs
--- a/Engine/Plugins/VirtualProduction/DMX/DMXEngine/Source/DMXRuntime/DMXRuntime.Build.cs
+++ b/Engine/Plugins/VirtualProduction/DMX/DMXEngine/Source/DMXRuntime/DMXRuntime.Build.cs
@@ -24,6 +24,7 @@ public class DMXRuntime : ModuleRules
 			new string[]
 			{
 				"AssetRegistry",
+				"DMXFixtureActorInterface",
 				"Json",
 				"MovieScene",
 				"Projects",
diff --git a/Engine/Plugins/VirtualProduction/DMX/DMXEngine/Source/DMXRuntime/Private/MVR/DMXMVRFixtureActorLibrary.cpp b/Engine/Plugins/VirtualProduction/DMX/DMXEngine/Source/DMXRuntime/Private/MVR/DMXMVRFixtureActorLibrary.cpp
--- a/Engine/Plugins/VirtualProduction/DMX/DMXEngine/Source/DMXRuntime/Private/MVR/DMXMVRFixtureActorLibrary.cpp
+++ b/Engine/Plugins/VirtualProduction/DMX/DMXEngine/Source/DMXRuntime/Private/MVR/DMXMVRFixtureActorLibrary.cpp
@@ -2,8 +2,8 @@
 
 #include "MVR/DMXMVRFixtureActorLibrary.h"
 
+#include "DMXMVRFixtureActorInterface.h"
 #include "DMXRuntimeLog.h"
-#include "MVR/DMXMVRFixtureActorInterface.h"
 #include "Library/DMXEntityFixturePatch.h"
 
 #include "AssetRegistry/AssetRegistryModule.h"
diff --git a/Engine/Plugins/VirtualProduction/DMX/DMXEngine/Source/DMXRuntime/Private/MVR/DMXMVRSceneActor.cpp b/Engine/Plugins/VirtualProduction/DMX/DMXEngine/Source/DMXRuntime/Private/MVR/DMXMVRSceneActor.cpp
--- a/Engine/Plugins/VirtualProduction/DMX/DMXEngine/Source/DMXRuntime/Private/MVR/DMXMVRSceneActor.cpp
+++ b/Engine/Plugins/VirtualProduction/DMX/DMXEngine/Source/DMXRuntime/Private/MVR/DMXMVRSceneActor.cpp
@@ -2,6 +2,7 @@
 
 #include "MVR/DMXMVRSceneActor.h"
 
+#include "DMXMVRFixtureActorInterface.h"
 #include "DMXRuntimeLog.h"
 #include "Game/DMXComponent.h"
 #include "Library/DMXEntityFixturePatch.h"
@@ -9,7 +10,6 @@
 #include "Library/DMXLibrary.h"
 #include "MVR/DMXMVRFixtureActorLibrary.h"
 #include "MVR/DMXMVRAssetUserData.h"
-#include "MVR/DMXMVRFixtureActorInterface.h"
 #include "MVR/Types/DMXMVRFixtureNode.h"
 
 #include "DatasmithAssetUserData.h"
@@ -386,7 +386,7 @@ AActor* ADMXMVRSceneActor::ReplaceMVRActor(AActor* ActorToReplace, const TSubcla
 		UDMXEntityFixturePatch* FixturePatch = nullptr;
 		if (IDMXMVRFixtureActorInterface* MVRFixtureActorInterface = Cast<IDMXMVRFixtureActorInterface>(ActorToReplace))
 		{
-			FixturePatch = MVRFixtureActorInterface->Execute_OnMVRGetFixturePatch(ActorToReplace);
+			FixturePatch = GetFixturePatch(ActorToReplace);
 		}
 		
 		if (!FixturePatch)
@@ -424,18 +424,15 @@ AActor* ADMXMVRSceneActor::ReplaceMVRActor(AActor* ActorToReplace, const TSubcla
 
 UDMXEntityFixturePatch* ADMXMVRSceneActor::GetFixturePatch(AActor* Actor) const
 {
-	UDMXEntityFixturePatch* FixturePatch = nullptr;
-	if (IDMXMVRFixtureActorInterface* MVRFixtureActorInterface = Cast<IDMXMVRFixtureActorInterface>(Actor))
+	TArray<UDMXComponent*> DMXComponents;
+	Actor->GetComponents<UDMXComponent>(DMXComponents);
+	if (!ensureAlwaysMsgf(!DMXComponents.IsEmpty(), TEXT("'%s' implements the DMXMVRFixtureActorInterface, but has no DMX component. The DMX Component is needed to patch and identify the fixture in the MVR Scene."), *Actor->GetName()))
 	{
-		FixturePatch = MVRFixtureActorInterface->Execute_OnMVRGetFixturePatch(Actor);
-	}
-	if (!FixturePatch)
-	{
-		if (UActorComponent* Component = Actor->GetComponentByClass(UDMXComponent::StaticClass()))
-		{
-			FixturePatch = CastChecked<UDMXComponent>(Component)->GetFixturePatch();
-		}
+		return nullptr;
 	}
+	ensureAlwaysMsgf(DMXComponents.Num() == 1, TEXT("'%s' implements the DMXMVRFixtureActorInterface, but has more than one DMX component. A single DMX component is required to clearly identify the fixture by MVR UUID in the MVR Scene."), *Actor->GetName());
+	
+	UDMXEntityFixturePatch* FixturePatch = DMXComponents[0]->GetFixturePatch();
 
 	return FixturePatch;
 }
@@ -447,14 +444,13 @@ void ADMXMVRSceneActor::SetFixturePatch(AActor* Actor, UDMXEntityFixturePatch* F
 		return;
 	}
 
-	// Set the patch either via the interface or via a present DMX Component.
-	// Prefer the interface way as it may further customize how the patch is set.
-	if (IDMXMVRFixtureActorInterface* MVRFixtureActorInterface = Cast<IDMXMVRFixtureActorInterface>(Actor))
-	{
-		MVRFixtureActorInterface->Execute_OnMVRSetFixturePatch(Actor, FixturePatch);
-	}
-	else if (UActorComponent* Component = Actor->GetComponentByClass(UDMXComponent::StaticClass()))
+	TArray<UDMXComponent*> DMXComponents;
+	Actor->GetComponents<UDMXComponent>(DMXComponents);
+	if (!ensureAlwaysMsgf(!DMXComponents.IsEmpty(), TEXT("'%s' implements the DMXMVRFixtureActorInterface, but has no DMX component. The DMX Component is needed to patch and identify the fixture in the MVR Scene."), *Actor->GetName()))
 	{
-		CastChecked<UDMXComponent>(Component)->SetFixturePatch(FixturePatch);
+		return;
 	}
+	ensureAlwaysMsgf(DMXComponents.Num() == 1, TEXT("'%s' implements the DMXMVRFixtureActorInterface, but has more than one DMX component. A single DMX component is required to clearly identify the fixture by MVR UUID in the MVR Scene."), *Actor->GetName());
+
+	DMXComponents[0]->SetFixturePatch(FixturePatch);
 }
diff --git a/Engine/Plugins/VirtualProduction/DMX/DMXEngine/Source/DMXRuntime/Public/MVR/DMXMVRFixtureActorInterface.h b/Engine/Plugins/VirtualProduction/DMX/DMXEngine/Source/DMXRuntime/Public/MVR/DMXMVRFixtureActorInterface.h
deleted file mode 100644
diff --git a/Engine/Plugins/VirtualProduction/DMX/DMXFixtures/DMXFixtures.uplugin b/Engine/Plugins/VirtualProduction/DMX/DMXFixtures/DMXFixtures.uplugin
--- a/Engine/Plugins/VirtualProduction/DMX/DMXFixtures/DMXFixtures.uplugin
+++ b/Engine/Plugins/VirtualProduction/DMX/DMXFixtures/DMXFixtures.uplugin
@@ -24,6 +24,10 @@
 			"Name": "DMXEngine",
 			"Enabled": true
 		},
+		{
+			"Name": "DMXModularFeatures",
+			"Enabled": true
+		},
 		{
 			"Name": "DMXProtocol",
 			"Enabled": true
diff --git a/Engine/Plugins/VirtualProduction/DMX/DMXFixtures/Source/DMXFixtures/DMXFixtures.Build.cs b/Engine/Plugins/VirtualProduction/DMX/DMXFixtures/Source/DMXFixtures/DMXFixtures.Build.cs
--- a/Engine/Plugins/VirtualProduction/DMX/DMXFixtures/Source/DMXFixtures/DMXFixtures.Build.cs
+++ b/Engine/Plugins/VirtualProduction/DMX/DMXFixtures/Source/DMXFixtures/DMXFixtures.Build.cs
@@ -11,6 +11,7 @@ public class DMXFixtures : ModuleRules
        PublicDependencyModuleNames.AddRange(new string[] { 
 		   "Core",
 		   "CoreUObject",
+			"DMXFixtureActorInterface",
 		   "DMXRuntime", 
 		   "DMXProtocol",
 		   "Engine",
diff --git a/Engine/Plugins/VirtualProduction/DMX/DMXFixtures/Source/DMXFixtures/Private/DMXFixtureActor.cpp b/Engine/Plugins/VirtualProduction/DMX/DMXFixtures/Source/DMXFixtures/Private/DMXFixtureActor.cpp
--- a/Engine/Plugins/VirtualProduction/DMX/DMXFixtures/Source/DMXFixtures/Private/DMXFixtureActor.cpp
+++ b/Engine/Plugins/VirtualProduction/DMX/DMXFixtures/Source/DMXFixtures/Private/DMXFixtureActor.cpp
@@ -50,7 +50,6 @@ ADMXFixtureActor::ADMXFixtureActor()
 	HasBeenInitialized = false;
 }
 
-
 void ADMXFixtureActor::PostLoad()
 {
 	Super::PostLoad();
@@ -62,8 +61,6 @@ void ADMXFixtureActor::PostLoad()
 	PRAGMA_ENABLE_DEPRECATION_WARNINGS
 }
 
-
-
 void ADMXFixtureActor::OnMVRGetSupportedDMXAttributes_Implementation(TArray<FName>& OutAttributeNames, TArray<FName>& OutMatrixAttributeNames) const
 {
 	for (UDMXFixtureComponent* DMXFixtureComponent : TInlineComponentArray<UDMXFixtureComponent*>(this))
@@ -75,18 +72,6 @@ void ADMXFixtureActor::OnMVRGetSupportedDMXAttributes_Implementation(TArray<FNam
 	}
 }
 
-void ADMXFixtureActor::OnMVRSetFixturePatch_Implementation(UDMXEntityFixturePatch* FixturePatch)
-{
-	check(DMX);
-	DMX->SetFixturePatch(FixturePatch);
-}
-
-UDMXEntityFixturePatch* ADMXFixtureActor::OnMVRGetFixturePatch_Implementation() const
-{
-	check(DMX);
-	return DMX->GetFixturePatch();
-}
-
 #if WITH_EDITOR
 void ADMXFixtureActor::PostEditChangeProperty(FPropertyChangedEvent& PropertyChangedEvent)
 {
diff --git a/Engine/Plugins/VirtualProduction/DMX/DMXFixtures/Source/DMXFixtures/Public/DMXFixtureActor.h b/Engine/Plugins/VirtualProduction/DMX/DMXFixtures/Source/DMXFixtures/Public/DMXFixtureActor.h
--- a/Engine/Plugins/VirtualProduction/DMX/DMXFixtures/Source/DMXFixtures/Public/DMXFixtureActor.h
+++ b/Engine/Plugins/VirtualProduction/DMX/DMXFixtures/Source/DMXFixtures/Public/DMXFixtureActor.h
@@ -8,9 +8,9 @@
 #include "DMXFixtureComponentDouble.h"
 #include "DMXFixtureComponentSingle.h"
 #include "DMXFixtureComponentColor.h"
+#include "DMXMVRFixtureActorInterface.h"
 #include "DMXProtocolTypes.h"
 #include "Library/DMXEntityFixturePatch.h"
-#include "MVR/DMXMVRFixtureActorInterface.h"
 
 #include "CoreMinimal.h"
 #include "Components/ArrowComponent.h"
@@ -49,8 +49,6 @@ public:
 
 	//~ Begin DMXMVRFixtureActorInterface interface
 	virtual void OnMVRGetSupportedDMXAttributes_Implementation(TArray<FName>& OutAttributeNames, TArray<FName>& OutMatrixAttributeNames) const override;
-	virtual void OnMVRSetFixturePatch_Implementation(UDMXEntityFixturePatch* FixturePatch) override;
-	virtual UDMXEntityFixturePatch* OnMVRGetFixturePatch_Implementation() const override;
 	//~ End DMXMVRFixtureActorInterface interface
 	
 	// UObject interface
diff --git a/Engine/Plugins/VirtualProduction/DMX/DMXModularFeatures/DMXModularFeatures.uplugin b/Engine/Plugins/VirtualProduction/DMX/DMXModularFeatures/DMXModularFeatures.uplugin
new file mode 100644
--- /dev/null
+++ b/Engine/Plugins/VirtualProduction/DMX/DMXModularFeatures/DMXModularFeatures.uplugin
@@ -0,0 +1,24 @@
+{
+	"FileVersion": 3,
+	"Version": 1,
+	"VersionName": "1.0",
+	"FriendlyName": "DMX Modular Features",
+	"Description": "Modular Features for DMX",
+	"Category": "Misc",
+	"CreatedBy": "Epic Games, Inc.",
+	"CreatedByURL": "http://epicgames.com",
+	"DocsURL": "",
+	"MarketplaceURL": "",
+	"SupportURL": "",
+	"CanContainContent": true,
+	"IsBetaVersion": true,
+	"Installed": false,
+	"Hidden": true,
+	"Modules": [
+		{
+			"Name": "DMXFixtureActorInterface",
+			"Type": "Runtime",
+			"LoadingPhase": "Default"
+		}
+	]
+}
diff --git a/Engine/Plugins/VirtualProduction/DMX/DMXModularFeatures/Source/DMXFixtureActorInterface/DMXFixtureActorInterface.build.cs b/Engine/Plugins/VirtualProduction/DMX/DMXModularFeatures/Source/DMXFixtureActorInterface/DMXFixtureActorInterface.build.cs
new file mode 100644
--- /dev/null
+++ b/Engine/Plugins/VirtualProduction/DMX/DMXModularFeatures/Source/DMXFixtureActorInterface/DMXFixtureActorInterface.build.cs
@@ -0,0 +1,16 @@
+// Copyright Epic Games, Inc. All Rights Reserved.
+
+using UnrealBuildTool;
+using System.IO;
+
+public class DMXFixtureActorInterface : ModuleRules
+{
+	public DMXFixtureActorInterface(ReadOnlyTargetRules ROTargetRules) : base(ROTargetRules)
+	{
+		PublicDependencyModuleNames.AddRange(
+			new string[] {
+				"Core",
+				"CoreUObject",
+			});
+	}
+}
diff --git a/Engine/Plugins/VirtualProduction/DMX/DMXModularFeatures/Source/DMXFixtureActorInterface/Private/DMXFixtureActorInterfaceModule.cpp b/Engine/Plugins/VirtualProduction/DMX/DMXModularFeatures/Source/DMXFixtureActorInterface/Private/DMXFixtureActorInterfaceModule.cpp
new file mode 100644
--- /dev/null
+++ b/Engine/Plugins/VirtualProduction/DMX/DMXModularFeatures/Source/DMXFixtureActorInterface/Private/DMXFixtureActorInterfaceModule.cpp
@@ -0,0 +1,5 @@
+// Copyright Epic Games, Inc. All Rights Reserved.
+
+#include "Modules/ModuleManager.h"
+
+IMPLEMENT_MODULE(FDefaultModuleImpl, DMXFixtureActorInterface)
diff --git a/Engine/Plugins/VirtualProduction/DMX/DMXModularFeatures/Source/DMXFixtureActorInterface/Public/DMXMVRFixtureActorInterface.h b/Engine/Plugins/VirtualProduction/DMX/DMXModularFeatures/Source/DMXFixtureActorInterface/Public/DMXMVRFixtureActorInterface.h
new file mode 100644
--- /dev/null
+++ b/Engine/Plugins/VirtualProduction/DMX/DMXModularFeatures/Source/DMXFixtureActorInterface/Public/DMXMVRFixtureActorInterface.h
@@ -0,0 +1,34 @@
+// Copyright Epic Games, Inc. All Rights Reserved.
+
+#pragma once
+
+#include "UObject/Interface.h"
+
+#include "DMXMVRFixtureActorInterface.generated.h"
+
+class UDMXEntityFixturePatch;
+
+
+/** 
+ * When implemented in an actor, MVR will find the Actor and consider it when auto-selecting Fixtures.
+ * 
+ * Note the Actor that implements this is required to have exactly one DMXComponent subobject. 
+ * It is the DMX Component that will write its MVR UUID to the Actor's MetaData and make the Actor identified as an MVR Actor.
+ */
+UINTERFACE(BlueprintType)
+class DMXFIXTUREACTORINTERFACE_API UDMXMVRFixtureActorInterface
+	: public UInterface
+{
+	GENERATED_BODY()
+
+};
+
+class DMXFIXTUREACTORINTERFACE_API IDMXMVRFixtureActorInterface
+{
+	GENERATED_BODY()
+
+public:
+	/** Should return the DMX Attributes the Actor supports */
+	UFUNCTION(BlueprintNativeEvent, BlueprintCallable, Category = "DMX", Meta = (DisplayName = "On MVR Get Supported DMX Attributes"))
+	void OnMVRGetSupportedDMXAttributes(TArray<FName>& OutAttributeNames, TArray<FName>& OutMatrixAttributeNames) const;
+};
