DMX - Fixture Actors now spawn at the correct location when changing the Actor Type spawned for a GDTF

diff --git a/Engine/Plugins/VirtualProduction/DMX/DMXEngine/Source/DMXRuntime/Private/MVR/DMXMVRSceneActor.cpp b/Engine/Plugins/VirtualProduction/DMX/DMXEngine/Source/DMXRuntime/Private/MVR/DMXMVRSceneActor.cpp
--- a/Engine/Plugins/VirtualProduction/DMX/DMXEngine/Source/DMXRuntime/Private/MVR/DMXMVRSceneActor.cpp
+++ b/Engine/Plugins/VirtualProduction/DMX/DMXEngine/Source/DMXRuntime/Private/MVR/DMXMVRSceneActor.cpp
@@ -309,7 +309,7 @@ void ADMXMVRSceneActor::HandleDefaultActorClassForGDTFChanged()
 }
 #endif // WITH_EDITOR
 
-AActor* ADMXMVRSceneActor::SpawnMVRActor(const TSubclassOf<AActor>&ActorClass, UDMXEntityFixturePatch* FixturePatch, const FTransform & Transform, AActor * Template)
+AActor* ADMXMVRSceneActor::SpawnMVRActor(const TSubclassOf<AActor>&ActorClass, UDMXEntityFixturePatch* FixturePatch, const FTransform& Transform, AActor* Template)
 {
 	UWorld* World = GetWorld();
 	if (!ensureAlwaysMsgf(World, TEXT("Trying to spawn MVR Fixture in MVR Scene, but the world is not valid.")))
@@ -348,7 +348,7 @@ AActor* ADMXMVRSceneActor::SpawnMVRActor(const TSubclassOf<AActor>&ActorClass, U
 
 	// Attach, set MVR Fixture UUID, set Fixture Patch, remember as a Related Actor
 	RootComponentOfChildActor->SetMobility(EComponentMobility::Movable);
-	RootComponentOfChildActor->AttachToComponent(MVRSceneRoot, FAttachmentTransformRules::KeepRelativeTransform);
+	RootComponentOfChildActor->AttachToComponent(MVRSceneRoot, FAttachmentTransformRules::KeepWorldTransform);
 	const FGuid& MVRFixtureUUID = FixturePatch->GetMVRFixtureUUID();
 	UDMXMVRAssetUserData::SetMVRAssetUserDataValueForKey(*NewFixtureActor, UDMXMVRAssetUserData::MVRFixtureUUIDMetaDataKey, MVRFixtureUUID.ToString());
 	SetFixturePatch(NewFixtureActor, FixturePatch);
