DMX - Duplicating fixture patches now results in a reasonable new selection

diff --git a/Engine/Plugins/VirtualProduction/DMX/DMXEngine/Source/DMXEditor/Private/Widgets/FixturePatch/SDMXMVRFixtureList.cpp b/Engine/Plugins/VirtualProduction/DMX/DMXEngine/Source/DMXEditor/Private/Widgets/FixturePatch/SDMXMVRFixtureList.cpp
--- a/Engine/Plugins/VirtualProduction/DMX/DMXEngine/Source/DMXEditor/Private/Widgets/FixturePatch/SDMXMVRFixtureList.cpp
+++ b/Engine/Plugins/VirtualProduction/DMX/DMXEngine/Source/DMXEditor/Private/Widgets/FixturePatch/SDMXMVRFixtureList.cpp
@@ -390,6 +390,14 @@ void SDMXMVRFixtureList::Construct(const FArguments& InArgs, TWeakPtr<FDMXEditor
 
 	RegisterCommands();
 	RefreshList();
+
+	AdoptSelectionFromFixturePatchSharedData();
+
+	// Make an initial selection if nothing was selected from Fixture Patch Shared Data, as if the user clicked it
+	if (ListView->GetSelectedItems().IsEmpty() && !ListSource.IsEmpty())
+	{
+		ListView->SetSelection(ListSource[0], ESelectInfo::OnMouseClick);
+	}
 }
 
 void SDMXMVRFixtureList::RequestListRefresh()
@@ -697,15 +705,6 @@ void SDMXMVRFixtureList::AdoptSelectionFromFixturePatchSharedData()
 		ListView->SetItemSelection(NewSelection, bSelected, ESelectInfo::OnMouseClick);
 		ListView->RequestScrollIntoView(NewSelection[0]);
 	}
-	else if(TSharedPtr<FDMXEditor> DMXEditor = WeakDMXEditor.Pin())
-	{
-		const UDMXLibrary* DMXLibrary = DMXEditor->GetDMXLibrary();
-		const TArray<UDMXEntityFixturePatch*> FixturePatches = DMXLibrary->GetEntitiesTypeCast<UDMXEntityFixturePatch>();
-		if (FixturePatches.Num() > 0)
-		{
-			FixturePatchSharedData->SelectFixturePatch(FixturePatches[0]);
-		}
-	}
 }
 
 void SDMXMVRFixtureList::AutoAssignFixturePatches()
