DMX - Fix ShowOnlyInnerProperties doesn't work for FDMXFixturePatchRef

diff --git a/Engine/Plugins/VirtualProduction/DMX/DMXEngine/Source/DMXEditor/Private/Customizations/DMXEntityReferenceCustomization.cpp b/Engine/Plugins/VirtualProduction/DMX/DMXEngine/Source/DMXEditor/Private/Customizations/DMXEntityReferenceCustomization.cpp
--- a/Engine/Plugins/VirtualProduction/DMX/DMXEngine/Source/DMXEditor/Private/Customizations/DMXEntityReferenceCustomization.cpp
+++ b/Engine/Plugins/VirtualProduction/DMX/DMXEngine/Source/DMXEditor/Private/Customizations/DMXEntityReferenceCustomization.cpp
@@ -24,22 +24,27 @@ void FDMXEntityReferenceCustomization::CustomizeHeader(TSharedRef<IPropertyHandl
 {
 	StructHandle = InPropertyHandle;
 
-	// If we should display any children, leave the standard header. Otherwise, create custom picker
-	TSharedRef<SWidget> ValueContent = !ShouldDisplayLibrary()
-		? CreateEntityPickerWidget(InPropertyHandle)
-		: InPropertyHandle->CreatePropertyValueWidget(false);
-
-	InHeaderRow
-		.NameContent()
-		[
-			InPropertyHandle->CreatePropertyNameWidget()
-		]
-		.ValueContent()
-		.MinDesiredWidth(200.0f)
-		.MaxDesiredWidth(400.0f)
-		[
-			ValueContent
-		];
+	constexpr TCHAR ShowOnlyInnerPropertiesMetaDataName[] = TEXT("ShowOnlyInnerProperties");
+	bool bShowHeader = !InPropertyHandle->HasMetaData(ShowOnlyInnerPropertiesMetaDataName);
+	if (bShowHeader)
+	{
+		// If we should display any children, leave the standard header. Otherwise, create custom picker
+		TSharedRef<SWidget> ValueContent = !ShouldDisplayLibrary()
+			? CreateEntityPickerWidget(InPropertyHandle)
+			: InPropertyHandle->CreatePropertyValueWidget(false);
+
+		InHeaderRow
+			.NameContent()
+			[
+				InPropertyHandle->CreatePropertyNameWidget()
+			]
+			.ValueContent()
+			.MinDesiredWidth(200.0f)
+			.MaxDesiredWidth(400.0f)
+			[
+				ValueContent
+			];
+	}
 }
 
 void FDMXEntityReferenceCustomization::CustomizeChildren(TSharedRef<IPropertyHandle> InPropertyHandle, IDetailChildrenBuilder& InChildBuilder, IPropertyTypeCustomizationUtils& CustomizationUtils)
