DMX - Fix Drag Drop reimport MVR clears source data

diff --git a/Engine/Plugins/VirtualProduction/DMX/DMXEngine/Source/DMXEditor/Private/Factories/DMXLibraryFromMVRFactory.cpp b/Engine/Plugins/VirtualProduction/DMX/DMXEngine/Source/DMXEditor/Private/Factories/DMXLibraryFromMVRFactory.cpp
--- a/Engine/Plugins/VirtualProduction/DMX/DMXEngine/Source/DMXEditor/Private/Factories/DMXLibraryFromMVRFactory.cpp
+++ b/Engine/Plugins/VirtualProduction/DMX/DMXEngine/Source/DMXEditor/Private/Factories/DMXLibraryFromMVRFactory.cpp
@@ -162,6 +162,9 @@ EReimportResult::Type UDMXLibraryFromMVRFactory::Reimport(UObject* Obj)
 		return EReimportResult::Failed;
 	}
 
+	// Since a new General Scene Description was created, the source file has to be set anew, 
+	MVRAssetImportData->SetSourceFile(SourceFilename);
+
 	TArray<UDMXImportGDTF*> GDTFs = CreateGDTFAssets(DMXLibrary->GetOuter(), RF_NoFlags, Zip, *GeneralSceneDescription);
 	InitDMXLibrary(DMXLibrary, GDTFs, GeneralSceneDescription);
 
