DMX - Fix odd naming of GDTFs in MVR file can break the import process

diff --git a/Engine/Plugins/VirtualProduction/DMX/DMXEngine/Source/DMXEditor/Private/Factories/DMXLibraryFromMVRFactory.cpp b/Engine/Plugins/VirtualProduction/DMX/DMXEngine/Source/DMXEditor/Private/Factories/DMXLibraryFromMVRFactory.cpp
--- a/Engine/Plugins/VirtualProduction/DMX/DMXEngine/Source/DMXEditor/Private/Factories/DMXLibraryFromMVRFactory.cpp
+++ b/Engine/Plugins/VirtualProduction/DMX/DMXEngine/Source/DMXEditor/Private/Factories/DMXLibraryFromMVRFactory.cpp
@@ -18,6 +18,7 @@
 #include "MVR/DMXMVRGeneralSceneDescription.h"
 #include "MVR/Types/DMXMVRFixtureNode.h"
 
+#include "Algo/Find.h"
 #include "AssetRegistry/AssetRegistryModule.h"
 #include "AssetToolsModule.h"
 #include "ObjectTools.h"
@@ -29,6 +30,37 @@
 
 #define LOCTEXT_NAMESPACE "DMXLibraryFromMVRFactory"
 
+namespace UE::DMX::DMXLibraryFromMVRFactory::Private
+{
+	/** Finds the file corresponding to the GDTF Spec in an MVR Zip file. Returns true if the file was found.  */
+	bool FindGDTFFilenameInMVRZip(const TSharedRef<FDMXZipper>& MVRZip, const FString& GDTFSpec, FString& OutGDTFFilename)
+	{
+		OutGDTFFilename.Reset();
+
+		// Some softwares store the GDTF spec with, some without the .gdtf extension. 
+		TArray<FString> Filenames = MVRZip->GetFiles();
+		const FString* FilenamePtr = Algo::FindByPredicate(Filenames, [GDTFSpec](const FString& Filename)
+			{
+				if (GDTFSpec == Filename)
+				{
+					return true;
+				}
+				else if (GDTFSpec + TEXT(".gdtf") == Filename)
+				{
+					return true;
+				}
+				return false;
+			});
+
+		if (FilenamePtr)
+		{
+			OutGDTFFilename = *FilenamePtr;
+			return true;
+		}
+		return false;
+	}
+}
+
 const FString UDMXLibraryFromMVRFactory::MVRFileExtension = TEXT("MVR");
 
 UDMXLibraryFromMVRFactory::UDMXLibraryFromMVRFactory()
@@ -261,7 +293,18 @@ TArray<UDMXImportGDTF*> UDMXLibraryFromMVRFactory::CreateGDTFAssets(UObject* Par
 		// Don't import the same GDTF twice
 		if (!ImportedGDTFNames.Contains(FixtureNode->GDTFSpec))
 		{
-			const FDMXZipper::FDMXScopedUnzipToTempFile ScopedUnzipGDTF(Zip, FixtureNode->GDTFSpec);
+			using namespace UE::DMX::DMXLibraryFromMVRFactory::Private;
+
+			FString GDTFFilename;
+			if (!FindGDTFFilenameInMVRZip(Zip, FixtureNode->GDTFSpec, GDTFFilename))
+			{
+				UE_LOG(LogDMXEditor, Warning, TEXT("Cannot find a GDTF file that corresponds to GDTF Spec '%s' of Fixture Node '%s'. Ignoring Fixture."), *FixtureNode->GDTFSpec, *FixtureNode->Name)
+				{
+					continue;
+				}
+			}
+
+			const FDMXZipper::FDMXScopedUnzipToTempFile ScopedUnzipGDTF(Zip, GDTFFilename);
 			if (!ScopedUnzipGDTF.TempFilePathAndName.IsEmpty())
 			{
 				constexpr bool bRemovePathFromDesiredName = true;
@@ -369,12 +412,18 @@ void UDMXLibraryFromMVRFactory::InitDMXLibrary(UDMXLibrary* DMXLibrary, const TA
 	// Update or create new Fixture Patches for the MVR Fixtures
 	for (const UDMXMVRFixtureNode* FixtureNode : FixtureNodes)
 	{
-		if (!GDTFSpecToFixtureTypeMap.Contains(FixtureNode->GDTFSpec))
+		// Find the GDTF Filename, irregardless if it has the .gdtf extension
+		FString GDTFFilename = FixtureNode->GDTFSpec;
+		if (!GDTFSpecToFixtureTypeMap.Contains(GDTFFilename))
+		{
+			GDTFFilename = FixtureNode->GDTFSpec + TEXT(".gdtf");
+		}
+		if (!GDTFSpecToFixtureTypeMap.Contains(GDTFFilename))
 		{
 			continue;
 		}
 
-		UDMXEntityFixtureType* FixtureType = GDTFSpecToFixtureTypeMap[FixtureNode->GDTFSpec];
+		UDMXEntityFixtureType* FixtureType = GDTFSpecToFixtureTypeMap[GDTFFilename];
 		if (!FixtureTypeToColorMap.Contains(FixtureType))
 		{
 			FLinearColor FixtureTypeColor = FLinearColor::MakeRandomColor();
