DMX - Add back mistakenly removed blueprint getter and setter functions for Starting Channel property of UDMXEntityFixturePatch

diff --git a/Engine/Plugins/VirtualProduction/DMX/DMXEngine/Source/DMXRuntime/Public/Library/DMXEntityFixturePatch.h b/Engine/Plugins/VirtualProduction/DMX/DMXEngine/Source/DMXRuntime/Public/Library/DMXEntityFixturePatch.h
--- a/Engine/Plugins/VirtualProduction/DMX/DMXEngine/Source/DMXRuntime/Public/Library/DMXEntityFixturePatch.h
+++ b/Engine/Plugins/VirtualProduction/DMX/DMXEngine/Source/DMXRuntime/Public/Library/DMXEntityFixturePatch.h
@@ -184,9 +184,11 @@ public:
 	 * 
 	 * If Auto Assign Address was set to true, turns off Auto Assign Address.
 	 */
+	UFUNCTION(BlueprintCallable, Category = "DMX|Fixture Patch")
 	void SetStartingChannel(int32 NewStartingChannel);
 
 	/** Return the starting channel */
+	UFUNCTION(BlueprintCallable, Category = "DMX|Fixture Patch")
 	FORCEINLINE int32 GetStartingChannel() const { return StartingChannel; }
 
 #if WITH_EDITOR
