DMX - Fix a crash that can occur when changing properties in the Mode view

diff --git a/Engine/Plugins/VirtualProduction/DMX/DMXEngine/Source/DMXEditor/Private/Widgets/FixtureType/SDMXFixtureModeEditor.cpp b/Engine/Plugins/VirtualProduction/DMX/DMXEngine/Source/DMXEditor/Private/Widgets/FixtureType/SDMXFixtureModeEditor.cpp
--- a/Engine/Plugins/VirtualProduction/DMX/DMXEngine/Source/DMXEditor/Private/Widgets/FixtureType/SDMXFixtureModeEditor.cpp
+++ b/Engine/Plugins/VirtualProduction/DMX/DMXEngine/Source/DMXEditor/Private/Widgets/FixtureType/SDMXFixtureModeEditor.cpp
@@ -22,7 +22,6 @@ void SDMXFixtureModeEditor::Construct(const FArguments& InArgs, const TSharedRef
 	WeakDMXEditor = InDMXEditor;
 	FixtureTypeSharedData = InDMXEditor->GetFixtureTypeSharedData();
 
-	UDMXEntityFixtureType::GetOnFixtureTypeChanged().AddSP(this, &SDMXFixtureModeEditor::OnFixtureTypePropertiesChanged);
 	FixtureTypeSharedData->OnFixtureTypesSelected.AddSP(this, &SDMXFixtureModeEditor::Refresh);
 	FixtureTypeSharedData->OnModesSelected.AddSP(this, &SDMXFixtureModeEditor::Refresh);
 
@@ -235,14 +234,6 @@ void SDMXFixtureModeEditor::SetMode(UDMXEntityFixtureType* InFixtureType, int32
 	}
 }
 
-void SDMXFixtureModeEditor::OnFixtureTypePropertiesChanged(const UDMXEntityFixtureType* FixtureType)
-{
-	if (!Transaction.IsValid() && FixtureType == WeakFixtureType.Get())
-	{
-		Refresh();
-	}
-}
-
 bool SDMXFixtureModeEditor::IsPropertyVisible(const FPropertyAndParent& PropertyAndParent) const
 {
 	if (PropertyAndParent.Property.GetFName() == GET_MEMBER_NAME_CHECKED(FDMXFixtureMode, Functions))
diff --git a/Engine/Plugins/VirtualProduction/DMX/DMXEngine/Source/DMXEditor/Private/Widgets/FixtureType/SDMXFixtureModeEditor.h b/Engine/Plugins/VirtualProduction/DMX/DMXEngine/Source/DMXEditor/Private/Widgets/FixtureType/SDMXFixtureModeEditor.h
--- a/Engine/Plugins/VirtualProduction/DMX/DMXEngine/Source/DMXEditor/Private/Widgets/FixtureType/SDMXFixtureModeEditor.h
+++ b/Engine/Plugins/VirtualProduction/DMX/DMXEngine/Source/DMXEditor/Private/Widgets/FixtureType/SDMXFixtureModeEditor.h
@@ -45,9 +45,6 @@ private:
 	/** Sets the Mode that is being edited */
 	void SetMode(UDMXEntityFixtureType* InFixtureType, int32 InModeIndex);
 
-	/** Called when the properties of the fixture type changed */
-	void OnFixtureTypePropertiesChanged(const UDMXEntityFixtureType* FixtureType);
-
 	/** Returns true if the properties are visible */
 	bool IsPropertyVisible(const FPropertyAndParent& PropertyAndParent) const;
 
