DMX - Fix wrong path for MustImplement interface of ActorClass in FDMXMVRSceneGDTFToActorClassPair

diff --git a/Engine/Plugins/VirtualProduction/DMX/DMXEngine/Source/DMXRuntime/Public/MVR/DMXMVRSceneActor.h b/Engine/Plugins/VirtualProduction/DMX/DMXEngine/Source/DMXRuntime/Public/MVR/DMXMVRSceneActor.h
--- a/Engine/Plugins/VirtualProduction/DMX/DMXEngine/Source/DMXRuntime/Public/MVR/DMXMVRSceneActor.h
+++ b/Engine/Plugins/VirtualProduction/DMX/DMXEngine/Source/DMXRuntime/Public/MVR/DMXMVRSceneActor.h
@@ -25,7 +25,7 @@ struct DMXRUNTIME_API FDMXMVRSceneGDTFToActorClassPair
 	TSoftObjectPtr<UDMXImportGDTF> GDTF;
 	
 	/** The Actor Class that should be or was spawned */
-	UPROPERTY(EditAnywhere, BlueprintReadOnly, Category = MVR, NoClear, Meta = (MustImplement = "/Script/DMXRuntime.DMXMVRFixtureActorInterface"))
+	UPROPERTY(EditAnywhere, BlueprintReadOnly, Category = MVR, NoClear, Meta = (MustImplement = "/Script/DMXFixtureActorInterface.DMXMVRFixtureActorInterface"))
 	TSoftClassPtr<AActor> ActorClass;
 };
 
