DMX - Fix Channel span of fixture types are not always updated correctly

diff --git a/Engine/Plugins/VirtualProduction/DMX/DMXEngine/Source/DMXRuntime/Private/Library/DMXEntityFixturePatch.cpp b/Engine/Plugins/VirtualProduction/DMX/DMXEngine/Source/DMXRuntime/Private/Library/DMXEntityFixturePatch.cpp
--- a/Engine/Plugins/VirtualProduction/DMX/DMXEngine/Source/DMXRuntime/Private/Library/DMXEntityFixturePatch.cpp
+++ b/Engine/Plugins/VirtualProduction/DMX/DMXEngine/Source/DMXRuntime/Private/Library/DMXEntityFixturePatch.cpp
@@ -1183,6 +1183,8 @@ void UDMXEntityFixturePatch::OnFixtureTypeChanged(const UDMXEntityFixtureType* F
 	{
 		ValidateActiveMode();
 		RebuildCache();
+
+		OnFixturePatchChangedDelegate.Broadcast(this);
 	}
 }
 
diff --git a/Engine/Plugins/VirtualProduction/DMX/DMXEngine/Source/DMXRuntime/Private/Library/DMXEntityFixtureType.cpp b/Engine/Plugins/VirtualProduction/DMX/DMXEngine/Source/DMXRuntime/Private/Library/DMXEntityFixtureType.cpp
--- a/Engine/Plugins/VirtualProduction/DMX/DMXEngine/Source/DMXRuntime/Private/Library/DMXEntityFixtureType.cpp
+++ b/Engine/Plugins/VirtualProduction/DMX/DMXEngine/Source/DMXRuntime/Private/Library/DMXEntityFixtureType.cpp
@@ -245,6 +245,12 @@ void UDMXEntityFixtureType::PostEditChangeProperty(FPropertyChangedEvent& Proper
 
 	if (PropertyChangedEvent.ChangeType != EPropertyChangeType::Interactive)
 	{
+		// Update the Channel Span for all Modes
+		for (int32 ModeIndex = 0; ModeIndex < Modes.Num(); ModeIndex++)
+		{
+			UpdateChannelSpan(ModeIndex);
+		}
+
 		OnFixtureTypeChangedDelegate.Broadcast(this);
 	}
 }
@@ -271,6 +277,12 @@ void UDMXEntityFixtureType::PostEditChangeChainProperty(FPropertyChangedChainEve
 
 	if (PropertyChangedChainEvent.ChangeType != EPropertyChangeType::Interactive)
 	{
+		// Update the Channel Span for all Modes
+		for (int32 ModeIndex = 0; ModeIndex < Modes.Num(); ModeIndex++)
+		{
+			UpdateChannelSpan(ModeIndex);
+		}
+
 		OnFixtureTypeChangedDelegate.Broadcast(this);
 	}
 }
