DMX - Fix various issues with auto-aligning Fixture Functions and how resulting Patches are displayed in the Fixture Patch editor

diff --git a/Engine/Plugins/VirtualProduction/DMX/DMXEngine/Source/DMXEditor/Private/Widgets/FixturePatch/SDMXMVRFixtureList.cpp b/Engine/Plugins/VirtualProduction/DMX/DMXEngine/Source/DMXEditor/Private/Widgets/FixturePatch/SDMXMVRFixtureList.cpp
--- a/Engine/Plugins/VirtualProduction/DMX/DMXEngine/Source/DMXEditor/Private/Widgets/FixturePatch/SDMXMVRFixtureList.cpp
+++ b/Engine/Plugins/VirtualProduction/DMX/DMXEngine/Source/DMXEditor/Private/Widgets/FixturePatch/SDMXMVRFixtureList.cpp
@@ -171,9 +171,9 @@ private:
 				const FText IssueText = LOCTEXT("NoModesIssue", "Fixture Type has no Modes defined.");
 				ItemToIssueMap.Add(Item, IssueText);
 			}
-			else if (Item->GetFixtureType() && Item->GetModeIndex() == INDEX_NONE)
+			else if (Item->GetFixturePatch()->GetActiveMode() && Item->GetFixturePatch()->GetActiveMode()->Functions.IsEmpty())
 			{
-				const FText IssueText = FText::Format(LOCTEXT("FixtureTypeWithoutModeIssue", "Fixture Type '{0}' does not specify any Modes."), FText::FromString(Item->GetFixtureType()->GetName()));
+				const FText IssueText = LOCTEXT("ActiveModeHasNoFunctionsIssue", "Mode does not define any Functions.");
 				ItemToIssueMap.Add(Item, IssueText);
 			}
 		}
diff --git a/Engine/Plugins/VirtualProduction/DMX/DMXEngine/Source/DMXEditor/Private/Widgets/FixturePatch/SDMXPatchedUniverse.cpp b/Engine/Plugins/VirtualProduction/DMX/DMXEngine/Source/DMXEditor/Private/Widgets/FixturePatch/SDMXPatchedUniverse.cpp
--- a/Engine/Plugins/VirtualProduction/DMX/DMXEngine/Source/DMXEditor/Private/Widgets/FixturePatch/SDMXPatchedUniverse.cpp
+++ b/Engine/Plugins/VirtualProduction/DMX/DMXEngine/Source/DMXEditor/Private/Widgets/FixturePatch/SDMXPatchedUniverse.cpp
@@ -215,9 +215,9 @@ bool SDMXPatchedUniverse::CanAssignFixturePatch(TWeakObjectPtr<UDMXEntityFixture
 	}
 
 	int32 ChannelSpan = FixturePatch->GetChannelSpan();
-	if (ChannelSpan == 0)
+	if (ChannelSpan == 0 || ChannelSpan > DMX_MAX_ADDRESS)
 	{
-		// Cannot patch a patch with 0 channel span
+		// Cannot patch a patch with 0 channel span or exceeding DMX_MAX_ADDRESS
 		return false;
 	}
 
@@ -321,18 +321,24 @@ void SDMXPatchedUniverse::OnFixturePatchChanged(const UDMXEntityFixturePatch* Fi
 		{
 			return Node->GetFixturePatch() == FixturePatch;
 		});
+	if (!NodePtr)
+	{
+		return;
+	}
 
-	if (NodePtr && FixturePatch && UniverseID != FixturePatch->GetUniverseID())
+	const bool bMovedToOtherUniverse = FixturePatch->GetUniverseID() != UniverseID;
+	const bool bExceedsUniverse = FixturePatch->GetChannelSpan() > DMX_MAX_ADDRESS;
+	if (bMovedToOtherUniverse || bExceedsUniverse)
 	{
 		// Patch was moved out of this universe
 		Unpatch(*NodePtr);
 	}
-	else if (!NodePtr && UniverseID == FixturePatch->GetUniverseID())
+	else if (UniverseID == FixturePatch->GetUniverseID())
 	{
 		// Patch was moved to this universe
 		SharedData->SelectUniverse(UniverseID);
 	}
-	else if (NodePtr && (*NodePtr)->NeedsUpdateGrid())
+	else if ((*NodePtr)->NeedsUpdateGrid())
 	{
 		RequestRefresh();
 	}
@@ -362,11 +368,14 @@ void SDMXPatchedUniverse::SetUniverseIDInternal(int32 NewUniverseID)
 		TArray<UDMXEntityFixturePatch*> PatchesInUniverse;
 		Library->ForEachEntityOfType<UDMXEntityFixturePatch>([&](UDMXEntityFixturePatch* Patch)
 			{
-				if (Patch->GetUniverseID() == UniverseID)
+				const bool bResidesInThisUniverse = Patch->GetUniverseID() == UniverseID;
+				const bool bHasValidChannelSpan = Patch->GetChannelSpan() > 0 && Patch->GetChannelSpan() < DMX_MAX_ADDRESS;
+				if (bResidesInThisUniverse && bHasValidChannelSpan)
 				{
 					PatchesInUniverse.Add(Patch);
 				}
 			});
+
 		PatchesInUniverse.Sort([](const UDMXEntityFixturePatch& FixturePatchA, const UDMXEntityFixturePatch& FixturePatchB)
 			{
 				return FixturePatchA.GetEndingChannel() <= FixturePatchB.GetEndingChannel();
diff --git a/Engine/Plugins/VirtualProduction/DMX/DMXEngine/Source/DMXEditor/Private/Widgets/FixtureType/SDMXFixtureModeEditor.cpp b/Engine/Plugins/VirtualProduction/DMX/DMXEngine/Source/DMXEditor/Private/Widgets/FixtureType/SDMXFixtureModeEditor.cpp
--- a/Engine/Plugins/VirtualProduction/DMX/DMXEngine/Source/DMXEditor/Private/Widgets/FixtureType/SDMXFixtureModeEditor.cpp
+++ b/Engine/Plugins/VirtualProduction/DMX/DMXEngine/Source/DMXEditor/Private/Widgets/FixtureType/SDMXFixtureModeEditor.cpp
@@ -96,15 +96,11 @@ void SDMXFixtureModeEditor::NotifyPostChange(const FPropertyChangedEvent& Proper
 		}
 		else if (PropertyName == GET_MEMBER_NAME_CHECKED(FDMXFixtureMatrix, XCells))
 		{
-			FixtureType->UpdateYCellsFromXCells(ModeIndex);
-
 			constexpr bool bSelectMatrix = true;
 			FixtureTypeSharedData->SetFunctionAndMatrixSelection(TArray<int32>(), bSelectMatrix);
 		}
 		else if (PropertyName == GET_MEMBER_NAME_CHECKED(FDMXFixtureMatrix, YCells))
 		{
-			FixtureType->UpdateXCellsFromYCells(ModeIndex);
-
 			constexpr bool bSelectMatrix = true;
 			FixtureTypeSharedData->SetFunctionAndMatrixSelection(TArray<int32>(), bSelectMatrix);
 		}
diff --git a/Engine/Plugins/VirtualProduction/DMX/DMXEngine/Source/DMXRuntime/Private/Library/DMXEntityFixtureType.cpp b/Engine/Plugins/VirtualProduction/DMX/DMXEngine/Source/DMXRuntime/Private/Library/DMXEntityFixtureType.cpp
--- a/Engine/Plugins/VirtualProduction/DMX/DMXEngine/Source/DMXRuntime/Private/Library/DMXEntityFixtureType.cpp
+++ b/Engine/Plugins/VirtualProduction/DMX/DMXEngine/Source/DMXRuntime/Private/Library/DMXEntityFixtureType.cpp
@@ -16,7 +16,7 @@
 
 #define LOCTEXT_NAMESPACE "DMXEntityFixtureType"
 
-uint8 FDMXFixtureFunction::GetLastChannel() const
+int32 FDMXFixtureFunction::GetLastChannel() const
 {
 	return Channel + GetNumChannels() - 1;
 }
@@ -242,13 +242,14 @@ bool UDMXEntityFixtureType::Modify(bool bAlwaysMarkDirty)
 void UDMXEntityFixtureType::PostEditChangeProperty(FPropertyChangedEvent& PropertyChangedEvent)
 {
 	Super::PostEditChangeProperty(PropertyChangedEvent);
-
+	
 	if (PropertyChangedEvent.ChangeType != EPropertyChangeType::Interactive)
 	{
 		// Update the Channel Span for all Modes
 		for (int32 ModeIndex = 0; ModeIndex < Modes.Num(); ModeIndex++)
 		{
 			UpdateChannelSpan(ModeIndex);
+			AlignFunctionChannels(ModeIndex);
 		}
 
 		OnFixtureTypeChangedDelegate.Broadcast(this);
@@ -260,9 +261,8 @@ void UDMXEntityFixtureType::PostEditChangeProperty(FPropertyChangedEvent& Proper
 void UDMXEntityFixtureType::PostEditChangeChainProperty(FPropertyChangedChainEvent& PropertyChangedChainEvent)
 {
 	Super::PostEditChangeChainProperty(PropertyChangedChainEvent);
-	const FName PropertyName = PropertyChangedChainEvent.GetPropertyName();
 
-	// Keep Fixture Patches' ActiveMode value valid
+	const FName PropertyName = PropertyChangedChainEvent.GetPropertyName();
 	if (PropertyName == GET_MEMBER_NAME_CHECKED(UDMXEntityFixtureType, Modes))
 	{
 		int32 ChangedModeIndex = PropertyChangedChainEvent.GetArrayIndex(PropertyName.ToString());
@@ -281,6 +281,7 @@ void UDMXEntityFixtureType::PostEditChangeChainProperty(FPropertyChangedChainEve
 		for (int32 ModeIndex = 0; ModeIndex < Modes.Num(); ModeIndex++)
 		{
 			UpdateChannelSpan(ModeIndex);
+			AlignFunctionChannels(ModeIndex);
 		}
 
 		OnFixtureTypeChangedDelegate.Broadcast(this);
@@ -590,6 +591,8 @@ void UDMXEntityFixtureType::AlignFunctionChannels(int32 ModeIndex)
 			}
 
 			Function.Channel = NextFreeChannel;
+
+			// Don't assign past channel 512
 			NextFreeChannel = Function.GetLastChannel() + 1;
 		}
 	}
@@ -1019,18 +1022,18 @@ void UDMXEntityFixtureType::UpdateYCellsFromXCells(int32 ModeIndex)
 	if (ensureMsgf(Modes.IsValidIndex(ModeIndex), TEXT("Trying to update YCells from XCells, but Mode Index is not valid.")))
 	{
 		FDMXFixtureMode& Mode = Modes[ModeIndex];
-		FDMXFixtureMatrix Matrix = Mode.FixtureMatrixConfig;
+		FDMXFixtureMatrix& Matrix = Mode.FixtureMatrixConfig;
 		if (ensureMsgf(Mode.bFixtureMatrixEnabled, TEXT("Trying to update YCells from XCells, but Fixture Matrix is not enabled.")))
 		{
-			// Disable the matrix while editing it so other functions align when enabling it again
-			SetFixtureMatrixEnabled(ModeIndex, false);
-
-			constexpr int32 MaxNumCells = 512;
-			Matrix.XCells = FMath::Clamp(Matrix.XCells, 1, MaxNumCells);
-			Matrix.YCells = FMath::Clamp(Matrix.YCells, 1, MaxNumCells - Matrix.XCells + 1);
+			int32 NumChannelsOfCell = 0;
+			for (const FDMXFixtureCellAttribute& CellAttribute : Matrix.CellAttributes)
+			{
+				NumChannelsOfCell += CellAttribute.GetNumChannels();
+			}
+			Matrix.YCells = DMX_MAX_ADDRESS / (Matrix.XCells * NumChannelsOfCell);
 
-			SetFixtureMatrixEnabled(ModeIndex, true);
 			UpdateChannelSpan(ModeIndex);
+			AlignFunctionChannels(ModeIndex);
 		}
 	}
 }
@@ -1040,18 +1043,22 @@ void UDMXEntityFixtureType::UpdateXCellsFromYCells(int32 ModeIndex)
 	if (ensureMsgf(Modes.IsValidIndex(ModeIndex), TEXT("Trying to update XCells from YCells, but Mode Index is not valid.")))
 	{
 		FDMXFixtureMode& Mode = Modes[ModeIndex];
-		FDMXFixtureMatrix Matrix = Mode.FixtureMatrixConfig;
+		FDMXFixtureMatrix& Matrix = Mode.FixtureMatrixConfig;
 		if (ensureMsgf(Mode.bFixtureMatrixEnabled, TEXT("Trying to update XCells from YCells, but Fixture Matrix is not enabled.")))
 		{
-			// Disable the matrix while editing it so other functions align when enabling it again
-			SetFixtureMatrixEnabled(ModeIndex, false);
+			if (Matrix.GetNumChannels() > DMX_MAX_ADDRESS)
+			{
+				int32 NumChannelsOfCell = 0;
+				for (const FDMXFixtureCellAttribute& CellAttribute : Matrix.CellAttributes)
+				{
+					NumChannelsOfCell += CellAttribute.GetNumChannels();
+				}
 
-			constexpr int32 MaxNumCells = 512;
-			Matrix.YCells = FMath::Clamp(Matrix.YCells, 1, MaxNumCells);
-			Matrix.XCells = FMath::Clamp(Matrix.XCells, 1, MaxNumCells - Matrix.YCells + 1);
+				Matrix.XCells = DMX_MAX_ADDRESS / (Matrix.YCells * NumChannelsOfCell);
 
-			SetFixtureMatrixEnabled(ModeIndex, true);
-			UpdateChannelSpan(ModeIndex);
+				UpdateChannelSpan(ModeIndex);
+				AlignFunctionChannels(ModeIndex);
+			}
 		}
 	}
 }
diff --git a/Engine/Plugins/VirtualProduction/DMX/DMXEngine/Source/DMXRuntime/Public/Library/DMXEntityFixtureType.h b/Engine/Plugins/VirtualProduction/DMX/DMXEngine/Source/DMXRuntime/Public/Library/DMXEntityFixtureType.h
--- a/Engine/Plugins/VirtualProduction/DMX/DMXEngine/Source/DMXRuntime/Public/Library/DMXEntityFixtureType.h
+++ b/Engine/Plugins/VirtualProduction/DMX/DMXEngine/Source/DMXRuntime/Public/Library/DMXEntityFixtureType.h
@@ -59,7 +59,7 @@ struct DMXRUNTIME_API FDMXFixtureFunction
 	FORCEINLINE uint8 GetNumChannels() const { return static_cast<uint8>(DataType) + 1; }
 
 	/** Returns the last channel of the Function */
-	uint8 GetLastChannel() const;
+	int32 GetLastChannel() const;
 
 	/**
 	 * The Attribute name to map this Function to.
@@ -190,10 +190,10 @@ struct DMXRUNTIME_API FDMXFixtureMatrix
 	UPROPERTY(VisibleAnywhere, BlueprintReadWrite, meta = (DisplayPriority = "20", DisplayName = "First Cell Channel", ClampMin = "1", ClampMax = "512"), Category = "Mode Settings")
 	int32 FirstCellChannel = 1;
 
-	UPROPERTY(EditAnywhere, BlueprintReadWrite, meta = (DisplayPriority = "30", DisplayName = "X Cells", ClampMin = "1"), Category = "Mode Settings")
+	UPROPERTY(EditAnywhere, BlueprintReadWrite, meta = (DisplayPriority = "30", DisplayName = "X Cells", ClampMin = "1", ClampMax = "512", UIMin = "1", UIMax = "512"), Category = "Mode Settings")
 	int32 XCells = 1;
 
-	UPROPERTY(EditAnywhere, BlueprintReadWrite, meta = (DisplayPriority = "40", DisplayName = "Y Cells", ClampMin = "1"), Category = "Mode Settings")
+	UPROPERTY(EditAnywhere, BlueprintReadWrite, meta = (DisplayPriority = "40", DisplayName = "Y Cells", ClampMin = "1", ClampMax = "512", UIMin = "1", UIMax = "512"), Category = "Mode Settings")
 	int32 YCells = 1;
 
 	UPROPERTY(EditAnywhere, BlueprintReadWrite, meta = (DisplayPriority = "50", DisplayName = "PixelMapping Distribution"), Category = "Mode Settings")
@@ -550,13 +550,13 @@ public:
 	UE_DEPRECATED(4.27, "Use MakeValid instead.")
 	void UpdateModeChannelProperties(FDMXFixtureMode& Mode);
 
-	UE_DEPRECATED(5.0, "Deprecated in favor of UDMXEntityFixtureType::UpdateChannelSpan.")
+	UE_DEPRECATED(5.0, "Deprecated in favor of UDMXEntityFixtureType::UpdateChannelSpan(int32 ModeIndex).")
 	void UpdateChannelSpan(FDMXFixtureMode& Mode);
 
-	UE_DEPRECATED(5.0, "Deprecated  in favor of UDMXEntityFixtureType::UpdateYCellsFromXCells.")
+	UE_DEPRECATED(5.0, "Deprecated  in favor of UDMXEntityFixtureType::UpdateYCellsFromXCells(int32 ModeIndex).")
 	void UpdateYCellsFromXCells(FDMXFixtureMode& Mode);
 
-	UE_DEPRECATED(5.0, "Deprecated to in favor of UDMXEntityFixtureType::UpdateXCellsFromYCells.")
+	UE_DEPRECATED(5.0, "Deprecated to in favor of UDMXEntityFixtureType::UpdateXCellsFromYCells(int32 ModeIndex).")
 	void UpdateXCellsFromYCells(FDMXFixtureMode& Mode);
 #endif // WITH_EDITOR
 
