DMX - Fix potential access violation in FDMXEntityReferenceCustomization when modulators in pixelmapping were changed

diff --git a/Engine/Plugins/VirtualProduction/DMX/DMXEngine/Source/DMXEditor/Private/Customizations/DMXEntityReferenceCustomization.cpp b/Engine/Plugins/VirtualProduction/DMX/DMXEngine/Source/DMXEditor/Private/Customizations/DMXEntityReferenceCustomization.cpp
--- a/Engine/Plugins/VirtualProduction/DMX/DMXEngine/Source/DMXEditor/Private/Customizations/DMXEntityReferenceCustomization.cpp
+++ b/Engine/Plugins/VirtualProduction/DMX/DMXEngine/Source/DMXEditor/Private/Customizations/DMXEntityReferenceCustomization.cpp
@@ -155,7 +155,7 @@ bool FDMXEntityReferenceCustomization::GetEntityIsMultipleValues() const
 {
 	TArray<void*> RawData;
 	StructHandle->AccessRawData(RawData);
-	if (RawData[0] == nullptr)
+	if (RawData.IsEmpty() || RawData[0] == nullptr)
 	{
 		return true;
 	}
