DMX - Fix right side Fixture Patch view doesn't update on Mode changes and auto-assign

diff --git a/Engine/Plugins/VirtualProduction/DMX/DMXEngine/Source/DMXEditor/Private/Widgets/FixturePatch/DMXFixturePatchNode.cpp b/Engine/Plugins/VirtualProduction/DMX/DMXEngine/Source/DMXEditor/Private/Widgets/FixturePatch/DMXFixturePatchNode.cpp
--- a/Engine/Plugins/VirtualProduction/DMX/DMXEngine/Source/DMXEditor/Private/Widgets/FixturePatch/DMXFixturePatchNode.cpp
+++ b/Engine/Plugins/VirtualProduction/DMX/DMXEngine/Source/DMXEditor/Private/Widgets/FixturePatch/DMXFixturePatchNode.cpp
@@ -57,12 +57,16 @@ public:
 
 			do
 			{
-				TSharedPtr<FDMXFixturePatchFragment> NewFragment = MakeShared<FDMXFixturePatchFragment>();
-
 				const int32 Row = ChannelIndex / FDMXChannelGridSpecs::NumColumns;
 				const int32 Column = ChannelIndex % FDMXChannelGridSpecs::NumColumns;
 				int32 ColumnSpan = Column + ChannelSpan < FDMXChannelGridSpecs::NumColumns ? ChannelSpan : FDMXChannelGridSpecs::NumColumns - Column;
+				if (ColumnSpan == 0)
+				{
+					// No empty fragments
+					break;
+				}
 
+				const TSharedPtr<FDMXFixturePatchFragment> NewFragment = MakeShared<FDMXFixturePatchFragment>();
 				NewFragment->FixturePatch = InFixturePatch;
 				NewFragment->ChannelIndex = ChannelIndex;
 				NewFragment->Column = Column;
@@ -335,14 +339,14 @@ TSharedPtr<FDMXFixturePatchNode> FDMXFixturePatchNode::Create(TWeakPtr<FDMXEdito
 	return NewNode;
 }
 
-void FDMXFixturePatchNode::SetAddresses(const TSharedRef<SDMXPatchedUniverse>& OwningUniverseWidget, int32 NewStartingChannel, int32 NewChannelSpan, bool bTransacted)
+void FDMXFixturePatchNode::SetAddresses(int32 NewUniverseID, int32 NewStartingChannel, int32 NewChannelSpan, bool bTransacted)
 {
 	if (!FixturePatch.IsValid())
 	{
 		return;
 	}
 
-	UniverseWidget = OwningUniverseWidget;
+	UniverseID = NewUniverseID;
 	StartingChannel = NewStartingChannel;
 	ChannelSpan = NewChannelSpan;
 
@@ -353,20 +357,20 @@ void FDMXFixturePatchNode::SetAddresses(const TSharedRef<SDMXPatchedUniverse>& O
 		FixturePatch->SetUniverseID(LastTransactedUniverseID);
 		FixturePatch->SetStartingChannel(LastTransactedChannelID);
 
-		LastTransactedUniverseID = UniverseWidget->GetUniverseID();
+		LastTransactedUniverseID = UniverseID;
 		LastTransactedChannelID = StartingChannel;
 
 		const FScopedTransaction Transaction = FScopedTransaction(LOCTEXT("FixturePatchAssigned", "Patched Fixture"));
 		FixturePatch->PreEditChange(nullptr);
 
-		FixturePatch->SetUniverseID(UniverseWidget->GetUniverseID());
+		FixturePatch->SetUniverseID(NewUniverseID);
 		FixturePatch->SetStartingChannel(StartingChannel);
 
 		FixturePatch->PostEditChange();
 	}
 	else
 	{
-		FixturePatch->SetUniverseID(UniverseWidget->GetUniverseID());
+		FixturePatch->SetUniverseID(NewUniverseID);
 		FixturePatch->SetStartingChannel(StartingChannel);
 	}
 }
@@ -455,6 +459,7 @@ int32 FDMXFixturePatchNode::GetChannelSpan() const
 bool FDMXFixturePatchNode::NeedsUpdateGrid() const
 {
 	return 
+		UniverseID != GetUniverseID() ||
 		StartingChannel != GetStartingChannel() ||
 		ChannelSpan != GetChannelSpan();
 }
diff --git a/Engine/Plugins/VirtualProduction/DMX/DMXEngine/Source/DMXEditor/Private/Widgets/FixturePatch/DMXFixturePatchNode.h b/Engine/Plugins/VirtualProduction/DMX/DMXEngine/Source/DMXEditor/Private/Widgets/FixturePatch/DMXFixturePatchNode.h
--- a/Engine/Plugins/VirtualProduction/DMX/DMXEngine/Source/DMXEditor/Private/Widgets/FixturePatch/DMXFixturePatchNode.h
+++ b/Engine/Plugins/VirtualProduction/DMX/DMXEngine/Source/DMXEditor/Private/Widgets/FixturePatch/DMXFixturePatchNode.h
@@ -25,7 +25,7 @@ public:
 	static TSharedPtr<FDMXFixturePatchNode> Create(TWeakPtr<FDMXEditor> InDMXEditor, const TWeakObjectPtr<UDMXEntityFixturePatch>& InFixturePatch);
 
 	/** Sets the Addresses of the Node */
-	void SetAddresses(const TSharedRef<SDMXPatchedUniverse>& OwningUniverseWidget, int32 NewStartingChannel, int32 NewChannelSpan, bool bTransacted);
+	void SetAddresses(int32 UniverseID, int32 NewStartingChannel, int32 NewChannelSpan, bool bTransacted);
 
 	/** Generates widgets. Requires Addresses to be set via SetAddresses */
 	TArray<TSharedPtr<SDMXFixturePatchFragment>> GenerateWidgets(const TArray<TSharedPtr<FDMXFixturePatchNode>>& FixturePatchNodeGroup);
@@ -73,6 +73,9 @@ private:
 	/** Universe the patch is assigned to */
 	TSharedPtr<SDMXPatchedUniverse> UniverseWidget;
 
+	/** Cached Universe of the patch */
+	int32 UniverseID = 0;
+
 	/** Cached starting channel of the patch */	
 	int32 StartingChannel = 0;
 
diff --git a/Engine/Plugins/VirtualProduction/DMX/DMXEngine/Source/DMXEditor/Private/Widgets/FixturePatch/SDMXFixturePatcher.cpp b/Engine/Plugins/VirtualProduction/DMX/DMXEngine/Source/DMXEditor/Private/Widgets/FixturePatch/SDMXFixturePatcher.cpp
--- a/Engine/Plugins/VirtualProduction/DMX/DMXEngine/Source/DMXEditor/Private/Widgets/FixturePatch/SDMXFixturePatcher.cpp
+++ b/Engine/Plugins/VirtualProduction/DMX/DMXEngine/Source/DMXEditor/Private/Widgets/FixturePatch/SDMXFixturePatcher.cpp
@@ -248,6 +248,7 @@ void SDMXFixturePatcher::OnDragEnterChannel(int32 UniverseID, int32 ChannelID, c
 					return ClampStartingChannel(StartingChannel, ChannelSpan);
 				}();
 
+
 				// Patch the node but do not transact it (transact on drop instead)
 				const TSharedPtr<SDMXPatchedUniverse>& Universe = PatchedUniversesByID.FindChecked(UniverseID);
 
@@ -480,23 +481,12 @@ void SDMXFixturePatcher::OnUniverseSelectionChanged()
 {
 	if (IsUniverseSelectionEnabled())
 	{
-		GEditor->GetTimerManager()->SetTimerForNextTick(FTimerDelegate::CreateLambda([this]()
-			{
-				ShowSelectedUniverse();
-			}));
+		GEditor->GetTimerManager()->SetTimerForNextTick(FTimerDelegate::CreateSP(this, &SDMXFixturePatcher::ShowSelectedUniverse));
 	}
 	else 
 	{
-		GEditor->GetTimerManager()->SetTimerForNextTick(FTimerDelegate::CreateLambda([this]()
-			{
-				ShowAllPatchedUniverses();
-
-				const int32 SelectedUniverse = SharedData->GetSelectedUniverse();
-				if (PatchedUniversesByID.Contains(SelectedUniverse))
-				{
-					PatchedUniverseScrollBox->ScrollDescendantIntoView(PatchedUniversesByID[SelectedUniverse]);
-				}
-			}));
+		constexpr bool bReconstructWidgets = false;
+		GEditor->GetTimerManager()->SetTimerForNextTick(FTimerDelegate::CreateSP(this, &SDMXFixturePatcher::ShowAllPatchedUniverses, bReconstructWidgets));
 	}
 }
 
@@ -598,6 +588,12 @@ void SDMXFixturePatcher::ShowAllPatchedUniverses(bool bForceReconstructWidget)
 	
 		int32 FirstEmptyUniverse = LastPatchedUniverseID + 1;
 		AddUniverse(FirstEmptyUniverse);
+
+		const int32 SelectedUniverse = SharedData->GetSelectedUniverse();
+		if (PatchedUniversesByID.Contains(SelectedUniverse))
+		{
+			PatchedUniverseScrollBox->ScrollDescendantIntoView(PatchedUniversesByID[SelectedUniverse]);
+		}
 	}
 }
 
diff --git a/Engine/Plugins/VirtualProduction/DMX/DMXEngine/Source/DMXEditor/Private/Widgets/FixturePatch/SDMXMVRFixtureList.cpp b/Engine/Plugins/VirtualProduction/DMX/DMXEngine/Source/DMXEditor/Private/Widgets/FixturePatch/SDMXMVRFixtureList.cpp
--- a/Engine/Plugins/VirtualProduction/DMX/DMXEngine/Source/DMXEditor/Private/Widgets/FixturePatch/SDMXMVRFixtureList.cpp
+++ b/Engine/Plugins/VirtualProduction/DMX/DMXEngine/Source/DMXEditor/Private/Widgets/FixturePatch/SDMXMVRFixtureList.cpp
@@ -407,13 +407,7 @@ void SDMXMVRFixtureList::RequestListRefresh()
 		return;
 	}
 
-	RequestListRefreshTimerHandle =
-		GEditor->GetTimerManager()->SetTimerForNextTick(
-			FTimerDelegate::CreateLambda([this]()
-				{
-					RefreshList();
-				})
-		);
+	RequestListRefreshTimerHandle = GEditor->GetTimerManager()->SetTimerForNextTick(FTimerDelegate::CreateSP(this, &SDMXMVRFixtureList::RefreshList));
 }
 
 void SDMXMVRFixtureList::EnterFixturePatchNameEditingMode()
@@ -642,10 +636,7 @@ void SDMXMVRFixtureList::OnFixturePatchSharedDataSelectedFixturePatches()
 		}
 	}
 
-	GEditor->GetTimerManager()->SetTimerForNextTick(FTimerDelegate::CreateLambda([this]()
-		{
-			FSlateApplication::Get().SetKeyboardFocus(AsShared());
-		}));
+	GEditor->GetTimerManager()->SetTimerForNextTick(FTimerDelegate::CreateSP(this, &SDMXMVRFixtureList::SetKeyboardFocus));
 }
 
 void SDMXMVRFixtureList::OnFixturePatchSharedDataSelectedUniverse()
@@ -800,6 +791,12 @@ TSharedRef<SHeaderRow> SDMXMVRFixtureList::GenerateHeaderRow()
 	return HeaderRow.ToSharedRef();
 }
 
+
+void SDMXMVRFixtureList::SetKeyboardFocus()
+{
+	FSlateApplication::Get().SetKeyboardFocus(AsShared());
+}
+
 void SDMXMVRFixtureList::SaveHeaderRowSettings()
 {
 	UDMXEditorSettings* EditorSettings = GetMutableDefault<UDMXEditorSettings>();
diff --git a/Engine/Plugins/VirtualProduction/DMX/DMXEngine/Source/DMXEditor/Private/Widgets/FixturePatch/SDMXMVRFixtureList.h b/Engine/Plugins/VirtualProduction/DMX/DMXEngine/Source/DMXEditor/Private/Widgets/FixturePatch/SDMXMVRFixtureList.h
--- a/Engine/Plugins/VirtualProduction/DMX/DMXEngine/Source/DMXEditor/Private/Widgets/FixturePatch/SDMXMVRFixtureList.h
+++ b/Engine/Plugins/VirtualProduction/DMX/DMXEngine/Source/DMXEditor/Private/Widgets/FixturePatch/SDMXMVRFixtureList.h
@@ -111,6 +111,9 @@ private:
 	/** Auto assigns selected Fixture Patches */
 	void AutoAssignFixturePatches();
 
+	/** Sets keyboard focus on this widget */
+	void SetKeyboardFocus();
+
 	/** Generates the Header Row of the List */
 	TSharedRef<SHeaderRow> GenerateHeaderRow();
 
diff --git a/Engine/Plugins/VirtualProduction/DMX/DMXEngine/Source/DMXEditor/Private/Widgets/FixturePatch/SDMXPatchedUniverse.cpp b/Engine/Plugins/VirtualProduction/DMX/DMXEngine/Source/DMXEditor/Private/Widgets/FixturePatch/SDMXPatchedUniverse.cpp
--- a/Engine/Plugins/VirtualProduction/DMX/DMXEngine/Source/DMXEditor/Private/Widgets/FixturePatch/SDMXPatchedUniverse.cpp
+++ b/Engine/Plugins/VirtualProduction/DMX/DMXEngine/Source/DMXEditor/Private/Widgets/FixturePatch/SDMXPatchedUniverse.cpp
@@ -112,24 +112,7 @@ END_SLATE_FUNCTION_BUILD_OPTIMIZATION
 
 void SDMXPatchedUniverse::RequestRefresh()
 {
-	GEditor->GetTimerManager()->SetTimerForNextTick(
-		FTimerDelegate::CreateLambda([this]()
-			{
-				RefreshInternal();
-			})
-	);
-}
-
-void SDMXPatchedUniverse::SetShowUniverseName(bool bShow)
-{
-	if (bShow)
-	{
-		UniverseName->SetVisibility(EVisibility::Visible);
-	}
-	else
-	{
-		UniverseName->SetVisibility(EVisibility::Collapsed);
-	}
+	GEditor->GetTimerManager()->SetTimerForNextTick(FTimerDelegate::CreateSP(this, &SDMXPatchedUniverse::RefreshInternal));
 }
 
 void SDMXPatchedUniverse::CreateChannelConnectors()
@@ -175,19 +158,10 @@ bool SDMXPatchedUniverse::Patch(const TSharedPtr<FDMXFixturePatchNode>& Node, in
 		return false;
 	}
 
-	const int32 NewChannelSpan = FixturePatch->GetChannelSpan();
+	PatchedNodes.AddUnique(Node);
 
-	// Unpatch from the old universe
-	if (const TSharedPtr<SDMXPatchedUniverse>& OldUniverseWidget = Node->GetUniverseWidget())
-	{
-		OldUniverseWidget->Unpatch(Node);
-	}
-
-	// Assign the node to this universe
-	check(!PatchedNodes.Contains(Node));
-	PatchedNodes.Add(Node);
-
-	Node->SetAddresses(SharedThis(this), NewStartingChannel, NewChannelSpan, bCreateTransaction);
+	const int32 NewChannelSpan = FixturePatch->GetChannelSpan();
+	Node->SetAddresses(UniverseID, NewStartingChannel, NewChannelSpan, bCreateTransaction);
 	RequestRefresh();
 
 	return true;
@@ -294,6 +268,28 @@ TSharedPtr<FDMXFixturePatchNode> SDMXPatchedUniverse::FindPatchNodeOfType(UDMXEn
 	return nullptr;
 }
 
+void SDMXPatchedUniverse::Tick(const FGeometry& AllottedGeometry, const double InCurrentTime, const float InDeltaTime) 
+{
+	TArray<TSharedPtr<FDMXFixturePatchNode>> RemovedNodes;
+	for (const TSharedPtr<FDMXFixturePatchNode>& Node : PatchedNodes)
+	{
+		if (Node->GetUniverseID() != UniverseID)
+		{
+			RemovedNodes.Add(Node);
+		}
+		else if (Node->NeedsUpdateGrid())
+		{
+			RequestRefresh();
+			break;
+		}
+	}
+
+	for (const TSharedPtr<FDMXFixturePatchNode>& RemovedNode : RemovedNodes)
+	{
+		Unpatch(RemovedNode);
+	}
+}
+
 void SDMXPatchedUniverse::Unpatch(const TSharedPtr<FDMXFixturePatchNode>& Node)
 {
 	check(Node.IsValid());
@@ -316,31 +312,24 @@ void SDMXPatchedUniverse::Unpatch(const TSharedPtr<FDMXFixturePatchNode>& Node)
 
 void SDMXPatchedUniverse::OnFixturePatchChanged(const UDMXEntityFixturePatch* FixturePatch)
 {
-	// Keep track of Universe Changes
-	const TSharedPtr<FDMXFixturePatchNode>* NodePtr = PatchedNodes.FindByPredicate([FixturePatch](const TSharedPtr<FDMXFixturePatchNode>& Node)
-		{
-			return Node->GetFixturePatch() == FixturePatch;
-		});
-	if (!NodePtr)
+	if (!FixturePatch)
 	{
 		return;
 	}
 
-	const bool bMovedToOtherUniverse = FixturePatch->GetUniverseID() != UniverseID;
-	const bool bExceedsUniverse = FixturePatch->GetChannelSpan() > DMX_MAX_ADDRESS;
-	if (bMovedToOtherUniverse || bExceedsUniverse)
-	{
-		// Patch was moved out of this universe
-		Unpatch(*NodePtr);
-	}
-	else if (UniverseID == FixturePatch->GetUniverseID())
-	{
-		// Patch was moved to this universe
-		SharedData->SelectUniverse(UniverseID);
-	}
-	else if ((*NodePtr)->NeedsUpdateGrid())
+	const int32 UniverseIDOfFixturePatch = FixturePatch->GetUniverseID();
+	if (UniverseIDOfFixturePatch == UniverseID)
 	{
-		RequestRefresh();
+		const TSharedPtr<FDMXFixturePatchNode>* NodePtr = PatchedNodes.FindByPredicate([FixturePatch](const TSharedPtr<FDMXFixturePatchNode>& Node)
+			{
+				return Node->GetFixturePatch() == FixturePatch;
+			});
+
+		// Redraw all if the patch cannot be found.
+		if (!NodePtr)
+		{
+			SetUniverseIDInternal(UniverseID);
+		}
 	}
 }
 
@@ -419,6 +408,12 @@ void SDMXPatchedUniverse::RefreshInternal()
 	TMap<int32, TArray<TSharedPtr<FDMXFixturePatchNode>>> AddressToNodeGroupMap;
 	for (const TSharedPtr<FDMXFixturePatchNode>& Node : PatchedNodes)
 	{
+		if (Node->GetChannelSpan() == 0 || Node->GetChannelSpan() > DMX_MAX_ADDRESS)
+		{
+			// Ignore nodes without or with excess channelspan, they're not displayed in UI either
+			continue;
+		}
+
 		const int32 Address = Node->GetStartingChannel();
 		AddressToNodeGroupMap.FindOrAdd(Address).Add(Node);
 	}
@@ -543,7 +538,7 @@ FText SDMXPatchedUniverse::GetHeaderText() const
 	switch (PatchedUniverseReachability)
 	{
 		case EDMXPatchedUniverseReachability::Reachable:
-			return FText::GetEmpty();
+			return FText::Format(LOCTEXT("UniverseIDLabel", "Universe {0}"), UniverseID);
 
 		case EDMXPatchedUniverseReachability::UnreachableForInputPorts:
 			return FText::Format(LOCTEXT("UnreachableForInputPorts", "Universe {0} - Unreachable by Input Ports"), UniverseID);
diff --git a/Engine/Plugins/VirtualProduction/DMX/DMXEngine/Source/DMXEditor/Private/Widgets/FixturePatch/SDMXPatchedUniverse.h b/Engine/Plugins/VirtualProduction/DMX/DMXEngine/Source/DMXEditor/Private/Widgets/FixturePatch/SDMXPatchedUniverse.h
--- a/Engine/Plugins/VirtualProduction/DMX/DMXEngine/Source/DMXEditor/Private/Widgets/FixturePatch/SDMXPatchedUniverse.h
+++ b/Engine/Plugins/VirtualProduction/DMX/DMXEngine/Source/DMXEditor/Private/Widgets/FixturePatch/SDMXPatchedUniverse.h
@@ -109,6 +109,11 @@ public:
 	/** Gets all nodes patched to this universe */
 	const TArray<TSharedPtr<FDMXFixturePatchNode>>& GetPatchedNodes() const { return PatchedNodes; }
 
+protected:
+	//~ Begin SWidget Interface
+	virtual void Tick(const FGeometry& AllottedGeometry, const double InCurrentTime, const float InDeltaTime) override;
+	//~ End SWidget Interface
+
 private:
 	/** Removes the node. Should be called when the node is Patched in another instance */
 	void Unpatch(const TSharedPtr<FDMXFixturePatchNode>& Node);
diff --git a/Engine/Plugins/VirtualProduction/DMX/DMXEngine/Source/DMXRuntime/Private/Library/DMXEntityFixtureType.cpp b/Engine/Plugins/VirtualProduction/DMX/DMXEngine/Source/DMXRuntime/Private/Library/DMXEntityFixtureType.cpp
--- a/Engine/Plugins/VirtualProduction/DMX/DMXEngine/Source/DMXRuntime/Private/Library/DMXEntityFixtureType.cpp
+++ b/Engine/Plugins/VirtualProduction/DMX/DMXEngine/Source/DMXRuntime/Private/Library/DMXEntityFixtureType.cpp
@@ -535,6 +535,10 @@ void UDMXEntityFixtureType::SetFixtureMatrixEnabled(int32 ModeIndex, bool bEnabl
 		{
 			Mode.bFixtureMatrixEnabled = bEnableMatrix;
 
+			// Some old assets may have a 0x0 matrix stored, but we expect it to be always at least one cell.
+			Mode.FixtureMatrixConfig.XCells = FMath::Max(Mode.FixtureMatrixConfig.XCells, 1);
+			Mode.FixtureMatrixConfig.YCells = FMath::Max(Mode.FixtureMatrixConfig.YCells, 1);
+
 			AlignFunctionChannels(ModeIndex);
 		}
 	}
