DMX - Fix Editor crash while adding function to a fixture (pure virtual function being called)

diff --git a/Engine/Plugins/VirtualProduction/DMX/DMXEngine/Source/DMXEditor/Private/Customizations/DMXAttributeNameCustomization.cpp b/Engine/Plugins/VirtualProduction/DMX/DMXEngine/Source/DMXEditor/Private/Customizations/DMXAttributeNameCustomization.cpp
--- a/Engine/Plugins/VirtualProduction/DMX/DMXEngine/Source/DMXEditor/Private/Customizations/DMXAttributeNameCustomization.cpp
+++ b/Engine/Plugins/VirtualProduction/DMX/DMXEngine/Source/DMXEditor/Private/Customizations/DMXAttributeNameCustomization.cpp
@@ -8,6 +8,7 @@
 #include "DetailWidgetRow.h"
 #include "IPropertyTypeCustomization.h"
 #include "IPropertyUtilities.h"
+#include "ScopedTransaction.h"
 #include "Widgets/SNameListPicker.h"
 
 
@@ -21,7 +22,9 @@ void FDMXAttributeNameCustomization::CustomizeHeader(TSharedRef<IPropertyHandle>
 	StructPropertyHandle = InPropertyHandle;
 	check(CastFieldChecked<FStructProperty>(StructPropertyHandle->GetProperty())->Struct == FDMXAttributeName::StaticStruct());
 
-	TSharedPtr<IPropertyUtilities> PropertyUtilities = CustomizationUtils.GetPropertyUtilities();
+	PropertyUtilities = CustomizationUtils.GetPropertyUtilities();
+
+	NameHandle = StructPropertyHandle->GetChildHandle(GET_MEMBER_NAME_CHECKED(FDMXAttributeName, Name));
 
 	UDMXProtocolSettings* ProtocolSettings = GetMutableDefault<UDMXProtocolSettings>();
 	if (!ProtocolSettings)
@@ -57,10 +60,7 @@ void FDMXAttributeNameCustomization::CustomizeHeader(TSharedRef<IPropertyHandle>
 			.OnValueChanged(this, &FDMXAttributeNameCustomization::SetValue)
 		];
 
-	ProtocolSettings->GetOnDefaultAttributesChanged().AddLambda([PropertyUtilities]()
-		{
-			PropertyUtilities->ForceRefresh();
-		});
+	ProtocolSettings->GetOnDefaultAttributesChanged().AddSP(this, &FDMXAttributeNameCustomization::ForceRefresh);
 }
 
 void FDMXAttributeNameCustomization::CustomizeChildren(TSharedRef<IPropertyHandle> InPropertyHandle, IDetailChildrenBuilder& InChildBuilder, IPropertyTypeCustomizationUtils& CustomizationUtils)
@@ -87,24 +87,9 @@ FName FDMXAttributeNameCustomization::GetValue() const
 
 void FDMXAttributeNameCustomization::SetValue(FName NewValue)
 {
-	FStructProperty* StructProperty = CastFieldChecked<FStructProperty>(StructPropertyHandle->GetProperty());
-
-	TArray<void*> RawData;
-	StructPropertyHandle->AccessRawData(RawData);
-
-	for (void* SingleRawData : RawData)
-	{
-		FDMXAttributeName* PreviousValue = reinterpret_cast<FDMXAttributeName*>(SingleRawData);
-		FDMXAttributeName NewAttributeName;
-		NewAttributeName.SetFromName(NewValue);
+	const FScopedTransaction SetAttributeValueTransaction(NSLOCTEXT("DMXAttributeNameCustomization", "SetAttributeValueTransaction", "Set Fixture Function Attribute Name"));
 
-		// Export new value to text format that can be imported later
-		FString TextValue;
-		StructProperty->Struct->ExportText(TextValue, &NewAttributeName, PreviousValue, nullptr, EPropertyPortFlags::PPF_None, nullptr);
-
-		// Set values on edited property handle from exported text
-		ensure(StructPropertyHandle->SetValueFromFormattedString(TextValue, EPropertyValueSetFlags::DefaultFlags) == FPropertyAccess::Result::Success);
-	}
+	NameHandle->SetValue(NewValue);
 }
 
 bool FDMXAttributeNameCustomization::HasMultipleValues() const
@@ -143,3 +128,8 @@ bool FDMXAttributeNameCustomization::HasMultipleValues() const
 
 	return false;
 }
+
+void FDMXAttributeNameCustomization::ForceRefresh()
+{
+	PropertyUtilities->ForceRefresh();
+}
diff --git a/Engine/Plugins/VirtualProduction/DMX/DMXEngine/Source/DMXEditor/Private/Customizations/DMXAttributeNameCustomization.h b/Engine/Plugins/VirtualProduction/DMX/DMXEngine/Source/DMXEditor/Private/Customizations/DMXAttributeNameCustomization.h
--- a/Engine/Plugins/VirtualProduction/DMX/DMXEngine/Source/DMXEditor/Private/Customizations/DMXAttributeNameCustomization.h
+++ b/Engine/Plugins/VirtualProduction/DMX/DMXEngine/Source/DMXEditor/Private/Customizations/DMXAttributeNameCustomization.h
@@ -10,6 +10,8 @@
 #include "IPropertyUtilities.h"
 #include "Widgets/SNameListPicker.h"
 
+class IPropertyUtilities;
+
 
 /**  
  * Customization for FDMXAttributeName struct
@@ -36,7 +38,15 @@ private:
 	/** True if the customization displays multiple values */
 	bool HasMultipleValues() const;
 
-private:
+	/** Forces a refresh of the customization */
+	void ForceRefresh();
+
+	/** Property handle for the Name property */
+	TSharedPtr<IPropertyHandle> NameHandle;
+
 	/** Handle for the FDMXAttributeName struct self */
 	TSharedPtr<IPropertyHandle> StructPropertyHandle;
+
+	/** Property utlities for this customization */
+	TSharedPtr<IPropertyUtilities> PropertyUtilities;
 };
diff --git a/Engine/Plugins/VirtualProduction/DMX/DMXEngine/Source/DMXEditor/Private/Widgets/FixtureType/SDMXFixtureFunctionEditor.cpp b/Engine/Plugins/VirtualProduction/DMX/DMXEngine/Source/DMXEditor/Private/Widgets/FixtureType/SDMXFixtureFunctionEditor.cpp
--- a/Engine/Plugins/VirtualProduction/DMX/DMXEngine/Source/DMXEditor/Private/Widgets/FixtureType/SDMXFixtureFunctionEditor.cpp
+++ b/Engine/Plugins/VirtualProduction/DMX/DMXEngine/Source/DMXEditor/Private/Widgets/FixtureType/SDMXFixtureFunctionEditor.cpp
@@ -24,15 +24,47 @@ void SDMXFixtureFunctionEditor::Construct(const FArguments& InArgs, const TShare
 	FixtureTypeSharedData->OnModesSelected.AddSP(this, &SDMXFixtureFunctionEditor::Refresh);
 	FixtureTypeSharedData->OnFunctionsSelected.AddSP(this, &SDMXFixtureFunctionEditor::Refresh);
 
-	// Note a Details View for the Fixture Type shows significant performance issues and a Struct Type Customization cannot be achieved since the outer struct type cannot be customized.
-	// Hence we just add an edit widget for each node
+	// Create a Struct Details View. This is not the most convenient type to work with as a property type customization for the FDMXFixtureFunction struct cannot be used.
+	// Reason is soley significant performance gains, it's much faster than the easier approach with a UDMXEntityFixtureType customization as it was used up to 4.27.
+	FDetailsViewArgs DetailsViewArgs;
+	DetailsViewArgs.bAllowSearch = true;
+	DetailsViewArgs.bHideSelectionTip = false;
+	DetailsViewArgs.bSearchInitialKeyFocus = true;
+	DetailsViewArgs.bShowOptions = false;
+	DetailsViewArgs.bShowModifiedPropertiesOption = false;
+	DetailsViewArgs.bShowObjectLabel = false;
+	DetailsViewArgs.bForceHiddenPropertyVisibility = false;
+	DetailsViewArgs.bShowScrollBar = false;
+	DetailsViewArgs.NotifyHook = this;
+
+	FStructureDetailsViewArgs StructureDetailsViewArgs;
+
+	FPropertyEditorModule& PropertyModule = FModuleManager::LoadModuleChecked<FPropertyEditorModule>("PropertyEditor");
+	StructDetailsView = PropertyModule.CreateStructureDetailView(DetailsViewArgs, StructureDetailsViewArgs, nullptr);
+	StructDetailsView->GetDetailsView()->SetIsPropertyVisibleDelegate(FIsPropertyVisible::CreateSP(this, &SDMXFixtureFunctionEditor::IsPropertyVisible));
+
+	StructDetailsViewWidget = StructDetailsView->GetWidget();
+
 	ChildSlot
 	[
-		SAssignNew(ContentBorder, SBorder)
-		.HAlign(HAlign_Fill)
-		.VAlign(VAlign_Fill)
-		.BorderImage(FAppStyle::GetBrush("NoBorder"))
-		.Padding(4.f)
+		SNew(SVerticalBox)
+
+		+ SVerticalBox::Slot()
+		.AutoHeight()
+		[
+			StructDetailsViewWidget.ToSharedRef()
+		]
+
+		+ SVerticalBox::Slot()
+		[
+			SNew(SBox)
+			.HAlign(HAlign_Center)
+			.VAlign(VAlign_Center)
+			[
+				SAssignNew(InfoTextBlock, STextBlock)
+				.Visibility(EVisibility::Collapsed)
+			]
+		]
 	];
 
 	Refresh();
@@ -142,16 +174,9 @@ void SDMXFixtureFunctionEditor::Refresh()
 			return LOCTEXT("NoFunctionSelectedWarning", "No Function selected");
 		}();
 
-		ContentBorder->SetContent
-		(
-			SNew(SBox)
-			.HAlign(HAlign_Center)
-			.VAlign(VAlign_Center)
-			[
-				SNew(STextBlock)
-				.Text(ErrorText)
-			]
-		);
+		InfoTextBlock->SetText(ErrorText);
+		InfoTextBlock->SetVisibility(EVisibility::Visible);
+		StructDetailsViewWidget->SetVisibility(EVisibility::Collapsed);
 	}
 }
 
@@ -170,48 +195,19 @@ void SDMXFixtureFunctionEditor::SetFunction(UDMXEntityFixtureType* InFixtureType
 			FDMXFixtureFunction& Function = Mode.Functions[FunctionIndex];
 			const TSharedRef<FStructOnScope> FunctionStructOnScope = MakeShared<FStructOnScope>(FDMXFixtureFunction::StaticStruct(), (uint8*)&Function);
 
-			// Create a Struct Details View. This is not the most convenient type to work with as a property type customization for the FDMXFixtureFunction struct cannot be used.
-			// Reason is soley significant performance gains, it's much faster than the easier approach with a UDMXEntityFixtureType customization as it was used up to 4.27.
-			FDetailsViewArgs DetailsViewArgs;
-			DetailsViewArgs.bAllowSearch = true;
-			DetailsViewArgs.bHideSelectionTip = false;
-			DetailsViewArgs.bSearchInitialKeyFocus = true;
-			DetailsViewArgs.bShowOptions = false;
-			DetailsViewArgs.bShowModifiedPropertiesOption = false;
-			DetailsViewArgs.bShowObjectLabel = false;
-			DetailsViewArgs.bForceHiddenPropertyVisibility = false;
-			DetailsViewArgs.bShowScrollBar = false;
-			DetailsViewArgs.NotifyHook = this;
-
-			FStructureDetailsViewArgs StructureDetailsViewArgs;
-
-			FPropertyEditorModule& PropertyModule = FModuleManager::LoadModuleChecked<FPropertyEditorModule>("PropertyEditor");
-			
-			const TSharedRef<IStructureDetailsView> StructDetailsView = PropertyModule.CreateStructureDetailView(DetailsViewArgs, StructureDetailsViewArgs, FunctionStructOnScope);
-
-			StructDetailsView->GetDetailsView()->SetIsPropertyVisibleDelegate(FIsPropertyVisible::CreateSP(this, &SDMXFixtureFunctionEditor::IsPropertyVisible));
-			StructDetailsView->GetDetailsView()->ForceRefresh();
-			
-			if (TSharedPtr<SWidget> Widget = StructDetailsView->GetWidget())
-			{
-				ContentBorder->SetContent(Widget.ToSharedRef());
-				bSuccess = true;
-			}
+			StructDetailsView->SetStructureData(FunctionStructOnScope);
+
+			StructDetailsViewWidget->SetVisibility(EVisibility::Visible);
+			InfoTextBlock->SetVisibility(EVisibility::Collapsed);
+			bSuccess = true;
 		}
 	}
 	
 	if (!bSuccess)
 	{
-		ContentBorder->SetContent
-		(
-			SNew(SBox)
-			.HAlign(HAlign_Center)
-			.VAlign(VAlign_Center)
-			[
-				SNew(STextBlock)
-				.Text(LOCTEXT("CannotCreateDetailViewForFunctionWarning", "Cannot create Detail View for Function. Fixture Type, Mode or Function no longer exist."))
-			]
-		);
+		InfoTextBlock->SetText(LOCTEXT("CannotCreateDetailViewForFunctionWarning", "Cannot create Detail View for Function. Fixture Type, Mode or Function no longer exist."));
+		InfoTextBlock->SetVisibility(EVisibility::Visible);
+		StructDetailsViewWidget->SetVisibility(EVisibility::Collapsed);
 	}
 }
 
diff --git a/Engine/Plugins/VirtualProduction/DMX/DMXEngine/Source/DMXEditor/Private/Widgets/FixtureType/SDMXFixtureFunctionEditor.h b/Engine/Plugins/VirtualProduction/DMX/DMXEngine/Source/DMXEditor/Private/Widgets/FixtureType/SDMXFixtureFunctionEditor.h
--- a/Engine/Plugins/VirtualProduction/DMX/DMXEngine/Source/DMXEditor/Private/Widgets/FixtureType/SDMXFixtureFunctionEditor.h
+++ b/Engine/Plugins/VirtualProduction/DMX/DMXEngine/Source/DMXEditor/Private/Widgets/FixtureType/SDMXFixtureFunctionEditor.h
@@ -16,7 +16,8 @@ class UDMXEntityFixtureType;
 
 struct FPropertyAndParent;
 class IStructureDetailsView;
-class SBorder;
+class STextBlock;
+class SWidget;
 
 
 /** Editor for a single Function struct inside a UDMXEntityFixtureType. */
@@ -70,8 +71,14 @@ private:
 	/** The current transaction or nullptr if there's no transaction ongoing */
 	TUniquePtr<FScopedTransaction> Transaction;
 	
-	/** Border that holds the content of the editor */
-	TSharedPtr<SBorder> ContentBorder;
+	/** Details view for the Function Struct being edited */
+	TSharedPtr<IStructureDetailsView> StructDetailsView;
+
+	/** Widget of the Details view  */
+	TSharedPtr<SWidget> StructDetailsViewWidget;
+
+	/** Text block that displays info in case the struct details view cannot be shown */
+	TSharedPtr<STextBlock> InfoTextBlock;
 
 	/** Shared Data for */
 	TSharedPtr<FDMXFixtureTypeSharedData> FixtureTypeSharedData;
