DMX - Fix exporting MVR doesn't export transforms of fixtures in the current level

diff --git a/Engine/Plugins/VirtualProduction/DMX/DMXEngine/Source/DMXEditor/Private/Exporters/DMXMVRExporter.cpp b/Engine/Plugins/VirtualProduction/DMX/DMXEngine/Source/DMXEditor/Private/Exporters/DMXMVRExporter.cpp
--- a/Engine/Plugins/VirtualProduction/DMX/DMXEngine/Source/DMXEditor/Private/Exporters/DMXMVRExporter.cpp
+++ b/Engine/Plugins/VirtualProduction/DMX/DMXEngine/Source/DMXEditor/Private/Exporters/DMXMVRExporter.cpp
@@ -13,15 +13,21 @@
 #include "Library/DMXImportGDTF.h"
 #include "Library/DMXLibrary.h"
 #include "MVR/DMXMVRAssetImportData.h"
+#include "MVR/DMXMVRAssetUserData.h"
 #include "MVR/DMXMVRGeneralSceneDescription.h"
+#include "MVR/DMXMVRSceneActor.h"
+#include "MVR/Types/DMXMVRFixtureNode.h"
 
 #include "DesktopPlatformModule.h"
+#include "Editor.h"
+#include "EngineUtils.h"
 #include "IDesktopPlatform.h"
 #include "XmlFile.h"
 #include "Framework/Notifications/NotificationManager.h"
 #include "Misc/FileHelper.h"
 #include "Misc/MessageDialog.h"
 #include "Misc/Paths.h"
+#include "Subsystems/UnrealEditorSubsystem.h"
 #include "Widgets/Notifications/SNotificationList.h"
 
 
@@ -43,16 +49,19 @@ void FDMXMVRExporter::ExportInternal(UDMXLibrary* DMXLibrary, const FString& Fil
 		return;
 	}
 
+	// Create a copy of the library's general scene description, so transforms exported aren't written to the dmx library
 	DMXLibrary->UpdateGeneralSceneDescription();
-	const UDMXMVRGeneralSceneDescription* GeneralSceneDescription = DMXLibrary->GetLazyGeneralSceneDescription();
-	if (!ensureAlwaysMsgf(GeneralSceneDescription, TEXT("Trying to export DMX Library '%s' as MVR file, but its General Scene Description is invalid."), *DMXLibrary->GetName()))
+	UDMXMVRGeneralSceneDescription* TempGeneralSceneDescription = DuplicateObject<UDMXMVRGeneralSceneDescription>(DMXLibrary->GetLazyGeneralSceneDescription(), GetTransientPackage());
+	if (!ensureAlwaysMsgf(TempGeneralSceneDescription, TEXT("Trying to export DMX Library '%s' as MVR file, but its General Scene Description is invalid."), *DMXLibrary->GetName()))
 	{
 		OutErrorReason = FText::Format(LOCTEXT("MVRExportGeneralSceneDescriptionInvalidReason", "DMX Library is invalid. Cannot export {0}."), FText::FromString(FilePathAndName));
 		return;
 	}
 
+	WriteMVRFixtureTransformsFromLevel(TempGeneralSceneDescription);
+
 	const TSharedRef<FDMXZipper> Zip = MakeShared<FDMXZipper>();
-	if (!ZipGeneralSceneDescription(Zip, GeneralSceneDescription, OutErrorReason))
+	if (!ZipGeneralSceneDescription(Zip, TempGeneralSceneDescription, OutErrorReason))
 	{
 		return;
 	}
@@ -63,7 +72,7 @@ void FDMXMVRExporter::ExportInternal(UDMXLibrary* DMXLibrary, const FString& Fil
 		// Allow continuation of export
 	}
 
-	ZipThirdPartyData(Zip, GeneralSceneDescription);
+	ZipThirdPartyData(Zip, TempGeneralSceneDescription);
 	if (!Zip->SaveToFile(FilePathAndName))
 	{
 		OutErrorReason = FText::Format(LOCTEXT("MVRExportWriteZipFailedReason", "File is not writable or locked by another process. Cannot export {0}."), FText::FromString(FilePathAndName));
@@ -207,6 +216,73 @@ void FDMXMVRExporter::ZipThirdPartyData(const TSharedRef<FDMXZipper>& Zip, const
 	}
 }
 
+void FDMXMVRExporter::WriteMVRFixtureTransformsFromLevel(UDMXMVRGeneralSceneDescription* GeneralSceneDescription)
+{
+	if (!ensureMsgf(GeneralSceneDescription, TEXT("Trying to write MVR Fixture Transforms from Level to MVR General Scene Description, but General Scene Description is invalid")))
+	{
+		return;
+	}
+
+	UUnrealEditorSubsystem* UnrealEditorSubsystem = GEditor->GetEditorSubsystem<UUnrealEditorSubsystem>();
+	if (!UnrealEditorSubsystem)
+	{
+		return;
+	} 
+	
+	const UWorld* EditorWorld = UnrealEditorSubsystem->GetEditorWorld();
+	if (!EditorWorld)
+	{
+		return;
+	}
+
+	TArray<const ADMXMVRSceneActor*> MVRSceneActors;
+	for (TActorIterator<ADMXMVRSceneActor> It(EditorWorld, ADMXMVRSceneActor::StaticClass()); It; ++It)
+	{
+		ADMXMVRSceneActor* Actor = *It;
+		MVRSceneActors.Add(Actor);
+	}
+
+	TArray<UDMXMVRFixtureNode*> FixtureNodes;
+	GeneralSceneDescription->GetFixtureNodes(FixtureNodes);
+
+	TArray<UDMXMVRFixtureNode*> HandledMVRFixtureNodes;
+	for (const ADMXMVRSceneActor* MVRSceneActor : MVRSceneActors)
+	{
+		for (const TSoftObjectPtr<AActor>& RelatedActorSoftObjectPtr : MVRSceneActor->GetRelatedActors())
+		{
+			const AActor* RelatedActor = RelatedActorSoftObjectPtr.Get();
+			if (!RelatedActor)
+			{
+				continue;
+			}
+
+			const FString MVRFixtureUUIDString = UDMXMVRAssetUserData::GetMVRAssetUserDataValueForKey(*RelatedActor, UDMXMVRAssetUserData::MVRFixtureUUIDMetaDataKey);
+			FGuid MVRFixtureUUID;
+			if (!FGuid::Parse(MVRFixtureUUIDString, MVRFixtureUUID))
+			{
+				continue;
+			}
+
+			UDMXMVRFixtureNode** FixtureNodePtr = Algo::FindByPredicate(FixtureNodes, [MVRFixtureUUID](UDMXMVRFixtureNode* MVRFixtureNode)
+				{
+					return MVRFixtureNode->UUID == MVRFixtureUUID;
+				});
+
+			if (FixtureNodePtr)
+			{
+				if (HandledMVRFixtureNodes.Contains(*FixtureNodePtr))
+				{
+					UE_LOG(LogDMXEditor, Warning, TEXT("Ambigous Fixture ID for Fixture '%s' when exporting MVR. Each Fixture Patch can only be used once in an MVR scene. Using transform of arbitrary instance with ambigous ID."), *RelatedActor->GetName());
+					continue;
+				}
+
+				(*FixtureNodePtr)->SetTransformAbsolute(RelatedActor->GetTransform());
+				HandledMVRFixtureNodes.Add(*FixtureNodePtr);
+			}
+		}
+	}
+}
+
 const TSharedPtr<FXmlFile> FDMXMVRExporter::CreateSourceGeneralSceneDescriptionXmlFile(const UDMXMVRGeneralSceneDescription* GeneralSceneDescription) const
 {
 	const UDMXMVRAssetImportData* AssetImportData = GeneralSceneDescription->GetMVRAssetImportData();
@@ -243,7 +319,6 @@ const TSharedPtr<FXmlFile> FDMXMVRExporter::CreateSourceGeneralSceneDescriptionX
 	return XmlFile;
 }
 
-
 const TArray64<uint8>& FDMXMVRExporter::RefreshSourceDataAndFixtureType(UDMXEntityFixtureType& FixtureType, UDMXGDTFAssetImportData& InOutGDTFAssetImportData) const
 {
 	const UDMXImportGDTF* GDTF = Cast<UDMXImportGDTF>(InOutGDTFAssetImportData.GetOuter());
diff --git a/Engine/Plugins/VirtualProduction/DMX/DMXEngine/Source/DMXEditor/Private/Exporters/DMXMVRExporter.h b/Engine/Plugins/VirtualProduction/DMX/DMXEngine/Source/DMXEditor/Private/Exporters/DMXMVRExporter.h
--- a/Engine/Plugins/VirtualProduction/DMX/DMXEngine/Source/DMXEditor/Private/Exporters/DMXMVRExporter.h
+++ b/Engine/Plugins/VirtualProduction/DMX/DMXEngine/Source/DMXEditor/Private/Exporters/DMXMVRExporter.h
@@ -37,6 +37,9 @@ private:
 	/** Creates an General Scene Description Xml File from the MVR Source, as it was imported */
 	const TSharedPtr<FXmlFile> CreateSourceGeneralSceneDescriptionXmlFile(const UDMXMVRGeneralSceneDescription* GeneralSceneDescription) const;
 
+	/** Writes transforms from level to the MVR Fixture Actors where applicable */
+	void WriteMVRFixtureTransformsFromLevel(UDMXMVRGeneralSceneDescription* GeneralSceneDescription);
+
 	/** 
 	 * Gets raw source data or creates (possibly empty) source data where the source data is not present. 
 	 * Prior 5.1 there was no source data stored. Offers a dialog to load missing data 
