DMX - Fix crash when changing a GDTF for a Fixture Type

diff --git a/Engine/Plugins/VirtualProduction/DMX/DMXEngine/Source/DMXEditor/Private/Widgets/FixtureType/SDMXFixtureModeEditor.cpp b/Engine/Plugins/VirtualProduction/DMX/DMXEngine/Source/DMXEditor/Private/Widgets/FixtureType/SDMXFixtureModeEditor.cpp
--- a/Engine/Plugins/VirtualProduction/DMX/DMXEngine/Source/DMXEditor/Private/Widgets/FixtureType/SDMXFixtureModeEditor.cpp
+++ b/Engine/Plugins/VirtualProduction/DMX/DMXEngine/Source/DMXEditor/Private/Widgets/FixtureType/SDMXFixtureModeEditor.cpp
@@ -22,16 +22,53 @@ void SDMXFixtureModeEditor::Construct(const FArguments& InArgs, const TSharedRef
 	WeakDMXEditor = InDMXEditor;
 	FixtureTypeSharedData = InDMXEditor->GetFixtureTypeSharedData();
 
+	UDMXEntityFixtureType::GetOnFixtureTypeChanged().AddSP(this, &SDMXFixtureModeEditor::OnFixtureTypePropertiesChanged);
+	FixtureTypeSharedData->OnFixtureTypesSelected.AddSP(this, &SDMXFixtureModeEditor::Refresh);
 	FixtureTypeSharedData->OnFixtureTypesSelected.AddSP(this, &SDMXFixtureModeEditor::Refresh);
 	FixtureTypeSharedData->OnModesSelected.AddSP(this, &SDMXFixtureModeEditor::Refresh);
 
+	// Create a Struct Details View. This is not the most convenient type to work with as a property type customization for the FDMXFixtureMode struct cannot be used.
+	// Reason is soley significant performance gains, it's much faster than the easier approach with a UDMXEntityFixtureType customization as it was used up to 4.27.
+	FDetailsViewArgs DetailsViewArgs;
+	DetailsViewArgs.bAllowSearch = true;
+	DetailsViewArgs.bHideSelectionTip = false;
+	DetailsViewArgs.bSearchInitialKeyFocus = true;
+	DetailsViewArgs.bShowOptions = false;
+	DetailsViewArgs.bShowModifiedPropertiesOption = false;
+	DetailsViewArgs.bShowObjectLabel = false;
+	DetailsViewArgs.bForceHiddenPropertyVisibility = false;
+	DetailsViewArgs.bShowScrollBar = false;
+	DetailsViewArgs.NotifyHook = this;
+
+	FStructureDetailsViewArgs StructureDetailsViewArgs;
+
+	FPropertyEditorModule& PropertyModule = FModuleManager::LoadModuleChecked<FPropertyEditorModule>("PropertyEditor");
+	StructDetailsView = PropertyModule.CreateStructureDetailView(DetailsViewArgs, StructureDetailsViewArgs, nullptr);
+	StructDetailsView->GetDetailsView()->SetCustomFilterLabel(LOCTEXT("SearchModePropertiesLabel", "Search Mode Properties"));
+	StructDetailsView->GetDetailsView()->SetIsPropertyVisibleDelegate(FIsPropertyVisible::CreateSP(this, &SDMXFixtureModeEditor::IsPropertyVisible));
+
+	StructDetailsViewWidget = StructDetailsView->GetWidget();
+
 	ChildSlot
 	[
-		SAssignNew(ContentBorder, SBorder)
-		.HAlign(HAlign_Fill)
-		.VAlign(VAlign_Fill)
-		.BorderImage(FAppStyle::GetBrush("NoBorder"))
-		.Padding(4.f)
+		SNew(SVerticalBox)
+
+		+ SVerticalBox::Slot()
+		.AutoHeight()
+		[
+			StructDetailsViewWidget.ToSharedRef()
+		]
+
+		+ SVerticalBox::Slot()
+		[
+			SNew(SBox)
+			.HAlign(HAlign_Center)
+			.VAlign(VAlign_Center)
+			[
+				SAssignNew(InfoTextBlock, STextBlock)
+				.Visibility(EVisibility::Collapsed)
+			]
+		]
 	];
 
 	Refresh();
@@ -163,22 +200,14 @@ void SDMXFixtureModeEditor::Refresh()
 			return LOCTEXT("NoModeSelectedWarning", "No Mode selected");
 		}();
 
-		ContentBorder->SetContent
-		(
-			SNew(SBox)
-			.HAlign(HAlign_Center)
-			.VAlign(VAlign_Center)
-			[
-				SNew(STextBlock)
-				.Text(ErrorText)
-			]
-		);
+		InfoTextBlock->SetText(ErrorText);
+		InfoTextBlock->SetVisibility(EVisibility::Visible);
+		StructDetailsViewWidget->SetVisibility(EVisibility::Collapsed);
 	}
 }
 
 void SDMXFixtureModeEditor::SetMode(UDMXEntityFixtureType* InFixtureType, int32 InModeIndex)
 {
-	bool bSuccess = false;
 	if (ensureMsgf(InFixtureType && InFixtureType->Modes.IsValidIndex(InModeIndex), TEXT("Invalid Fixture Type or Mode when setting the Mode in the Mode editor.")))
 	{
 		WeakFixtureType = InFixtureType;
@@ -186,47 +215,24 @@ void SDMXFixtureModeEditor::SetMode(UDMXEntityFixtureType* InFixtureType, int32
 		FDMXFixtureMode& Mode = InFixtureType->Modes[ModeIndex];
 		const TSharedRef<FStructOnScope> ModeStructOnScope = MakeShared<FStructOnScope>(FDMXFixtureMode::StaticStruct(), (uint8*)&Mode);
 
-		// Create a Struct Details View. This is not the most convenient type to work with as a property type customization for the FDMXFixtureMode struct cannot be used.
-		// Reason is soley significant performance gains, it's much faster than the easier approach with a UDMXEntityFixtureType customization as it was used up to 4.27.
-		FDetailsViewArgs DetailsViewArgs; 
-		DetailsViewArgs.bAllowSearch = true;
-		DetailsViewArgs.bHideSelectionTip = false;
-		DetailsViewArgs.bSearchInitialKeyFocus = true;
-		DetailsViewArgs.bShowOptions = false;
-		DetailsViewArgs.bShowModifiedPropertiesOption = false;
-		DetailsViewArgs.bShowObjectLabel = false;
-		DetailsViewArgs.bForceHiddenPropertyVisibility = false;
-		DetailsViewArgs.bShowScrollBar = false;
-		DetailsViewArgs.NotifyHook = this;
-
-		FStructureDetailsViewArgs StructureDetailsViewArgs;
-
-		FPropertyEditorModule& PropertyModule = FModuleManager::LoadModuleChecked<FPropertyEditorModule>("PropertyEditor");
-		const TSharedRef<IStructureDetailsView> StructDetailsView = PropertyModule.CreateStructureDetailView(DetailsViewArgs, StructureDetailsViewArgs, ModeStructOnScope);
-		StructDetailsView->GetDetailsView()->SetCustomFilterLabel(LOCTEXT("SearchModePropertiesLabel", "Search Mode Properties"));
-
-		StructDetailsView->GetDetailsView()->SetIsPropertyVisibleDelegate(FIsPropertyVisible::CreateSP(this, &SDMXFixtureModeEditor::IsPropertyVisible));
-		StructDetailsView->GetDetailsView()->ForceRefresh();
-
-		if (TSharedPtr<SWidget> Widget = StructDetailsView->GetWidget())
-		{
-			ContentBorder->SetContent(Widget.ToSharedRef());
-			bSuccess = true;
-		}
+		StructDetailsView->SetStructureData(ModeStructOnScope);
+
+		StructDetailsViewWidget->SetVisibility(EVisibility::Visible);
+		InfoTextBlock->SetVisibility(EVisibility::Collapsed);
+	}
+	else
+	{
+		InfoTextBlock->SetText(LOCTEXT("CannotCreateDetailViewForModeWarning", "Cannot create Detail View for Mode. Fixture Type or Mode no longer exist."));
+		InfoTextBlock->SetVisibility(EVisibility::Visible);
+		StructDetailsViewWidget->SetVisibility(EVisibility::Collapsed);
 	}
+}
 
-	if (!bSuccess)
+void SDMXFixtureModeEditor::OnFixtureTypePropertiesChanged(const UDMXEntityFixtureType* FixtureType)
+{
+	if (!Transaction.IsValid() && FixtureType == WeakFixtureType.Get())
 	{
-		ContentBorder->SetContent
-		(
-			SNew(SBox)
-			.HAlign(HAlign_Center)
-			.VAlign(VAlign_Center)
-			[
-				SNew(STextBlock)
-				.Text(LOCTEXT("CannotCreateDetailViewForModeWarning", "Cannot create Detail View for Mode. Fixture Type or Mode no longer exist."))
-			]
-		);
+		Refresh();
 	}
 }
 
diff --git a/Engine/Plugins/VirtualProduction/DMX/DMXEngine/Source/DMXEditor/Private/Widgets/FixtureType/SDMXFixtureModeEditor.h b/Engine/Plugins/VirtualProduction/DMX/DMXEngine/Source/DMXEditor/Private/Widgets/FixtureType/SDMXFixtureModeEditor.h
--- a/Engine/Plugins/VirtualProduction/DMX/DMXEngine/Source/DMXEditor/Private/Widgets/FixtureType/SDMXFixtureModeEditor.h
+++ b/Engine/Plugins/VirtualProduction/DMX/DMXEngine/Source/DMXEditor/Private/Widgets/FixtureType/SDMXFixtureModeEditor.h
@@ -16,6 +16,7 @@ struct FPropertyAndParent;
 class FScopedTransaction;
 class IStructureDetailsView;
 class SBorder;
+class STextBlock;
 
 
 /** Editor for a single Mode struct inside a UDMXEntityFixtureType. */
@@ -45,6 +46,9 @@ private:
 	/** Sets the Mode that is being edited */
 	void SetMode(UDMXEntityFixtureType* InFixtureType, int32 InModeIndex);
 
+	/** Called when the properties of the fixture type changed */
+	void OnFixtureTypePropertiesChanged(const UDMXEntityFixtureType* FixtureType);
+
 	/** Returns true if the properties are visible */
 	bool IsPropertyVisible(const FPropertyAndParent& PropertyAndParent) const;
 
@@ -60,8 +64,14 @@ private:
 	/** The current transaction or nullptr if there's no transaction ongoing */
 	TUniquePtr<FScopedTransaction> Transaction;
 
-	/** Border that holds the content of the editor */
-	TSharedPtr<SBorder> ContentBorder;
+	/** Details view for the Mode Struct being edited */
+	TSharedPtr<IStructureDetailsView> StructDetailsView;
+
+	/** Widget of the Details view */
+	TSharedPtr<SWidget> StructDetailsViewWidget;
+
+	/** Text block that displays info in case the struct details view cannot be shown */
+	TSharedPtr<STextBlock> InfoTextBlock;
 
 	/** Shared Data for */
 	TSharedPtr<FDMXFixtureTypeSharedData> FixtureTypeSharedData;
