DMX - Fix fixtures that do not specify a GDTF are spawned in MVR Scenes

diff --git a/Engine/Plugins/VirtualProduction/DMX/DMXEngine/Source/DMXRuntime/Private/MVR/DMXMVRSceneActor.cpp b/Engine/Plugins/VirtualProduction/DMX/DMXEngine/Source/DMXRuntime/Private/MVR/DMXMVRSceneActor.cpp
--- a/Engine/Plugins/VirtualProduction/DMX/DMXEngine/Source/DMXRuntime/Private/MVR/DMXMVRSceneActor.cpp
+++ b/Engine/Plugins/VirtualProduction/DMX/DMXEngine/Source/DMXRuntime/Private/MVR/DMXMVRSceneActor.cpp
@@ -322,6 +322,12 @@ AActor* ADMXMVRSceneActor::SpawnMVRActor(const TSubclassOf<AActor>&ActorClass, U
 		return nullptr;
 	}
 
+	if (!FixturePatch->GetFixtureType() || !FixturePatch->GetFixtureType()->DMXImport)
+	{
+		UE_LOG(LogDMXRuntime, Warning, TEXT("Cannot spawn fixture for Fixture Patch '%s'. It's Fixture Type does not source from a GDTF."), *FixturePatch->Name);
+		return nullptr;
+	}
+
 	FActorSpawnParameters ActorSpawnParameters;
 	ActorSpawnParameters.Template = Template;
 	ActorSpawnParameters.SpawnCollisionHandlingOverride = ESpawnActorCollisionHandlingMethod::AlwaysSpawn;
