DMX - Fix various issues with the new Fixture Patch editor, now performs as expected, copy/paste/cut work as expected.

diff --git a/Engine/Plugins/VirtualProduction/DMX/DMXEngine/Source/DMXEditor/Private/Customizations/DMXEntityFixturePatchDetails.cpp b/Engine/Plugins/VirtualProduction/DMX/DMXEngine/Source/DMXEditor/Private/Customizations/DMXEntityFixturePatchDetails.cpp
--- a/Engine/Plugins/VirtualProduction/DMX/DMXEngine/Source/DMXEditor/Private/Customizations/DMXEntityFixturePatchDetails.cpp
+++ b/Engine/Plugins/VirtualProduction/DMX/DMXEngine/Source/DMXEditor/Private/Customizations/DMXEntityFixturePatchDetails.cpp
@@ -120,11 +120,12 @@ TSharedRef<SWidget> FDMXEntityFixturePatchDetails::GenerateActiveModeWidget(cons
 void FDMXEntityFixturePatchDetails::OnParentFixtureTypeChanged(UDMXEntity* NewTemplate) const
 {
 	ParentFixtureTypeHandle->SetValue(Cast<UDMXEntityFixtureType>(NewTemplate));
+	PropertyUtilities->ForceRefresh();
 }
 
 void FDMXEntityFixturePatchDetails::OnFixtureTypeChanged(const UDMXEntityFixtureType* FixtureType)
 {
-	if (IsValid(FixtureType) && !FixtureType->HasAnyFlags(RF_Transactional))
+	if (IsValid(FixtureType))
 	{
 		// Keep the active mode valid
 		int32 ActiveMode;
@@ -264,11 +265,14 @@ void FDMXEntityFixturePatchDetails::SetActiveMode(int32 ModeIndex)
 
 	if (OuterObjects.Num() > 0)
 	{
-		const FScopedTransaction Transaction(LOCTEXT("SetFixturePatchActiveModeTransaction", "Set DMX Fixture Patch Active Mode"));
-
 		for (UObject* Object : OuterObjects)
 		{
 			UDMXEntityFixturePatch* Patch = CastChecked<UDMXEntityFixturePatch>(Object);
+			if (Patch->GetActiveModeIndex() == ModeIndex)
+			{
+				continue;
+			}
+			
 			Patch->Modify();
 			Patch->PreEditChange(UDMXEntityFixturePatch::StaticClass()->FindPropertyByName(UDMXEntityFixturePatch::GetActiveModePropertyNameChecked()));
 
diff --git a/Engine/Plugins/VirtualProduction/DMX/DMXEngine/Source/DMXEditor/Private/Widgets/FixturePatch/DMXMVRFixtureListItem.cpp b/Engine/Plugins/VirtualProduction/DMX/DMXEngine/Source/DMXEditor/Private/Widgets/FixturePatch/DMXMVRFixtureListItem.cpp
--- a/Engine/Plugins/VirtualProduction/DMX/DMXEngine/Source/DMXEditor/Private/Widgets/FixturePatch/DMXMVRFixtureListItem.cpp
+++ b/Engine/Plugins/VirtualProduction/DMX/DMXEngine/Source/DMXEditor/Private/Widgets/FixturePatch/DMXMVRFixtureListItem.cpp
@@ -15,8 +15,8 @@
 
 #define LOCTEXT_NAMESPACE "DMXMVRFixtureListItem"
 
-FDMXMVRFixtureListItem::FDMXMVRFixtureListItem(TWeakPtr<FDMXEditor> InDMXEditor, const FGuid& InMVRFixtureUUID)
-	: MVRFixtureUUID(InMVRFixtureUUID)
+FDMXMVRFixtureListItem::FDMXMVRFixtureListItem(TWeakPtr<FDMXEditor> InDMXEditor, UDMXMVRFixtureNode& InMVRFixtureNode)
+	: MVRFixtureNode(&InMVRFixtureNode)
 	, WeakDMXEditor(InDMXEditor)
 {
 	const UDMXLibrary* const DMXLibrary = GetDMXLibrary();
@@ -35,7 +35,7 @@ FDMXMVRFixtureListItem::FDMXMVRFixtureListItem(TWeakPtr<FDMXEditor> InDMXEditor,
 	const TArray<UDMXEntityFixturePatch*> FixturePatches = DMXLibrary->GetEntitiesTypeCast<UDMXEntityFixturePatch>();
 	UDMXEntityFixturePatch* const* FixturePatchPtr = FixturePatches.FindByPredicate([this](const UDMXEntityFixturePatch* FixturePatch)
 		{
-			return FixturePatch->GetMVRFixtureUUID() == MVRFixtureUUID;
+			return FixturePatch->GetMVRFixtureUUID() == MVRFixtureNode->UUID;
 		});
 
 	if (!ensureAlwaysMsgf(FixturePatchPtr, TEXT("Trying to create an MVR Fixture List Item, but there's no corresponding Fixture Patch for the MVR Fixture UUID.")))
@@ -45,13 +45,13 @@ FDMXMVRFixtureListItem::FDMXMVRFixtureListItem(TWeakPtr<FDMXEditor> InDMXEditor,
 	WeakFixturePatch = *FixturePatchPtr;
 
 	// Keep the MVR Fixture Name sync with the Fixture Patch name
-	MVRFixtureNode = FindMVRFixture();
 	MVRFixtureNode->Name = WeakFixturePatch->Name;
 }
 
 const FGuid& FDMXMVRFixtureListItem::GetMVRUUID() const
 {
-	return MVRFixtureUUID;
+	check(MVRFixtureNode);
+	return MVRFixtureNode->UUID;
 }
 
 FLinearColor FDMXMVRFixtureListItem::GetBackgroundColor() const
@@ -81,6 +81,7 @@ FString FDMXMVRFixtureListItem::GetFixturePatchName() const
 
 void FDMXMVRFixtureListItem::SetFixturePatchName(const FString& InDesiredName, FString& OutNewName)
 {
+	check(MVRFixtureNode);
 	if (UDMXEntityFixturePatch* FixturePatch = WeakFixturePatch.Get())
 	{
 		if (FixturePatch->Name == InDesiredName)
@@ -95,7 +96,6 @@ void FDMXMVRFixtureListItem::SetFixturePatchName(const FString& InDesiredName, F
 		FixturePatch->SetName(InDesiredName);
 		OutNewName = FixturePatch->Name;
 
-		MVRFixtureNode = FindMVRFixture();
 		MVRFixtureNode->Name = FixturePatch->Name;
 
 		FixturePatch->PostEditChange();
@@ -104,13 +104,16 @@ void FDMXMVRFixtureListItem::SetFixturePatchName(const FString& InDesiredName, F
 
 FString FDMXMVRFixtureListItem::GetFixtureID() const
 {
+	check(MVRFixtureNode);
 	return MVRFixtureNode->FixtureID;
 }
 
 void FDMXMVRFixtureListItem::SetFixtureID(int32 InFixtureID)
 {
+	check(MVRFixtureNode);
+
 	UDMXLibrary* DMXLibrary = GetDMXLibrary();
-	if (MVRFixtureNode && DMXLibrary)
+	if (DMXLibrary)
 	{
 		const FScopedTransaction SetMVRFixtureFixtureIDTransaction(LOCTEXT("SetMVRFixtureFixtureIDTransaction", "Set MVR Fixture ID"));
 		DMXLibrary->PreEditChange(UDMXLibrary::StaticClass()->FindPropertyByName(UDMXLibrary::GetGeneralSceneDescriptionPropertyName()));
@@ -256,67 +259,6 @@ UDMXEntityFixturePatch* FDMXMVRFixtureListItem::GetFixturePatch() const
 	return WeakFixturePatch.Get();
 }
 
-void FDMXMVRFixtureListItem::PasteItemsOntoItem(TWeakPtr<FDMXEditor> WeakDMXEditor, const TSharedPtr<FDMXMVRFixtureListItem>& PasteOntoItem, const TArray<TSharedPtr<FDMXMVRFixtureListItem>>& ItemsToPaste)
-{
-	TArray<UDMXEntityFixturePatch*> FixturePatchesToPaste;
-	for (const TSharedPtr<FDMXMVRFixtureListItem>& Item : ItemsToPaste)
-	{
-		if (UDMXEntityFixturePatch* FixturePatch = Item->GetFixturePatch())
-		{
-			FixturePatchesToPaste.Add(FixturePatch);
-		}
-	}
-
-	const FText PasteFixturePatchesTransactionText = FText::Format(LOCTEXT("DuplicateFixturePatchesTransaction", "Duplicate Fixture {0}|plural(one=Patch, other=Patches)"), ItemsToPaste.Num() > 1);
-	DuplicateFixturePatchesInternal(WeakDMXEditor, FixturePatchesToPaste, PasteFixturePatchesTransactionText);
-}
-
-void FDMXMVRFixtureListItem::DuplicateItems(TWeakPtr<FDMXEditor> WeakDMXEditor, const TArray<TSharedPtr<FDMXMVRFixtureListItem>>& ItemsToDuplicate)
-{
-	TArray<UDMXEntityFixturePatch*> FixturePatchesToDuplicate;
-	for (const TSharedPtr<FDMXMVRFixtureListItem>& Item : ItemsToDuplicate)
-	{
-		if (UDMXEntityFixturePatch* FixturePatch = Item->GetFixturePatch())
-		{
-			FixturePatchesToDuplicate.Add(FixturePatch);
-		}
-	}
-
-	const FText DuplicateFixturePatchesTransactionText = FText::Format(LOCTEXT("DuplicateFixturePatchesTransaction", "Duplicate Fixture {0}|plural(one=Patch, other=Patches)"), ItemsToDuplicate.Num() > 1);
-	DuplicateFixturePatchesInternal(WeakDMXEditor, FixturePatchesToDuplicate, DuplicateFixturePatchesTransactionText);
-}
-
-void FDMXMVRFixtureListItem::DeleteItems(const TArray<TSharedPtr<FDMXMVRFixtureListItem>>& ItemsToDelete)
-{
-	if (ItemsToDelete.Num() == 0)
-	{
-		return;
-	}
-	
-	// Its safe to assume all patches are in the same Library - A Multi-Library Editor wouldn't make sense.
-	UDMXLibrary* DMXLibrary = ItemsToDelete[0]->GetDMXLibrary();
-	if (!DMXLibrary)
-	{
-		return;
-	}
-	
-	const FText DeleteFixturePatchesTransactionText = FText::Format(LOCTEXT("DeleteFixturePatchesTransaction", "Delete Fixture {0}|plural(one=Patch, other=Patches)"), ItemsToDelete.Num() > 1);
-	const FScopedTransaction DeleteFixturePatchTransaction(DeleteFixturePatchesTransactionText);
-	DMXLibrary->PreEditChange(nullptr);
-
-	for (const TSharedPtr<FDMXMVRFixtureListItem>& Item : ItemsToDelete)
-	{
-		if (UDMXEntityFixturePatch* FixturePatch = Item->GetFixturePatch())
-		{
-			FixturePatch->PreEditChange(UDMXEntityFixturePatch::StaticClass()->FindPropertyByName(UDMXEntityFixturePatch::GetMVRFixtureUUIDPropertyNameChecked()));
-			UDMXEntityFixturePatch::RemoveFixturePatchFromLibrary(FixturePatch);
-			FixturePatch->PostEditChange();
-		}
-	}
-
-	DMXLibrary->PostEditChange();
-}
-
 UDMXLibrary* FDMXMVRFixtureListItem::GetDMXLibrary() const
 {
 	if (const TSharedPtr<FDMXEditor> DMXEditor = WeakDMXEditor.Pin())
@@ -336,124 +278,4 @@ void FDMXMVRFixtureListItem::AddReferencedObjects(FReferenceCollector& Collector
 	Collector.AddReferencedObject(MVRFixtureNode);
 }
 
-void FDMXMVRFixtureListItem::DuplicateFixturePatchesInternal(TWeakPtr<FDMXEditor> WeakDMXEditor, const TArray<UDMXEntityFixturePatch*>& FixturePatchesToDuplicate, const FText& TransactionText)
-{
-	const TSharedPtr<FDMXEditor> DMXEditor = WeakDMXEditor.Pin();
-	if (!DMXEditor.IsValid())
-	{
-		return;
-	}
-	const TSharedPtr<FDMXFixturePatchSharedData> FixturePatchSharedData = DMXEditor->GetFixturePatchSharedData();
-
-	UDMXLibrary* DMXLibrary = DMXEditor->GetDMXLibrary();
-	if (!DMXLibrary)
-	{
-		return;
-	}
-
-	// Find the first free Addresses
-	TArray<UDMXEntityFixturePatch*> FixturePatches = DMXLibrary->GetEntitiesTypeCast<UDMXEntityFixturePatch>();
-	if (FixturePatches.Num() == 0)
-	{
-		return;
-	}
-
-	FixturePatches.Sort([](const UDMXEntityFixturePatch& FixturePatchA, const UDMXEntityFixturePatch& FixturePatchB)
-		{
-			const bool bUniverseIsSmaller = FixturePatchA.GetUniverseID() < FixturePatchB.GetUniverseID();
-			const bool bUniverseIsEqual = FixturePatchA.GetUniverseID() == FixturePatchB.GetUniverseID();
-			const bool bAddressIsSmallerOrEqual = FixturePatchA.GetStartingChannel() <= FixturePatchB.GetStartingChannel();
-
-			return bUniverseIsSmaller || (bUniverseIsEqual && bAddressIsSmallerOrEqual);
-		});
-
-	int32 Address = FixturePatches.Last()->GetStartingChannel() + FixturePatches.Last()->GetChannelSpan();
-	int32 Universe = -1;
-	if (Address > DMX_MAX_ADDRESS)
-	{
-		Address = 1;
-		Universe = FixturePatches.Last()->GetUniverseID() + 1;
-	}
-	else
-	{
-		Universe = FixturePatches.Last()->GetUniverseID();
-	}
-
-	// Duplicate
-	const FScopedTransaction DuplicateFixturePatchTransaction(TransactionText);
-	DMXLibrary->PreEditChange(UDMXLibrary::StaticClass()->FindPropertyByName(UDMXLibrary::GetEntitiesPropertyName()));
-
-	TArray<TWeakObjectPtr<UDMXEntityFixturePatch>> NewFixturePatches;
-	for (const UDMXEntityFixturePatch* FixturePatchToDuplicate : FixturePatchesToDuplicate)
-	{
-		// If this is duplicated from one library onto another (e.g. when pasting via duplicate), create a Copy of the Fixture Type in the duplicated to DMX Library.
-		UDMXEntityFixtureType* FixtureTypeOfPatchToDuplicate = FixturePatchToDuplicate->GetFixtureType();
-		if (FixtureTypeOfPatchToDuplicate && FixtureTypeOfPatchToDuplicate->GetParentLibrary() != DMXLibrary)
-		{
-			FDMXEntityFixtureTypeConstructionParams FixtureTypeConstructionParams;
-			FixtureTypeConstructionParams.DMXCategory = FixtureTypeOfPatchToDuplicate->DMXCategory;
-			FixtureTypeConstructionParams.Modes = FixtureTypeOfPatchToDuplicate->Modes;
-			FixtureTypeConstructionParams.ParentDMXLibrary = DMXLibrary;
-
-			constexpr bool bMarkLibraryDirty = false;
-			UDMXEntityFixtureType::CreateFixtureTypeInLibrary(FixtureTypeConstructionParams, FixtureTypeOfPatchToDuplicate->Name, bMarkLibraryDirty);
-		}
-
-		// Duplicate the Fixture Patch
-		const int32 ChannelSpan = FixturePatchToDuplicate->GetChannelSpan();
-		if (Address + ChannelSpan - 1 > DMX_MAX_ADDRESS)
-		{
-			Address = 1;
-			Universe++;
-		}
-
-		FDMXEntityFixturePatchConstructionParams ConstructionParams;
-		ConstructionParams.FixtureTypeRef = FixturePatchToDuplicate->GetFixtureType();
-		ConstructionParams.ActiveMode = FixturePatchToDuplicate->GetActiveModeIndex();
-		ConstructionParams.UniverseID = Universe;
-		ConstructionParams.StartingAddress = Address;
-
-		constexpr bool bMarkLibraryDirty = false;
-		UDMXEntityFixturePatch* NewFixturePatch = UDMXEntityFixturePatch::CreateFixturePatchInLibrary(ConstructionParams, FixturePatchToDuplicate->Name, bMarkLibraryDirty);
-		NewFixturePatches.Add(NewFixturePatch);
-
-		Address += ChannelSpan;
-	}
-
-	DMXLibrary->PostEditChange();
-
-	if (NewFixturePatches.Num() > 0)
-	{
-		// Select universe and fixture patch
-		const TWeakObjectPtr<UDMXEntityFixturePatch>* FirstFixturePatchPtr = Algo::MinElementBy(NewFixturePatches, [](TWeakObjectPtr<UDMXEntityFixturePatch> FixturePatch)
-			{
-				return FixturePatch->GetUniverseID();
-			});
-		if (FirstFixturePatchPtr->IsValid())
-		{
-			FixturePatchSharedData->SelectUniverse(FirstFixturePatchPtr->Get()->GetUniverseID());
-		}
-		FixturePatchSharedData->SelectFixturePatches(NewFixturePatches);
-	}
-}
-
-UDMXMVRFixtureNode* FDMXMVRFixtureListItem::FindMVRFixture() const
-{
-	UDMXLibrary* DMXLibrary = GetDMXLibrary();
-
-	if (!DMXLibrary)
-	{
-		return nullptr;
-	}
-
-	DMXLibrary->UpdateGeneralSceneDescription();
-	UDMXMVRGeneralSceneDescription* GeneralSceneDescription = DMXLibrary->GetLazyGeneralSceneDescription();
-	if (!ensureMsgf(GeneralSceneDescription, TEXT("Found Library without a General Scene Description. This should never occur.")))
-	{
-		return nullptr;
-	}
-
-	return GeneralSceneDescription->FindFixtureNode(MVRFixtureUUID);
-}
-
 #undef LOCTEXT_NAMESPACE
diff --git a/Engine/Plugins/VirtualProduction/DMX/DMXEngine/Source/DMXEditor/Private/Widgets/FixturePatch/DMXMVRFixtureListItem.h b/Engine/Plugins/VirtualProduction/DMX/DMXEngine/Source/DMXEditor/Private/Widgets/FixturePatch/DMXMVRFixtureListItem.h
--- a/Engine/Plugins/VirtualProduction/DMX/DMXEngine/Source/DMXEditor/Private/Widgets/FixturePatch/DMXMVRFixtureListItem.h
+++ b/Engine/Plugins/VirtualProduction/DMX/DMXEngine/Source/DMXEditor/Private/Widgets/FixturePatch/DMXMVRFixtureListItem.h
@@ -23,7 +23,7 @@ class FDMXMVRFixtureListItem
 {
 public:	
 	/** Constructor */
-	FDMXMVRFixtureListItem(TWeakPtr<FDMXEditor> InDMXEditor, const FGuid& InMVRFixtureUUID);
+	FDMXMVRFixtureListItem(TWeakPtr<FDMXEditor> InDMXEditor, UDMXMVRFixtureNode& InMVRFixtureNode);
 
 	/** Returns the MVR UUID of the MVR Fixture */
 	const FGuid& GetMVRUUID() const;
@@ -67,15 +67,6 @@ public:
 	/** Returns Num Channels the MVR Fixture spans */
 	int32 GetNumChannels() const;
 
-	/** Pastes Pasted Items onto an item */
-	static void PasteItemsOntoItem(TWeakPtr<FDMXEditor> WeakDMXEditor, const TSharedPtr<FDMXMVRFixtureListItem>& PasteOntoItem, const TArray<TSharedPtr<FDMXMVRFixtureListItem>>& ItemsToPaste);
-
-	/** Duplicates the Fixture Patches of the items */
-	static void DuplicateItems(TWeakPtr<FDMXEditor> WeakDMXEditor, const TArray<TSharedPtr<FDMXMVRFixtureListItem>>& ItemsToDuplicate);
-
-	/** Deletes specified items. Note, if all items of a Fixture Patch is deleted, the Fixture Patch is remvoed from its DMX Library. */
-	static void DeleteItems(const TArray<TSharedPtr<FDMXMVRFixtureListItem>>& ItemsToDelete);
-
 	/** Returns the Fixture Patch. Const on purpose, should be edited via this class's methods */
 	UDMXEntityFixturePatch* GetFixturePatch() const;
 
@@ -99,21 +90,12 @@ protected:
 
 
 private:
-	/** Duplicates Fixture Patches */
-	static void DuplicateFixturePatchesInternal(TWeakPtr<FDMXEditor> WeakDMXEditor, const TArray<UDMXEntityFixturePatch*>& FixturePatchesToDuplicate, const FText& TransactionText);
-
-	/** Returns the MVR Fixture Node that corresponds to this item */
-	UDMXMVRFixtureNode* FindMVRFixture() const;
-
 	/** True if this is in an even group */
 	bool bIsEvenGroup = false;
 
 	/** Cached MVR Fixture to speed up getters. */
 	UDMXMVRFixtureNode* MVRFixtureNode = nullptr;
 
-	/** The MVRFixtureUUID this Item handles */
-	FGuid MVRFixtureUUID;
-
 	/** The fixture patch the MVR Fixture UUID is assigned to */
 	TWeakObjectPtr<UDMXEntityFixturePatch> WeakFixturePatch;
 
diff --git a/Engine/Plugins/VirtualProduction/DMX/DMXEngine/Source/DMXEditor/Private/Widgets/FixturePatch/SDMXMVRFixtureList.cpp b/Engine/Plugins/VirtualProduction/DMX/DMXEngine/Source/DMXEditor/Private/Widgets/FixturePatch/SDMXMVRFixtureList.cpp
--- a/Engine/Plugins/VirtualProduction/DMX/DMXEngine/Source/DMXEditor/Private/Widgets/FixturePatch/SDMXMVRFixtureList.cpp
+++ b/Engine/Plugins/VirtualProduction/DMX/DMXEngine/Source/DMXEditor/Private/Widgets/FixturePatch/SDMXMVRFixtureList.cpp
@@ -6,6 +6,7 @@
 #include "DMXEditorSettings.h"
 #include "DMXEditorUtils.h"
 #include "DMXFixturePatchSharedData.h"
+#include "DMXRuntimeUtils.h"
 #include "Commands/DMXEditorCommands.h"
 #include "Framework/Application/SlateApplication.h"
 #include "Library/DMXEntityFixturePatch.h"
@@ -15,9 +16,14 @@
 #include "Widgets/FixturePatch/SDMXMVRFixtureListRow.h"
 #include "Widgets/FixturePatch/SDMXMVRFixtureListToolbar.h"
 
+#include "Factories.h"
 #include "ScopedTransaction.h"
 #include "TimerManager.h"
+#include "UnrealExporter.h"
+#include "Algo/MinElement.h"
 #include "Algo/Sort.h"
+#include "Exporters/Exporter.h"
+#include "Framework/Application/SlateApplication.h"
 #include "Framework/Commands/GenericCommands.h"
 #include "Framework/Commands/UICommandList.h"
 #include "Framework/MultiBox/MultiBoxBuilder.h"
@@ -30,6 +36,240 @@
 
 #define LOCTEXT_NAMESPACE "SDMXMVRFixtureList"
 
+namespace UE::DMX::SDMXMVRFixtureList::Private
+{
+	/** Copies a fixture patch as text to the clipboard */
+	static void ClipboardCopyFixturePatches(const TArray<UDMXEntityFixturePatch*>& FixturePatches)
+	{
+		// Clear the mark state for saving.
+		UnMarkAllObjects(EObjectMark(OBJECTMARK_TagExp | OBJECTMARK_TagImp));
+
+		const FExportObjectInnerContext Context;
+		FStringOutputDevice Archive;
+
+		// Export the component object(s) to text for copying
+		for (UDMXEntityFixturePatch* FixturePatch : FixturePatches)
+		{
+			if (!FixturePatch)
+			{
+				continue;
+			}
+
+			// Export the entity object to the given string
+			UExporter::ExportToOutputDevice(&Context, FixturePatch, nullptr, Archive, TEXT("copy"), 0, PPF_ExportsNotFullyQualified | PPF_Copy | PPF_Delimited, false, GetTransientPackage());
+		}
+
+		// Copy text to clipboard
+		FString ExportedText = Archive;
+
+		// Avoid exporting the OnFixturePatchReceived Binding
+		TArray<FString> ExportedTextLines;
+		constexpr bool bCullEmpty = false;
+		ExportedText.ParseIntoArrayLines(ExportedTextLines, bCullEmpty);
+		FString ExportedTextWithoutOnFixturePatchReceivedBinding;
+		for (const FString& String : ExportedTextLines)
+		{
+			if (String.Contains(TEXT("OnFixturePatchReceivedDMX")))
+			{
+				continue;
+			}
+			ExportedTextWithoutOnFixturePatchReceivedBinding.Append(String + LINE_TERMINATOR);
+		}
+
+		FPlatformApplicationMisc::ClipboardCopy(*ExportedTextWithoutOnFixturePatchReceivedBinding);
+	}
+
+	/** Duplicates an existing patch */
+	UDMXEntityFixturePatch* DuplicatePatchByMVRFixtureUUID(UDMXLibrary* DMXLibrary, const FGuid& MVRFixtureUUID)
+	{
+		if (!DMXLibrary)
+		{
+			return nullptr;
+		}
+
+		// Find the first free Addresses
+		TArray<UDMXEntityFixturePatch*> FixturePatches = DMXLibrary->GetEntitiesTypeCast<UDMXEntityFixturePatch>();
+		if (FixturePatches.Num() == 0)
+		{
+			return nullptr;
+		}
+
+		FixturePatches.Sort([](const UDMXEntityFixturePatch& FixturePatchA, const UDMXEntityFixturePatch& FixturePatchB)
+			{
+				const bool bUniverseIsSmaller = FixturePatchA.GetUniverseID() < FixturePatchB.GetUniverseID();
+				const bool bUniverseIsEqual = FixturePatchA.GetUniverseID() == FixturePatchB.GetUniverseID();
+				const bool bAddressIsSmallerOrEqual = FixturePatchA.GetStartingChannel() <= FixturePatchB.GetStartingChannel();
+
+				return bUniverseIsSmaller || (bUniverseIsEqual && bAddressIsSmallerOrEqual);
+			});
+
+		int32 Address = FixturePatches.Last()->GetStartingChannel() + FixturePatches.Last()->GetChannelSpan();
+		int32 Universe = -1;
+		if (Address > DMX_MAX_ADDRESS)
+		{
+			Address = 1;
+			Universe = FixturePatches.Last()->GetUniverseID() + 1;
+		}
+		else
+		{
+			Universe = FixturePatches.Last()->GetUniverseID();
+		}
+
+		UDMXEntityFixturePatch** FixturePatchToDuplicatePtr = Algo::FindByPredicate(FixturePatches, [MVRFixtureUUID](const UDMXEntityFixturePatch* FixturePatch)
+			{
+				return FixturePatch->GetMVRFixtureUUID() == MVRFixtureUUID;
+			});
+
+		if (!ensureMsgf(FixturePatchToDuplicatePtr, TEXT("Trying to duplicate fixture patch, but source fixture patch cannot be found")))
+		{
+			return nullptr;
+		}
+		UDMXEntityFixturePatch* FixturePatchToDuplicate = *FixturePatchToDuplicatePtr;
+		check(FixturePatchToDuplicate);
+
+		// Duplicate
+		DMXLibrary->PreEditChange(UDMXLibrary::StaticClass()->FindPropertyByName(UDMXLibrary::GetEntitiesPropertyName()));
+		UDMXEntityFixtureType* FixtureTypeOfPatchToDuplicate = FixturePatchToDuplicate->GetFixtureType();
+		if (FixtureTypeOfPatchToDuplicate && FixtureTypeOfPatchToDuplicate->GetParentLibrary() != DMXLibrary)
+		{
+			FDMXEntityFixtureTypeConstructionParams FixtureTypeConstructionParams;
+			FixtureTypeConstructionParams.DMXCategory = FixtureTypeOfPatchToDuplicate->DMXCategory;
+			FixtureTypeConstructionParams.Modes = FixtureTypeOfPatchToDuplicate->Modes;
+			FixtureTypeConstructionParams.ParentDMXLibrary = DMXLibrary;
+
+			constexpr bool bMarkLibraryDirty = false;
+			UDMXEntityFixtureType::CreateFixtureTypeInLibrary(FixtureTypeConstructionParams, FixtureTypeOfPatchToDuplicate->Name, bMarkLibraryDirty);
+		}
+
+		// Duplicate the Fixture Patch
+		const int32 ChannelSpan = FixturePatchToDuplicate->GetChannelSpan();
+		if (Address + ChannelSpan - 1 > DMX_MAX_ADDRESS)
+		{
+			Address = 1;
+			Universe++;
+		}
+
+		FDMXEntityFixturePatchConstructionParams ConstructionParams;
+		ConstructionParams.FixtureTypeRef = FixturePatchToDuplicate->GetFixtureType();
+		ConstructionParams.ActiveMode = FixturePatchToDuplicate->GetActiveModeIndex();
+		ConstructionParams.UniverseID = Universe;
+		ConstructionParams.StartingAddress = Address;
+
+		constexpr bool bMarkLibraryDirty = false;
+		UDMXEntityFixturePatch* NewFixturePatch = UDMXEntityFixturePatch::CreateFixturePatchInLibrary(ConstructionParams, FixturePatchToDuplicate->Name, bMarkLibraryDirty);
+
+		Address += ChannelSpan;
+
+		DMXLibrary->PostEditChange();
+
+		return NewFixturePatch;
+	}
+
+	/** Text object factory for pasting DMX Fixture Patches */
+	struct FDMXFixturePatchObjectTextFactory final
+		: public FCustomizableTextObjectFactory
+	{
+		/** Constructor */
+		FDMXFixturePatchObjectTextFactory(UDMXLibrary* InDMXLibrary)
+			: FCustomizableTextObjectFactory(GWarn)
+			, WeakDMXLibrary(InDMXLibrary)
+		{}
+
+		/** Returns true if Fixture Patches can be constructed from the Text Buffer */
+		static bool CanCreate(const FString& InTextBuffer, UDMXLibrary* InDMXLibrary)
+		{
+			TSharedRef<FDMXFixturePatchObjectTextFactory> Factory = MakeShared<FDMXFixturePatchObjectTextFactory>(InDMXLibrary);
+
+			// Create new objects if we're allowed to
+			return Factory->CanCreateObjectsFromText(InTextBuffer);
+		}
+
+		/**
+		 * Constructs a new object factory from the given text buffer. Returns the factor or nullptr if no factory can be created.
+		 * An updated General Scene Description of the library needs be passed explicitly to avoid recurring update calls.
+		 */
+		static bool Create(const FString& InTextBuffer, UDMXLibrary* InDMXLibrary, TArray<UDMXEntityFixturePatch*>& OutNewFixturePatches)
+		{
+			if (!InDMXLibrary)
+			{
+				return false;
+			}
+
+			OutNewFixturePatches.Reset();
+
+			const TSharedRef<FDMXFixturePatchObjectTextFactory> Factory = MakeShared < FDMXFixturePatchObjectTextFactory>(InDMXLibrary);
+
+			// Create new objects if we're allowed to
+			if (Factory->CanCreateObjectsFromText(InTextBuffer))
+			{
+				Factory->WeakDMXLibrary = InDMXLibrary;
+
+				EObjectFlags ObjectFlags = RF_Transactional;
+				Factory->ProcessBuffer(InDMXLibrary, ObjectFlags, InTextBuffer);
+
+				OutNewFixturePatches = Factory->NewFixturePatches;
+			}
+
+			return true;
+		}
+
+	protected:
+		//~ Begin FCustomizableTextObjectFactory implementation
+		virtual bool CanCreateClass(UClass* ObjectClass, bool& bOmitSubObjs) const override
+		{
+			return ObjectClass->IsChildOf(UDMXEntityFixturePatch::StaticClass());
+		}
+
+		virtual void ProcessConstructedObject(UObject* NewObject) override
+		{
+			UDMXLibrary* DMXLibrary = WeakDMXLibrary.Get();
+			UDMXMVRGeneralSceneDescription* GeneralSceneDescription = DMXLibrary ? DMXLibrary->GetLazyGeneralSceneDescription() : nullptr;
+			UDMXEntityFixturePatch* NewFixturePatch = Cast<UDMXEntityFixturePatch>(NewObject);
+			if (DMXLibrary && GeneralSceneDescription && NewFixturePatch)
+			{
+				const FScopedTransaction Transaction(TransactionText);
+				const FGuid MVRFixtureUUID = NewFixturePatch->GetMVRFixtureUUID();
+
+				const bool bIsDuplicating = GeneralSceneDescription->FindFixtureNode(MVRFixtureUUID) != nullptr;
+				if (bIsDuplicating)
+				{
+					// In cases where duplicates are created, create new instances to follow the common construction path for new Fixture Patches that ensures unqiue MVR UUIDs.
+					// Otherwise just add the new instances to the library, but keep the orig MVR UUID (e.g. when cut and paste or copy from one library to another).
+
+					using namespace UE::DMX::SDMXMVRFixtureList::Private;
+					UDMXEntityFixturePatch* DuplicatedFixturePatch = DuplicatePatchByMVRFixtureUUID(DMXLibrary, MVRFixtureUUID);
+					NewFixturePatches.Add(DuplicatedFixturePatch);
+
+					// Remove the fixtue patch added by the factory
+					UDMXEntityFixturePatch::RemoveFixturePatchFromLibrary(NewFixturePatch);
+				}
+				else
+				{
+					// Simply assign the new patch to the library
+					NewFixturePatch->Rename(*MakeUniqueObjectName(DMXLibrary, NewFixturePatch->GetClass()).ToString(), DMXLibrary, REN_DoNotDirty | REN_DontCreateRedirectors);
+					NewFixturePatch->SetName(FDMXRuntimeUtils::FindUniqueEntityName(DMXLibrary, NewFixturePatch->GetClass(), NewFixturePatch->GetDisplayName()));
+					NewFixturePatch->SetParentLibrary(DMXLibrary);
+					NewFixturePatch->RefreshID();
+
+					NewFixturePatches.Add(NewFixturePatch);
+				}
+			}
+		}
+		//~ End FCustomizableTextObjectFactory implementation
+
+	private:
+		/** Instantiated Fixture Patches */
+		TArray<UDMXEntityFixturePatch*> NewFixturePatches;
+
+		/** Transaction text displayed when pasting */
+		FText TransactionText;
+
+		/** Weak DMX Editor in which the operation should occur */
+		TWeakObjectPtr<UDMXLibrary> WeakDMXLibrary;
+	};
+}
+
+
 /** Helper to generate Status Text for MVR Fixture List Items */
 class FDMXMVRFixtureListStatusTextGenerator
 {
@@ -77,7 +317,6 @@ private:
 				InOutConflictTexts.Add(ItemToConflictTextPair);
 			}
 		}
-
 	}
 
 	/** The patch of an item. Useful to Get Conflicts with Other */
@@ -172,7 +411,9 @@ private:
 				const FText IssueText = LOCTEXT("NoModesIssue", "Fixture Type has no Modes defined.");
 				ItemToIssueMap.Add(Item, IssueText);
 			}
-			else if (Item->GetFixturePatch()->GetActiveMode() && Item->GetFixturePatch()->GetActiveMode()->Functions.IsEmpty())
+			else if (Item->GetFixturePatch()->GetActiveMode() && 
+				!Item->GetFixturePatch()->GetActiveMode()->bFixtureMatrixEnabled &&
+				Item->GetFixturePatch()->GetActiveMode()->Functions.IsEmpty())
 			{
 				const FText IssueText = LOCTEXT("ActiveModeHasNoFunctionsIssue", "Mode does not define any Functions.");
 				ItemToIssueMap.Add(Item, IssueText);
@@ -320,8 +561,6 @@ const FName FDMXMVRFixtureListCollumnIDs::FixtureType = "FixtureType";
 const FName FDMXMVRFixtureListCollumnIDs::Mode = "Mode";
 const FName FDMXMVRFixtureListCollumnIDs::Patch = "Patch";
 
-const FString SDMXMVRFixtureList::ClipboardCopyMVRFixtureHeader = FString(TEXT("UE::DMX::SDMXMVRFIXTURELIST::COPY_MVRFIXTUREEVENT"));
-
 SDMXMVRFixtureList::SDMXMVRFixtureList()
 	: SortMode(EColumnSortMode::Ascending)
 	, SortedByColumnID(FDMXMVRFixtureListCollumnIDs::Patch)
@@ -353,12 +592,13 @@ void SDMXMVRFixtureList::Construct(const FArguments& InArgs, TWeakPtr<FDMXEditor
 	FixturePatchSharedData = InDMXEditor.Pin()->GetFixturePatchSharedData();
 
 	// Handle Entity changes
+	UDMXLibrary::GetOnEntitiesAdded().AddSP(this, &SDMXMVRFixtureList::OnEntityAddedOrRemoved);
+	UDMXLibrary::GetOnEntitiesRemoved().AddSP(this, &SDMXMVRFixtureList::OnEntityAddedOrRemoved);
 	UDMXEntityFixturePatch::GetOnFixturePatchChanged().AddSP(this, &SDMXMVRFixtureList::OnFixturePatchChanged);
 	UDMXEntityFixtureType::GetOnFixtureTypeChanged().AddSP(this, &SDMXMVRFixtureList::OnFixtureTypeChanged);
 
 	// Handle Shared Data selection changes
 	FixturePatchSharedData->OnFixturePatchSelectionChanged.AddSP(this, &SDMXMVRFixtureList::OnFixturePatchSharedDataSelectedFixturePatches);
-	FixturePatchSharedData->OnUniverseSelectionChanged.AddSP(this, &SDMXMVRFixtureList::OnFixturePatchSharedDataSelectedUniverse);
 
 	static const FTableViewStyle TableViewStyle = FAppStyle::Get().GetWidgetStyle<FTableViewStyle>("TreeView");
 
@@ -457,7 +697,6 @@ void SDMXMVRFixtureList::RefreshList()
 	SaveHeaderRowSettings();
 
 	// Clear cached data
-	ListSource.Reset();
 	Rows.Reset();
 
 	TSharedPtr<FDMXEditor> DMXEditor = WeakDMXEditor.Pin();
@@ -472,13 +711,37 @@ void SDMXMVRFixtureList::RefreshList()
 		return;
 	}
 
+	DMXLibrary->UpdateGeneralSceneDescription();
+	UDMXMVRGeneralSceneDescription* GeneralSceneDescription = DMXLibrary->GetLazyGeneralSceneDescription();
+	if (!GeneralSceneDescription)
+	{
+		return;
+	}
+
 	const TArray<UDMXEntityFixturePatch*> FixturePatches = DMXLibrary->GetEntitiesTypeCast<UDMXEntityFixturePatch>();
+	ListSource.RemoveAll([&FixturePatches](const TSharedPtr<FDMXMVRFixtureListItem>& Item)
+		{
+			return !FixturePatches.Contains(Item->GetFixturePatch());
+		});
+	
 	for (UDMXEntityFixturePatch* FixturePatch : FixturePatches)
 	{
-		const FGuid& MVRFixtureUUID = FixturePatch->GetMVRFixtureUUID();
-		const TSharedRef<FDMXMVRFixtureListItem> NewItem = MakeShared<FDMXMVRFixtureListItem>(DMXEditor.ToSharedRef(), MVRFixtureUUID);
+		bool bExisting = ListSource.ContainsByPredicate([FixturePatch](const TSharedPtr<FDMXMVRFixtureListItem>& Item)
+			{
+				return FixturePatch == Item->GetFixturePatch();
+			});
+
+		if (!bExisting)
+		{
+			UDMXMVRFixtureNode* MVRFixtureNode = FindMVRFixtureNode(GeneralSceneDescription, FixturePatch);
+			if (!MVRFixtureNode)
+			{
+				continue;
+			}
+			const TSharedRef<FDMXMVRFixtureListItem> NewItem = MakeShared<FDMXMVRFixtureListItem>(DMXEditor.ToSharedRef(), *MVRFixtureNode);
 
-		ListSource.Add(NewItem);
+			ListSource.Add(NewItem);
+		}
 	}
 
 	// Generate status texts
@@ -491,7 +754,7 @@ void SDMXMVRFixtureList::RefreshList()
 	}
 	
 	// Update and sort the list and its widgets
-	ListView->RequestListRefresh();
+	ListView->RebuildList();
 
 	SortByColumnID(EColumnSortPriority::Max, FDMXMVRFixtureListCollumnIDs::Patch, EColumnSortMode::Ascending);
 
@@ -538,57 +801,48 @@ void SDMXMVRFixtureList::OnSelectionChanged(TSharedPtr<FDMXMVRFixtureListItem> I
 	}
  
 	const TArray<TSharedPtr<FDMXMVRFixtureListItem>> SelectedItems = ListView->GetSelectedItems();
-	if (SelectedItems.Num() > 0)
+	TArray<TWeakObjectPtr<UDMXEntityFixturePatch>> FixturePatchesToSelect;
+	for (const TSharedPtr<FDMXMVRFixtureListItem>& Item : SelectedItems)
 	{
-		TArray<TWeakObjectPtr<UDMXEntityFixturePatch>> FixturePatchesToSelect;
-		for (const TSharedPtr<FDMXMVRFixtureListItem>& Item : SelectedItems)
+		if (UDMXEntityFixturePatch* FixturePatch = Item->GetFixturePatch())
 		{
-			if (UDMXEntityFixturePatch* FixturePatch = Item->GetFixturePatch())
-			{
-				FixturePatchesToSelect.AddUnique(FixturePatch);
-			}
+			FixturePatchesToSelect.AddUnique(FixturePatch);
 		}
+	}
 
-		FixturePatchSharedData->SelectFixturePatches(FixturePatchesToSelect);
+	FixturePatchSharedData->SelectFixturePatches(FixturePatchesToSelect);
 
-		if (SelectedItems.Num() > 0)
+	if (FixturePatchesToSelect.Num() > 0)
+	{
+		const int32 SelectedUniverse = FixturePatchSharedData->GetSelectedUniverse();
+		const int32 UniverseOfFirstItem = FixturePatchesToSelect[0]->GetUniverseID();
+		if (SelectedUniverse != UniverseOfFirstItem)
 		{
-			const int32 SelectedUniverse = FixturePatchSharedData->GetSelectedUniverse();
-			const int32 UniverseOfFirstItem = SelectedItems[0]->GetUniverse();
-			if (SelectedUniverse != UniverseOfFirstItem)
-			{
-				FixturePatchSharedData->SelectUniverse(UniverseOfFirstItem);
-			}
+			FixturePatchSharedData->SelectUniverse(UniverseOfFirstItem);
 		}
 	}
-	else
+}
+
+void SDMXMVRFixtureList::OnEntityAddedOrRemoved(UDMXLibrary* DMXLibrary, TArray<UDMXEntity*> Entities)
+{
+	if (bChangingDMXLibrary)
 	{
-		// Restore selection when nothing got selected
-		const TArray<TWeakObjectPtr<UDMXEntityFixturePatch>> SelectedFixturePatches = FixturePatchSharedData->GetSelectedFixturePatches();
-		TArray<TSharedPtr<FDMXMVRFixtureListItem>> ItemsToSelect;
-		for (TWeakObjectPtr<UDMXEntityFixturePatch> FixturePatch : SelectedFixturePatches)
-		{
-			if (FixturePatch.IsValid())
-			{
-				const TSharedPtr<FDMXMVRFixtureListItem>* ItemPtr = ListSource.FindByPredicate([&FixturePatch](const TSharedPtr<FDMXMVRFixtureListItem>& Item)
-					{
-						return Item->GetFixturePatch() == FixturePatch;
-					});
-				if (ItemPtr)
-				{
-					ItemsToSelect.Add(*ItemPtr);
-				}
-			}
-		}
-		ListView->ClearSelection();
-		constexpr bool bSelected = true;
-		ListView->SetItemSelection(ItemsToSelect, bSelected);
+		return;
 	}
+
+	RequestListRefresh();
 }
 
 void SDMXMVRFixtureList::OnFixturePatchChanged(const UDMXEntityFixturePatch* FixturePatch)
 {
-	if (!bChangingDMXLibrary)
+	if (bChangingDMXLibrary)
+	{
+		return;
+	}
+
+	// Refresh only if the fixture patch is in the library this editor handles
+	const UDMXLibrary* DMXLibrary = WeakDMXEditor.IsValid() ? WeakDMXEditor.Pin()->GetDMXLibrary() : nullptr;
+	if (FixturePatch && FixturePatch->GetParentLibrary() == DMXLibrary)
 	{
 		RequestListRefresh();
 	}
@@ -596,7 +850,14 @@ void SDMXMVRFixtureList::OnFixturePatchChanged(const UDMXEntityFixturePatch* Fix
 
 void SDMXMVRFixtureList::OnFixtureTypeChanged(const UDMXEntityFixtureType* FixtureType)
 {
-	if (!bChangingDMXLibrary)
+	if (bChangingDMXLibrary)
+	{
+		return;
+	}
+
+	// Refresh only if the fixture type is in the library this editor handles
+	const UDMXLibrary* DMXLibrary = WeakDMXEditor.IsValid() ? WeakDMXEditor.Pin()->GetDMXLibrary() : nullptr;
+	if (FixtureType && FixtureType->GetParentLibrary() == DMXLibrary)
 	{
 		RequestListRefresh();
 	}
@@ -624,46 +885,13 @@ void SDMXMVRFixtureList::OnFixturePatchSharedDataSelectedFixturePatches()
 		if (NewSelection.Num() > 0)
 		{
 			ListView->ClearSelection();
-			ListView->SetItemSelection(NewSelection, true, ESelectInfo::Direct);
-
-			if (!ListView->IsItemVisible(NewSelection[0]))
-			{
-				ListView->RequestScrollIntoView(NewSelection[0]);
-			}
+			ListView->SetItemSelection(NewSelection, true, ESelectInfo::OnMouseClick);
 		}
 		else
 		{
 			ListView->ClearSelection();
 		}
 	}
-
-	GEditor->GetTimerManager()->SetTimerForNextTick(FTimerDelegate::CreateSP(this, &SDMXMVRFixtureList::SetKeyboardFocus));
-}
-
-void SDMXMVRFixtureList::OnFixturePatchSharedDataSelectedUniverse()
-{
-	if (!bChangingDMXLibrary)
-	{
-		const int32 SelectedUniverse = FixturePatchSharedData->GetSelectedUniverse();
-		const TArray<TSharedPtr<FDMXMVRFixtureListItem>> SelectedItems = ListView->GetSelectedItems();
-		const bool bUniverseAlreadySelected = SelectedItems.ContainsByPredicate([SelectedUniverse](const TSharedPtr<FDMXMVRFixtureListItem>& Item)
-			{
-				return Item->GetUniverse() == SelectedUniverse;
-			});
-		if (bUniverseAlreadySelected)
-		{
-			return;
-		}
-
-		const TSharedPtr<FDMXMVRFixtureListItem>* ItemPtr = ListSource.FindByPredicate([SelectedUniverse](const TSharedPtr<FDMXMVRFixtureListItem>& Item)
-			{
-				return Item->GetUniverse() == SelectedUniverse;
-			});
-		if (ItemPtr)
-		{
-			ListView->RequestScrollIntoView(*ItemPtr);
-		}
-	}
 }
 
 void SDMXMVRFixtureList::AdoptSelectionFromFixturePatchSharedData()
@@ -952,6 +1180,17 @@ void SDMXMVRFixtureList::SortByColumnID(const EColumnSortPriority::Type SortPrio
 	ListView->RequestListRefresh();
 }
 
+UDMXMVRFixtureNode* SDMXMVRFixtureList::FindMVRFixtureNode(UDMXMVRGeneralSceneDescription* GeneralSceneDescription, UDMXEntityFixturePatch* FixturePatch) const
+{
+	if (!GeneralSceneDescription || !FixturePatch)
+	{
+		return nullptr;
+	}
+
+	const FGuid& MVRFixtureUUID  = FixturePatch->GetMVRFixtureUUID();
+	return GeneralSceneDescription->FindFixtureNode(MVRFixtureUUID);
+}
+
 TSharedPtr<SWidget> SDMXMVRFixtureList::OnContextMenuOpening()
 {
 	if (!ListView.IsValid())
@@ -962,35 +1201,32 @@ TSharedPtr<SWidget> SDMXMVRFixtureList::OnContextMenuOpening()
 	const bool bCloseWindowAfterMenuSelection = true;
 	FMenuBuilder MenuBuilder(bCloseWindowAfterMenuSelection, CommandList);
 
-	if (ListView->GetNumItemsSelected() > 0)
+	// Auto Assign Section
+	MenuBuilder.BeginSection("AutoAssignSection", LOCTEXT("AutoAssignSection", "Auto-Assign"));
 	{
-		// Auto Assign Section
-		MenuBuilder.BeginSection("AutoAssignSection", LOCTEXT("AutoAssignSection", "Auto-Assign"));
-		{
-			// Auto Assign Entry
-			const FUIAction Action(FExecuteAction::CreateSP(this, &SDMXMVRFixtureList::AutoAssignFixturePatches));
-
-			const FText AutoAssignText = LOCTEXT("AutoAssignContextMenuEntry", "Auto-Assign Selection");
-			const TSharedRef<SWidget> Widget =
-				SNew(STextBlock)
-				.Text(AutoAssignText);
+		// Auto Assign Entry
+		const FUIAction Action(FExecuteAction::CreateSP(this, &SDMXMVRFixtureList::AutoAssignFixturePatches));
 
-			MenuBuilder.AddMenuEntry(Action, Widget);
-			MenuBuilder.EndSection();
-		}
+		const FText AutoAssignText = LOCTEXT("AutoAssignContextMenuEntry", "Auto-Assign Selection");
+		const TSharedRef<SWidget> Widget =
+			SNew(STextBlock)
+			.Text(AutoAssignText);
 
-		// Basic Operations Section
-		MenuBuilder.BeginSection("BasicOperationsSection", LOCTEXT("BasicOperationsSection", "Basic Operations"));
-		{
-			MenuBuilder.AddMenuEntry(FGenericCommands::Get().Cut);
-			MenuBuilder.AddMenuEntry(FGenericCommands::Get().Copy);
-			MenuBuilder.AddMenuEntry(FGenericCommands::Get().Paste);
-			MenuBuilder.AddMenuEntry(FGenericCommands::Get().Duplicate);
-			MenuBuilder.AddMenuEntry(FGenericCommands::Get().Delete);
-		}
+		MenuBuilder.AddMenuEntry(Action, Widget);
 		MenuBuilder.EndSection();
 	}
 
+	// Basic Operations Section
+	MenuBuilder.BeginSection("BasicOperationsSection", LOCTEXT("BasicOperationsSection", "Basic Operations"));
+	{
+		MenuBuilder.AddMenuEntry(FGenericCommands::Get().Cut);
+		MenuBuilder.AddMenuEntry(FGenericCommands::Get().Copy);
+		MenuBuilder.AddMenuEntry(FGenericCommands::Get().Paste);
+		MenuBuilder.AddMenuEntry(FGenericCommands::Get().Duplicate);
+		MenuBuilder.AddMenuEntry(FGenericCommands::Get().Delete);
+	}
+	MenuBuilder.EndSection();
+
 	return MenuBuilder.MakeWidget();
 }
 
@@ -1038,7 +1274,7 @@ void SDMXMVRFixtureList::RegisterCommands()
 
 bool SDMXMVRFixtureList::CanCutItems() const
 {
-	return CanCopyItems() && CanDeleteItems();
+	return CanCopyItems() && CanDeleteItems() && !GIsTransacting;
 }
 
 void SDMXMVRFixtureList::OnCutSelectedItems()
@@ -1053,96 +1289,109 @@ void SDMXMVRFixtureList::OnCutSelectedItems()
 
 bool SDMXMVRFixtureList::CanCopyItems() const
 {
-	const TArray<TSharedPtr<FDMXMVRFixtureListItem>> SelectedItems = ListView->GetSelectedItems();
-	return SelectedItems.Num() > 0;
+	return FixturePatchSharedData->GetSelectedFixturePatches().Num() > 0 && !GIsTransacting;
 }
 
 void SDMXMVRFixtureList::OnCopySelectedItems()
 {
 	const TArray<TSharedPtr<FDMXMVRFixtureListItem>> SelectedItems = ListView->GetSelectedItems();
-
-	FString SelectedMVRUUIDsString(ClipboardCopyMVRFixtureHeader);
+	TArray<UDMXEntityFixturePatch*> FixturePatchesToCopy;
 	for (const TSharedPtr<FDMXMVRFixtureListItem>& Item : SelectedItems)
 	{
-		SelectedMVRUUIDsString += Item->GetMVRUUID().ToString() + FString(LINE_TERMINATOR);
+		FixturePatchesToCopy.Add(Item->GetFixturePatch());
 	}
 
-	FPlatformApplicationMisc::ClipboardCopy(*SelectedMVRUUIDsString);
+	using namespace UE::DMX::SDMXMVRFixtureList::Private;
+	ClipboardCopyFixturePatches(FixturePatchesToCopy);
 }
 
 bool SDMXMVRFixtureList::CanPasteItems() const
 {
-	const TArray<TSharedPtr<FDMXMVRFixtureListItem>> SelectedItems = ListView->GetSelectedItems();
-	return SelectedItems.Num() == 1;
+	using namespace UE::DMX::SDMXMVRFixtureList::Private;
+
+	UDMXLibrary* DMXLibrary = WeakDMXEditor.IsValid() ? WeakDMXEditor.Pin()->GetDMXLibrary() : nullptr;
+	if (!DMXLibrary)
+	{
+		return false;
+	}
+
+	// Get the text from the clipboard
+	FString TextToImport;
+	FPlatformApplicationMisc::ClipboardPaste(TextToImport);
+
+	return FDMXFixturePatchObjectTextFactory::CanCreate(TextToImport, DMXLibrary) && !GIsTransacting;
 }
 
 void SDMXMVRFixtureList::OnPasteItems()
 {
-	TGuardValue<bool>(bChangingDMXLibrary, true);
-
-	FString ClipboardString;
-	FPlatformApplicationMisc::ClipboardPaste(ClipboardString);
+	using namespace UE::DMX::SDMXMVRFixtureList::Private;
 
-	if (ClipboardString.RemoveFromStart(ClipboardCopyMVRFixtureHeader, ESearchCase::CaseSensitive))
+	UDMXLibrary* DMXLibrary = WeakDMXEditor.IsValid() ? WeakDMXEditor.Pin()->GetDMXLibrary() : nullptr;
+	if (!DMXLibrary)
 	{
-		TArray<FString> UUIDStrings;
-		ClipboardString.ParseIntoArray(UUIDStrings, LINE_TERMINATOR);
-
-		TArray<FGuid> UUIDs;
-		for (const FString& UUIDString : UUIDStrings)
-		{
-			FGuid UUID;
-			if (FGuid::Parse(UUIDString, UUID))
-			{
-				UUIDs.Add(UUID);
-			}
-		}
+		return;
+	}
 
-		TArray<TSharedPtr<FDMXMVRFixtureListItem>> ItemsToPaste;
-		for (const FGuid& UUID : UUIDs)
-		{
-			const TSharedPtr<FDMXMVRFixtureListItem>* ItemPtr = ListSource.FindByPredicate([UUID](const TSharedPtr<FDMXMVRFixtureListItem>& Item)
-				{
-					return UUID == Item->GetMVRUUID();
-				});
+	TGuardValue<bool>(bChangingDMXLibrary, true);
+	const FText TransactionText = LOCTEXT("PasteFixturePatchesTransaction", "Paste Fixture Patches");
+	const FScopedTransaction PasteTransaction(TransactionText);
 
-			if (ItemPtr)
-			{
-				ItemsToPaste.Add(*ItemPtr);
-			}
-		}
+	FString TextToImport;
+	FPlatformApplicationMisc::ClipboardPaste(TextToImport);
 
-		const TArray<TSharedPtr<FDMXMVRFixtureListItem>> SelectedItems = ListView->GetSelectedItems();
-		if (SelectedItems.Num() == 1) // Should always be one given CanPasteItems
+	TArray<UDMXEntityFixturePatch*> PastedFixturePatches;
+	if(FDMXFixturePatchObjectTextFactory::Create(TextToImport, DMXLibrary, PastedFixturePatches))
+	{
+		TArray<TWeakObjectPtr<UDMXEntityFixturePatch>> WeakPastedFixturePatches;
+		for (UDMXEntityFixturePatch* FixturePatch : PastedFixturePatches)
 		{
-			FDMXMVRFixtureListItem::PasteItemsOntoItem(WeakDMXEditor, SelectedItems[0], ItemsToPaste);
+			WeakPastedFixturePatches.Add(FixturePatch);
 		}
 
+		FixturePatchSharedData->SelectFixturePatches(WeakPastedFixturePatches);
+
 		RequestListRefresh();
 	}
 }
 
 bool SDMXMVRFixtureList::CanDuplicateItems() const
 {
-	const TArray<TSharedPtr<FDMXMVRFixtureListItem>> SelectedItems = ListView->GetSelectedItems();
-	return SelectedItems.Num() > 0;
+	return FixturePatchSharedData->GetSelectedFixturePatches().Num() > 0 && !GIsTransacting;
 }
 
 void SDMXMVRFixtureList::OnDuplicateItems()
 {
+	UDMXLibrary* DMXLibrary = WeakDMXEditor.IsValid() ? WeakDMXEditor.Pin()->GetDMXLibrary() : nullptr;
+	if (!DMXLibrary)
+	{
+		return;
+	}
+
 	TGuardValue<bool>(bChangingDMXLibrary, true);
+	const FText TransactionText = LOCTEXT("DuplicateFixturePatchesTransaction", "Duplicate Fixture Patches");
+	const FScopedTransaction PasteTransaction(TransactionText);
 
-	const TArray<TSharedPtr<FDMXMVRFixtureListItem>> SelectedItems = ListView->GetSelectedItems();
+	const TArray<TWeakObjectPtr<UDMXEntityFixturePatch>> SelectedFixturePatches = FixturePatchSharedData->GetSelectedFixturePatches();
+	TArray<TWeakObjectPtr<UDMXEntityFixturePatch>> NewFixturePatches;
+	for (const TWeakObjectPtr<UDMXEntityFixturePatch> FixturePatch : SelectedFixturePatches)
+	{
+		if(!FixturePatch.IsValid() || !FixturePatch->GetParentLibrary())
+		{
+			continue;
+		}
 
-	FDMXMVRFixtureListItem::DuplicateItems(WeakDMXEditor, SelectedItems);
+		using namespace UE::DMX::SDMXMVRFixtureList::Private;
+		UDMXEntityFixturePatch* NewFixturePatch = DuplicatePatchByMVRFixtureUUID(DMXLibrary, FixturePatch->GetMVRFixtureUUID());
+		NewFixturePatches.Add(NewFixturePatch);
+	}
+	FixturePatchSharedData->SelectFixturePatches(NewFixturePatches);
 
 	RequestListRefresh();
 }
 
 bool SDMXMVRFixtureList::CanDeleteItems() const
 {
-	const TArray<TSharedPtr<FDMXMVRFixtureListItem>> SelectedItems = ListView->GetSelectedItems();
-	return SelectedItems.Num() > 0;
+	return FixturePatchSharedData->GetSelectedFixturePatches().Num() > 0 && !GIsTransacting;
 }
 
 void SDMXMVRFixtureList::OnDeleteItems()
@@ -1151,6 +1400,34 @@ void SDMXMVRFixtureList::OnDeleteItems()
 
 	const TArray<TSharedPtr<FDMXMVRFixtureListItem>> SelectedItems = ListView->GetSelectedItems();
 
+	if (SelectedItems.Num() == 0)
+	{
+		return;
+	}
+	
+	// Its safe to assume all patches are in the same Library - A Multi-Library Editor wouldn't make sense.
+	UDMXLibrary* DMXLibrary = SelectedItems[0]->GetDMXLibrary();
+	if (!DMXLibrary)
+	{
+		return;
+	}
+	
+	const FText DeleteFixturePatchesTransactionText = FText::Format(LOCTEXT("DeleteFixturePatchesTransaction", "Delete Fixture {0}|plural(one=Patch, other=Patches)"), SelectedItems.Num() > 1);
+	const FScopedTransaction DeleteFixturePatchTransaction(DeleteFixturePatchesTransactionText);
+	DMXLibrary->PreEditChange(nullptr);
+
+	for (const TSharedPtr<FDMXMVRFixtureListItem>& Item : SelectedItems)
+	{
+		if (UDMXEntityFixturePatch* FixturePatch = Item->GetFixturePatch())
+		{
+			FixturePatch->PreEditChange(UDMXEntityFixturePatch::StaticClass()->FindPropertyByName(UDMXEntityFixturePatch::GetMVRFixtureUUIDPropertyNameChecked()));
+			UDMXEntityFixturePatch::RemoveFixturePatchFromLibrary(FixturePatch);
+			FixturePatch->PostEditChange();
+		}
+	}
+
+	DMXLibrary->PostEditChange();
+
 	// Make a meaningful selection invariant to ordering of the List
 	TSharedPtr<FDMXMVRFixtureListItem> NewSelection;
 	for (int32 ItemIndex = 0; ItemIndex < ListSource.Num(); ItemIndex++)
@@ -1174,7 +1451,6 @@ void SDMXMVRFixtureList::OnDeleteItems()
 		ListView->SetSelection(NewSelection, ESelectInfo::OnMouseClick);
 	}
 
-	FDMXMVRFixtureListItem::DeleteItems(SelectedItems);
 	RequestListRefresh();
 }
 
diff --git a/Engine/Plugins/VirtualProduction/DMX/DMXEngine/Source/DMXEditor/Private/Widgets/FixturePatch/SDMXMVRFixtureList.h b/Engine/Plugins/VirtualProduction/DMX/DMXEngine/Source/DMXEditor/Private/Widgets/FixturePatch/SDMXMVRFixtureList.h
--- a/Engine/Plugins/VirtualProduction/DMX/DMXEngine/Source/DMXEditor/Private/Widgets/FixturePatch/SDMXMVRFixtureList.h
+++ b/Engine/Plugins/VirtualProduction/DMX/DMXEngine/Source/DMXEditor/Private/Widgets/FixturePatch/SDMXMVRFixtureList.h
@@ -13,8 +13,12 @@ class FDMXMVRFixtureListItem;
 class FDMXFixturePatchSharedData;
 class SDMXMVRFixtureListRow;
 class SDMXMVRFixtureListToolbar;
+class UDMXEntity;
 class UDMXEntityFixturePatch;
 class UDMXEntityFixtureType;
+class UDMXLibrary;
+class UDMXMVRFixtureNode;
+class UDMXMVRGeneralSceneDescription;
 
 class FUICommandList;
 class ITableRow;
@@ -93,6 +97,9 @@ private:
 	/** Called when the Selection in the list changed */
 	void OnSelectionChanged(TSharedPtr<FDMXMVRFixtureListItem> InItem, ESelectInfo::Type SelectInfo);
 
+	/** Called when entities were added or removed from the DMX Library */
+	void OnEntityAddedOrRemoved(UDMXLibrary* DMXLibrary, TArray<UDMXEntity*> Entities);
+
 	/** Called when a Fixture Patch changed */
 	void OnFixturePatchChanged(const UDMXEntityFixturePatch* FixturePatch);
 
@@ -102,9 +109,6 @@ private:
 	/** Called when Fixture Patch Shared Data selected Fixture Patches */
 	void OnFixturePatchSharedDataSelectedFixturePatches();
 
-	/** Called when Fixture Patch Shared Data selected a Universe */
-	void OnFixturePatchSharedDataSelectedUniverse();
-
 	/** Selects what is selected in Fixture Patch Shared Data */
 	void AdoptSelectionFromFixturePatchSharedData();
 
@@ -129,6 +133,9 @@ private:
 	/** Sort the list by ColumnId */
 	void SortByColumnID(const EColumnSortPriority::Type SortPriority, const FName& ColumnId, const EColumnSortMode::Type InSortMode);
 
+	/** Returns the MVR Fixture corresponding to the fixture patch */
+	UDMXMVRFixtureNode* FindMVRFixtureNode(UDMXMVRGeneralSceneDescription* GeneralSceneDescription, UDMXEntityFixturePatch* FixturePatch) const;
+
 	/** The current Sort Mode */
 	EColumnSortMode::Type SortMode;
 
@@ -186,6 +193,4 @@ private:
 	void OnDeleteItems();
 
 	TSharedPtr<FUICommandList> CommandList;
-
-	static const FString ClipboardCopyMVRFixtureHeader;
 };
diff --git a/Engine/Plugins/VirtualProduction/DMX/DMXEngine/Source/DMXEditor/Private/Widgets/FixturePatch/SDMXMVRFixtureListRow.cpp b/Engine/Plugins/VirtualProduction/DMX/DMXEngine/Source/DMXEditor/Private/Widgets/FixturePatch/SDMXMVRFixtureListRow.cpp
--- a/Engine/Plugins/VirtualProduction/DMX/DMXEngine/Source/DMXEditor/Private/Widgets/FixturePatch/SDMXMVRFixtureListRow.cpp
+++ b/Engine/Plugins/VirtualProduction/DMX/DMXEngine/Source/DMXEditor/Private/Widgets/FixturePatch/SDMXMVRFixtureListRow.cpp
@@ -12,6 +12,7 @@
 #include "Widgets/FixturePatch/DMXMVRFixtureListItem.h"
 
 #include "SSearchableComboBox.h"
+#include "Engine/EngineTypes.h"
 #include "Widgets/Text/SInlineEditableTextBlock.h"
 #include "Widgets/Input/STextComboBox.h"
 #include "Styling/AppStyle.h"
@@ -65,7 +66,7 @@ public:
 				]
 			];
 
-		Refresh();
+		RefreshInternal();
 	}
 
 	/** Sets the currently selected Fixture Type */
@@ -89,6 +90,15 @@ public:
 		}
 	}
 
+	/** Requests to refresh the widget on the next tick */
+	void RequestRefresh()
+	{
+		if (!RefreshTimerHandle.IsValid())
+		{
+			RefreshTimerHandle = GEditor->GetTimerManager()->SetTimerForNextTick(FTimerDelegate::CreateSP(this, &SDMXMVRFixtureFixtureTypePicker::RefreshInternal));
+		}
+	}
+
 private:
 	/** Called when the combo box generates a widget */
 	TSharedRef<SWidget> OnGenerateWidget(TSharedPtr<FString> InItem)
@@ -101,6 +111,11 @@ private:
 	/** Called when the combo box selection changed */
 	void OnSelectionChanged(TSharedPtr<FString> NewSelection, ESelectInfo::Type SelectInfo)
 	{
+		if (SelectInfo == ESelectInfo::Direct)
+		{
+			return;
+		}
+
 		if (UDMXEntityFixtureType* SelectedFixtureType = GetSelectedFixtureType())
 		{
 			OnFixtureTypeSelectedDelegate.ExecuteIfBound(SelectedFixtureType);
@@ -134,14 +149,27 @@ private:
 	}
 
 	/** Refreshes the Combo Box */
-	void Refresh()
+	void RefreshInternal()
 	{
+		RefreshTimerHandle.Invalidate();
+
 		FixtureTypeNames.Reset();
+		const UDMXLibrary* DMXLibrary = WeakDMXLibrary.Get();
+		if (!DMXLibrary)
+		{
+			return;
+		}
 
-		FixtureTypeNames.Add(MakeShared<FString>(TEXT("None")));
-		if (const UDMXLibrary* DMXLibrary = WeakDMXLibrary.Get())
+		const TArray<UDMXEntityFixtureType*> FixtureTypes = DMXLibrary->GetEntitiesTypeCast<UDMXEntityFixtureType>();
+
+		if (FixtureTypes.IsEmpty())
+		{
+			const TSharedRef<FString> NoneItem = MakeShared<FString>(TEXT("None"));
+			FixtureTypeNames.Add(NoneItem);
+			ComboBox->SetSelectedItem(NoneItem);
+		}
+		else
 		{
-			const TArray<UDMXEntityFixtureType*> FixtureTypes = DMXLibrary->GetEntitiesTypeCast<UDMXEntityFixtureType>();
 			for (UDMXEntityFixtureType* FixtureType : FixtureTypes)
 			{
 				const TSharedPtr<FString> FixtureTypeName = MakeShared<FString>(FixtureType->Name);
@@ -151,19 +179,40 @@ private:
 			}
 		}
 
+		if (!ensureMsgf(!FixtureTypeNames.IsEmpty(), TEXT("No combo box option for Fixture Type.")))
+		{
+			return;
+		}
+
+		const FString SelectedFixtureTypeName = ComboBox->GetSelectedItem() ? *ComboBox->GetSelectedItem() : TEXT("");
+		const TSharedPtr<FString>* SelectionPtr = Algo::FindByPredicate(FixtureTypeNames, [SelectedFixtureTypeName](const TSharedPtr<FString>& FixtureTypeName)
+			{
+				return *FixtureTypeName == SelectedFixtureTypeName;
+			});
+		if (!SelectionPtr)
+		{
+			ComboBox->SetSelectedItem(FixtureTypeNames[0], ESelectInfo::Direct);
+		}
+
 		ComboBox->RefreshOptions();
 	}
 
 	/** Called when the Entities in the DMX Library were added or removed */
 	void OnEntitiesAddedOrRemoved(UDMXLibrary* DMXLibrary, TArray<UDMXEntity*> AddedOrRemovedEntities)
 	{
-		Refresh();
+		if (DMXLibrary == WeakDMXLibrary)
+		{
+			RequestRefresh();
+		}
 	}
 
 	/** Called when a Fixture Type changed */
 	void OnFixtureTypeChanged(const UDMXEntityFixtureType* ChangedFixtureType)
 	{
-		Refresh();
+		if (ChangedFixtureType && ChangedFixtureType->GetParentLibrary() == WeakDMXLibrary)
+		{
+			RequestRefresh();
+		}
 	}
 
 	/** Names of Fixture Types in the Combo Box */
@@ -178,6 +227,9 @@ private:
 	/** The DMX Library from which to select fixture types */
 	TWeakObjectPtr<UDMXLibrary> WeakDMXLibrary;
 
+	/** Timer handle used for the RequestRefresh method */
+	FTimerHandle RefreshTimerHandle;
+
 	// Slate args
 	FDMXMVRFixtureListRowOnFixtureTypeSelectedDelegate OnFixtureTypeSelectedDelegate;
 };
@@ -282,6 +334,9 @@ private:
 	/** Refreshes the Combo Box */
 	void RefreshOptions()
 	{
+		// Remember selection
+		const FString PreviousSelectedModeName = ComboBox->GetSelectedItem().IsValid() ? *ComboBox->GetSelectedItem() : FString();
+
 		ModeNames.Reset();
 
 		if (UDMXEntityFixtureType* FixtureType = WeakFixtureType.Get())
@@ -308,12 +363,25 @@ private:
 		}
 
 		ComboBox->RefreshOptions();
+
+		// Restore selection
+		TSharedPtr<FString>* PreviousSelectionPtr = Algo::FindByPredicate(ModeNames, [PreviousSelectedModeName](const TSharedPtr<FString>& ModeName)
+			{
+				return ModeName.IsValid() && *ModeName == PreviousSelectedModeName;
+			});
+		if (PreviousSelectionPtr)
+		{
+			ComboBox->SetSelectedItem(*PreviousSelectionPtr);
+		}
 	}
 
 	/** Called when the Entities in the DMX Library were added or removed */
 	void OnFixtureTypeChanged(const UDMXEntityFixtureType* ChangedFixtureType)
 	{
-		RefreshOptions();
+		if (ChangedFixtureType == WeakFixtureType)
+		{
+			RefreshOptions();
+		}
 	}
 
 	/** True if the Combo Box should be enabled */
diff --git a/Engine/Plugins/VirtualProduction/DMX/DMXEngine/Source/DMXEditor/Private/Widgets/FixturePatch/SDMXMVRFixtureListToolbar.cpp b/Engine/Plugins/VirtualProduction/DMX/DMXEngine/Source/DMXEditor/Private/Widgets/FixturePatch/SDMXMVRFixtureListToolbar.cpp
--- a/Engine/Plugins/VirtualProduction/DMX/DMXEngine/Source/DMXEditor/Private/Widgets/FixturePatch/SDMXMVRFixtureListToolbar.cpp
+++ b/Engine/Plugins/VirtualProduction/DMX/DMXEngine/Source/DMXEditor/Private/Widgets/FixturePatch/SDMXMVRFixtureListToolbar.cpp
@@ -10,6 +10,7 @@
 #include "Widgets/SDMXEntityDropdownMenu.h"
 
 #include "EditorStyleSet.h"
+#include "ScopedTransaction.h"
 #include "Internationalization/Regex.h"
 #include "Widgets/SBoxPanel.h"
 #include "Widgets/Input/SCheckBox.h"
@@ -360,6 +361,7 @@ void SDMXMVRFixtureListToolbar::OnAddNewMVRFixtureClicked(UDMXEntity* InSelected
 	FixturePatchConstructionParams.UniverseID = Universe;
 	FixturePatchConstructionParams.StartingAddress = Address;
 
+	const FScopedTransaction CreateFixturePatchTransaction(LOCTEXT("CreateFixturePatchTransaction", "Create Fixture Patch"));
 	constexpr bool bMarkLibraryDirty = true;
 	UDMXEntityFixturePatch* NewFixturePatch = UDMXEntityFixturePatch::CreateFixturePatchInLibrary(FixturePatchConstructionParams, FixtureType->Name, bMarkLibraryDirty);
 }
diff --git a/Engine/Plugins/VirtualProduction/DMX/DMXEngine/Source/DMXEditor/Private/Widgets/FixturePatch/SDMXPatchedUniverse.cpp b/Engine/Plugins/VirtualProduction/DMX/DMXEngine/Source/DMXEditor/Private/Widgets/FixturePatch/SDMXPatchedUniverse.cpp
--- a/Engine/Plugins/VirtualProduction/DMX/DMXEngine/Source/DMXEditor/Private/Widgets/FixturePatch/SDMXPatchedUniverse.cpp
+++ b/Engine/Plugins/VirtualProduction/DMX/DMXEngine/Source/DMXEditor/Private/Widgets/FixturePatch/SDMXPatchedUniverse.cpp
@@ -114,7 +114,10 @@ END_SLATE_FUNCTION_BUILD_OPTIMIZATION
 
 void SDMXPatchedUniverse::RequestRefresh()
 {
-	GEditor->GetTimerManager()->SetTimerForNextTick(FTimerDelegate::CreateSP(this, &SDMXPatchedUniverse::RefreshInternal));
+	if (!RequestRefreshTimerHandle.IsValid())
+	{
+		RequestRefreshTimerHandle = GEditor->GetTimerManager()->SetTimerForNextTick(FTimerDelegate::CreateSP(this, &SDMXPatchedUniverse::RefreshInternal));
+	}
 }
 
 void SDMXPatchedUniverse::CreateChannelConnectors()
@@ -308,13 +311,13 @@ void SDMXPatchedUniverse::Unpatch(const TSharedPtr<FDMXFixturePatchNode>& Node)
 				Grid->RemoveSlot(Widget.ToSharedRef());
 			}
 		}
-		RefreshInternal();
+		RequestRefresh();
 	}
 }
 
 void SDMXPatchedUniverse::OnFixturePatchChanged(const UDMXEntityFixturePatch* FixturePatch)
 {
-	if (!FixturePatch)
+	if (!FixturePatch || FixturePatch->GetParentLibrary() != GetDMXLibrary())
 	{
 		return;
 	}
@@ -391,12 +394,14 @@ void SDMXPatchedUniverse::SetUniverseIDInternal(int32 NewUniverseID)
 			Connector->SetUniverseID(NewUniverseID);
 		}
 
-		RefreshInternal();
+		RequestRefresh();
 	}
 }
 
 void SDMXPatchedUniverse::RefreshInternal()
 {
+	RequestRefreshTimerHandle.Invalidate();
+
 	UpdateZOrderOfNodes();
 
 	// Remove old widgets from the grid
diff --git a/Engine/Plugins/VirtualProduction/DMX/DMXEngine/Source/DMXEditor/Private/Widgets/FixturePatch/SDMXPatchedUniverse.h b/Engine/Plugins/VirtualProduction/DMX/DMXEngine/Source/DMXEditor/Private/Widgets/FixturePatch/SDMXPatchedUniverse.h
--- a/Engine/Plugins/VirtualProduction/DMX/DMXEngine/Source/DMXEditor/Private/Widgets/FixturePatch/SDMXPatchedUniverse.h
+++ b/Engine/Plugins/VirtualProduction/DMX/DMXEngine/Source/DMXEditor/Private/Widgets/FixturePatch/SDMXPatchedUniverse.h
@@ -205,6 +205,9 @@ private:
 	/** The Fixture Patch Widgets that are currently being displalyed */
 	TArray<TSharedPtr<SDMXFixturePatchFragment>> FixturePatchWidgets;
 
+	/** Timer handle for the Request Refresh method */
+	FTimerHandle RequestRefreshTimerHandle;
+
 	/** Shared data for fixture patch editors */
 	TSharedPtr<FDMXFixturePatchSharedData> SharedData;
 
diff --git a/Engine/Plugins/VirtualProduction/DMX/DMXEngine/Source/DMXRuntime/Private/Library/DMXEntityFixturePatch.cpp b/Engine/Plugins/VirtualProduction/DMX/DMXEngine/Source/DMXRuntime/Private/Library/DMXEntityFixturePatch.cpp
--- a/Engine/Plugins/VirtualProduction/DMX/DMXEngine/Source/DMXRuntime/Private/Library/DMXEntityFixturePatch.cpp
+++ b/Engine/Plugins/VirtualProduction/DMX/DMXEngine/Source/DMXRuntime/Private/Library/DMXEntityFixturePatch.cpp
@@ -43,7 +43,7 @@ UDMXEntityFixturePatch* UDMXEntityFixturePatch::CreateFixturePatchInLibrary(FDMX
 {
 	UDMXEntityFixtureType* FixtureType = ConstructionParams.FixtureTypeRef.GetFixtureType();
 
-	if (ensureMsgf(FixtureType, TEXT("Cannot create Fixture Patch when Fixture Type is invalid.")))
+	if (FixtureType)
 	{
 		UDMXLibrary* DMXLibrary = FixtureType->GetParentLibrary();
 		if (ensureMsgf(DMXLibrary, TEXT("Cannot create Fixture Patch when Fixture Type's DMX Library is invalid.")))
diff --git a/Engine/Plugins/VirtualProduction/DMX/DMXEngine/Source/DMXRuntime/Private/MVR/DMXMVRGeneralSceneDescription.cpp b/Engine/Plugins/VirtualProduction/DMX/DMXEngine/Source/DMXRuntime/Private/MVR/DMXMVRGeneralSceneDescription.cpp
--- a/Engine/Plugins/VirtualProduction/DMX/DMXEngine/Source/DMXRuntime/Private/MVR/DMXMVRGeneralSceneDescription.cpp
+++ b/Engine/Plugins/VirtualProduction/DMX/DMXEngine/Source/DMXRuntime/Private/MVR/DMXMVRGeneralSceneDescription.cpp
@@ -20,6 +20,7 @@
 #include "MVR/Types/DMXMVRSceneNode.h"
 
 #include "XmlFile.h"
+#include "Algo/RemoveIf.h"
 #include "EditorFramework/AssetImportData.h"
 #include "Misc/Paths.h"
 
@@ -74,6 +75,13 @@ UDMXMVRGeneralSceneDescription* UDMXMVRGeneralSceneDescription::CreateFromDMXLib
 }
 #endif // WITH_EDITOR
 
+
+// Hotfix UE5.1.1: To set MVR Fixture UUIDs faster, hold newly added nodes in the array here
+namespace UE::DMX::DMXMVRGeneralSceneDescription::HotFix::Private
+{
+	TArray<UDMXMVRFixtureNode*> NewNodes;
+}
+
 #if WITH_EDITOR
 void UDMXMVRGeneralSceneDescription::WriteDMXLibraryToGeneralSceneDescription(const UDMXLibrary& DMXLibrary)
 {
@@ -104,6 +112,7 @@ void UDMXMVRGeneralSceneDescription::WriteDMXLibraryToGeneralSceneDescription(co
 	}
 
 	// Create or update MVR Fixtures for each Fixture Patch's MVR Fixture UUIDs
+	UE::DMX::DMXMVRGeneralSceneDescription::HotFix::Private::NewNodes.Reset();
 	for (UDMXEntityFixturePatch* FixturePatch : FixturePatches)
 	{
 		if (FixturePatch)
@@ -111,6 +120,16 @@ void UDMXMVRGeneralSceneDescription::WriteDMXLibraryToGeneralSceneDescription(co
 			WriteFixturePatchToGeneralSceneDescription(*FixturePatch);
 		}
 	}
+
+	// Generate a Fixture ID for all new nodes. This is a 5.1.1 optimization to avoid the O1 lookup for each patch individually.
+	const TArray<int32> FixtureIDsInUse = GetNumericalFixtureIDsInUse(DMXLibrary);
+	int32 NextFixtureID = FixtureIDsInUse.Num() > 0 ? FixtureIDsInUse.Last() + 1 : 1;
+	
+	for (UDMXMVRFixtureNode* FixtureNode : UE::DMX::DMXMVRGeneralSceneDescription::HotFix::Private::NewNodes)
+	{
+		FixtureNode->FixtureID = FString::FromInt(NextFixtureID);
+		NextFixtureID++;
+	}
 }
 #endif // WITH_EDITOR
 
@@ -160,12 +179,10 @@ void UDMXMVRGeneralSceneDescription::WriteFixturePatchToGeneralSceneDescription(
 		UDMXMVRChildListNode& AnyChildList = RootNode->GetOrCreateFirstChildListNode();
 		MVRFixtureNode = AnyChildList.CreateParametricObject<UDMXMVRFixtureNode>();
 
-		const TArray<int32> FixtureIDsInUse = GetNumericalFixtureIDsInUse(*DMXLibrary);
-		const int32 FixtureID = FixtureIDsInUse.Num() > 0 ? FixtureIDsInUse.Last() + 1 : 1;
-
 		MVRFixtureNode->Name = FixturePatch.Name;
 		MVRFixtureNode->UUID = MVRFixtureUUID;
-		MVRFixtureNode->FixtureID = FString::FromInt(FixtureID);
+
+		UE::DMX::DMXMVRGeneralSceneDescription::HotFix::Private::NewNodes.Add(MVRFixtureNode);
 	}
 	check(MVRFixtureNode);
 
