DMX - Typing an attribute and selecting another from dropdown now updates UI correctly

diff --git a/Engine/Plugins/VirtualProduction/DMX/DMXEngine/Source/DMXEditor/Private/Widgets/SNameListPicker.cpp b/Engine/Plugins/VirtualProduction/DMX/DMXEngine/Source/DMXEditor/Private/Widgets/SNameListPicker.cpp
--- a/Engine/Plugins/VirtualProduction/DMX/DMXEngine/Source/DMXEditor/Private/Widgets/SNameListPicker.cpp
+++ b/Engine/Plugins/VirtualProduction/DMX/DMXEngine/Source/DMXEditor/Private/Widgets/SNameListPicker.cpp
@@ -152,6 +152,11 @@ void SNameListPicker::OnSelectionChanged(const TSharedPtr<FName> Item, ESelectIn
 		return;
 	}
 
+	if (Item.IsValid())
+	{
+		EdititableTextBox->SetText(FText::FromString(Item->ToString()));
+	}
+
 	OnUserSelectedItem(Item);
 }
 
