DMX - Fix selecting a mode in one DMX Library causes the mode being displayed in any other open DMX Library

diff --git a/Engine/Plugins/VirtualProduction/DMX/DMXEngine/Source/DMXEditor/Private/DMXFixturePatchSharedData.cpp b/Engine/Plugins/VirtualProduction/DMX/DMXEngine/Source/DMXEditor/Private/DMXFixturePatchSharedData.cpp
--- a/Engine/Plugins/VirtualProduction/DMX/DMXEngine/Source/DMXEditor/Private/DMXFixturePatchSharedData.cpp
+++ b/Engine/Plugins/VirtualProduction/DMX/DMXEngine/Source/DMXEditor/Private/DMXFixturePatchSharedData.cpp
@@ -12,7 +12,7 @@
 FDMXFixturePatchSharedData::FDMXFixturePatchSharedData(TWeakPtr<FDMXEditor> InDMXEditorPtr)
 	: DMXEditorPtr(InDMXEditorPtr)
 {
-	Selection = NewObject<UDMXFixturePatchSharedDataSelection>(GetTransientPackage(), "FixturePatchSharedDataSelection", RF_Transactional);
+	Selection = NewObject<UDMXFixturePatchSharedDataSelection>(GetTransientPackage(), NAME_None, RF_Transactional);
 }
 
 void FDMXFixturePatchSharedData::AddReferencedObjects(FReferenceCollector& Collector)
diff --git a/Engine/Plugins/VirtualProduction/DMX/DMXEngine/Source/DMXEditor/Private/DMXFixtureTypeSharedData.cpp b/Engine/Plugins/VirtualProduction/DMX/DMXEngine/Source/DMXEditor/Private/DMXFixtureTypeSharedData.cpp
--- a/Engine/Plugins/VirtualProduction/DMX/DMXEngine/Source/DMXEditor/Private/DMXFixtureTypeSharedData.cpp
+++ b/Engine/Plugins/VirtualProduction/DMX/DMXEngine/Source/DMXEditor/Private/DMXFixtureTypeSharedData.cpp
@@ -74,7 +74,7 @@ namespace
 
 FDMXFixtureTypeSharedData::FDMXFixtureTypeSharedData(TWeakPtr<FDMXEditor> InDMXEditorPtr)
 {
-	Selection = NewObject<UDMXFixtureTypeSharedDataSelection>(GetTransientPackage(), "FixtureTypeSharedDataSelection", RF_Transactional);
+	Selection = NewObject<UDMXFixtureTypeSharedDataSelection>(GetTransientPackage(), NAME_None, RF_Transactional);
 }
 
 void FDMXFixtureTypeSharedData::AddReferencedObjects(FReferenceCollector& Collector)
