DMX - Fixture Mode Channel Span now calculates correctly in UI

diff --git a/Engine/Plugins/VirtualProduction/DMX/DMXEngine/Source/DMXRuntime/Private/Library/DMXEntityFixtureType.cpp b/Engine/Plugins/VirtualProduction/DMX/DMXEngine/Source/DMXRuntime/Private/Library/DMXEntityFixtureType.cpp
--- a/Engine/Plugins/VirtualProduction/DMX/DMXEngine/Source/DMXRuntime/Private/Library/DMXEntityFixtureType.cpp
+++ b/Engine/Plugins/VirtualProduction/DMX/DMXEngine/Source/DMXRuntime/Private/Library/DMXEntityFixtureType.cpp
@@ -199,6 +199,15 @@ void UDMXEntityFixtureType::RemoveFixtureTypeFromLibrary(FDMXEntityFixtureTypeRe
 	{
 		if (UDMXLibrary* DMXLibrary = FixtureType->GetParentLibrary())
 		{
+			const TArray<UDMXEntityFixturePatch*> FixturePatches = DMXLibrary->GetEntitiesTypeCast<UDMXEntityFixturePatch>();
+			for (UDMXEntityFixturePatch* FixturePatch : FixturePatches)
+			{
+				if (FixturePatch->GetFixtureType() == FixtureType)
+				{
+					FixturePatch->SetFixtureType(nullptr);
+				}
+			}
+
 			DMXLibrary->Modify();
 			FixtureType->Modify();
 			FixtureType->Destroy();
@@ -248,8 +257,8 @@ void UDMXEntityFixtureType::PostEditChangeProperty(FPropertyChangedEvent& Proper
 		// Update the Channel Span for all Modes
 		for (int32 ModeIndex = 0; ModeIndex < Modes.Num(); ModeIndex++)
 		{
-			UpdateChannelSpan(ModeIndex);
 			AlignFunctionChannels(ModeIndex);
+			UpdateChannelSpan(ModeIndex);
 		}
 
 		OnFixtureTypeChangedDelegate.Broadcast(this);
@@ -280,8 +289,8 @@ void UDMXEntityFixtureType::PostEditChangeChainProperty(FPropertyChangedChainEve
 		// Update the Channel Span for all Modes
 		for (int32 ModeIndex = 0; ModeIndex < Modes.Num(); ModeIndex++)
 		{
-			UpdateChannelSpan(ModeIndex);
 			AlignFunctionChannels(ModeIndex);
+			UpdateChannelSpan(ModeIndex);
 		}
 
 		OnFixtureTypeChangedDelegate.Broadcast(this);
@@ -540,6 +549,7 @@ void UDMXEntityFixtureType::SetFixtureMatrixEnabled(int32 ModeIndex, bool bEnabl
 			Mode.FixtureMatrixConfig.YCells = FMath::Max(Mode.FixtureMatrixConfig.YCells, 1);
 
 			AlignFunctionChannels(ModeIndex);
+			UpdateChannelSpan(ModeIndex);
 		}
 	}
 }
@@ -1036,8 +1046,8 @@ void UDMXEntityFixtureType::UpdateYCellsFromXCells(int32 ModeIndex)
 			}
 			Matrix.YCells = DMX_MAX_ADDRESS / (Matrix.XCells * NumChannelsOfCell);
 
-			UpdateChannelSpan(ModeIndex);
 			AlignFunctionChannels(ModeIndex);
+			UpdateChannelSpan(ModeIndex);
 		}
 	}
 }
@@ -1060,8 +1070,8 @@ void UDMXEntityFixtureType::UpdateXCellsFromYCells(int32 ModeIndex)
 
 				Matrix.XCells = DMX_MAX_ADDRESS / (Matrix.YCells * NumChannelsOfCell);
 
-				UpdateChannelSpan(ModeIndex);
 				AlignFunctionChannels(ModeIndex);
+				UpdateChannelSpan(ModeIndex);
 			}
 		}
 	}
