DMX - Fix Mode details never show a scrollbar

diff --git a/Engine/Plugins/VirtualProduction/DMX/DMXEngine/Source/DMXEditor/Private/Widgets/FixtureType/SDMXFixtureModeEditor.cpp b/Engine/Plugins/VirtualProduction/DMX/DMXEngine/Source/DMXEditor/Private/Widgets/FixtureType/SDMXFixtureModeEditor.cpp
--- a/Engine/Plugins/VirtualProduction/DMX/DMXEngine/Source/DMXEditor/Private/Widgets/FixtureType/SDMXFixtureModeEditor.cpp
+++ b/Engine/Plugins/VirtualProduction/DMX/DMXEngine/Source/DMXEditor/Private/Widgets/FixtureType/SDMXFixtureModeEditor.cpp
@@ -10,6 +10,7 @@
 #include "PropertyEditorModule.h"
 #include "ScopedTransaction.h"
 #include "Modules/ModuleManager.h"
+#include "Widgets/SOverlay.h"
 #include "Widgets/Layout/SBorder.h"
 #include "Widgets/Layout/SBox.h"
 #include "Widgets/Text/STextBlock.h"
@@ -37,7 +38,7 @@ void SDMXFixtureModeEditor::Construct(const FArguments& InArgs, const TSharedRef
 	DetailsViewArgs.bShowModifiedPropertiesOption = false;
 	DetailsViewArgs.bShowObjectLabel = false;
 	DetailsViewArgs.bForceHiddenPropertyVisibility = false;
-	DetailsViewArgs.bShowScrollBar = false;
+	DetailsViewArgs.bShowScrollBar = true;
 	DetailsViewArgs.NotifyHook = this;
 
 	FStructureDetailsViewArgs StructureDetailsViewArgs;
@@ -51,15 +52,14 @@ void SDMXFixtureModeEditor::Construct(const FArguments& InArgs, const TSharedRef
 
 	ChildSlot
 	[
-		SNew(SVerticalBox)
+		SNew(SOverlay)
 
-		+ SVerticalBox::Slot()
-		.AutoHeight()
+		+ SOverlay::Slot()
 		[
 			StructDetailsViewWidget.ToSharedRef()
 		]
 
-		+ SVerticalBox::Slot()
+		+ SOverlay::Slot()
 		[
 			SNew(SBox)
 			.HAlign(HAlign_Center)
