DMX - Fix contextual "New Fixture Type" option when right-clicking an empty area

diff --git a/Engine/Plugins/VirtualProduction/DMX/DMXEngine/Source/DMXEditor/Private/Widgets/FixtureType/SDMXFixtureTypeTree.cpp b/Engine/Plugins/VirtualProduction/DMX/DMXEngine/Source/DMXEditor/Private/Widgets/FixtureType/SDMXFixtureTypeTree.cpp
--- a/Engine/Plugins/VirtualProduction/DMX/DMXEngine/Source/DMXEditor/Private/Widgets/FixtureType/SDMXFixtureTypeTree.cpp
+++ b/Engine/Plugins/VirtualProduction/DMX/DMXEngine/Source/DMXEditor/Private/Widgets/FixtureType/SDMXFixtureTypeTree.cpp
@@ -40,6 +40,10 @@ void SDMXFixtureTypeTree::Construct(const FArguments& InArgs)
 	
 	if (const TSharedPtr<FDMXEditor> PinnedDMXEditor = DMXEditor.Pin())
 	{
+		CommandList->MapAction(FDMXEditorCommands::Get().AddNewEntityFixtureType,
+			FUIAction(FExecuteAction::CreateSP(this, &SDMXFixtureTypeTree::AddNewFixtureType))
+		);
+
 		FixtureTypeSharedData = PinnedDMXEditor->GetFixtureTypeSharedData();
 		FixtureTypeSharedData->OnFixtureTypesSelected.AddSP(this, &SDMXFixtureTypeTree::OnFixtureTypesSelected);
 		PinnedDMXEditor->GetFixturePatchSharedData()->OnFixturePatchSelectionChanged.AddSP(this, &SDMXFixtureTypeTree::OnFixturePatchesSelected);
@@ -620,9 +624,15 @@ void SDMXFixtureTypeTree::OnFixtureTypeChanged(const UDMXEntityFixtureType* Fixt
 }
 
 FReply SDMXFixtureTypeTree::OnAddNewFixtureTypeClicked()
+{
+	AddNewFixtureType();
+	return FReply::Handled();
+}
+
+void SDMXFixtureTypeTree::AddNewFixtureType()
 {
 	if (TSharedPtr<FDMXEditor> PinnedEditor = DMXEditor.Pin())
-	{		
+	{
 		const FScopedTransaction Transaction(LOCTEXT("CreateFixtureTypeTransaction", "Create DMX Fixture Type"));
 
 		FDMXEntityFixtureTypeConstructionParams FixtureTypeConstructionParams;
@@ -633,11 +643,7 @@ FReply SDMXFixtureTypeTree::OnAddNewFixtureTypeClicked()
 		FixtureTypeSharedData->SelectFixtureTypes(TArray<TWeakObjectPtr<UDMXEntityFixtureType>>({ NewFixtureType }));
 
 		UpdateTree();
-
-		return FReply::Handled();
 	}
-
-	return FReply::Unhandled();
 }
 
 #undef LOCTEXT_NAMESPACE
diff --git a/Engine/Plugins/VirtualProduction/DMX/DMXEngine/Source/DMXEditor/Private/Widgets/FixtureType/SDMXFixtureTypeTree.h b/Engine/Plugins/VirtualProduction/DMX/DMXEngine/Source/DMXEditor/Private/Widgets/FixtureType/SDMXFixtureTypeTree.h
--- a/Engine/Plugins/VirtualProduction/DMX/DMXEngine/Source/DMXEditor/Private/Widgets/FixtureType/SDMXFixtureTypeTree.h
+++ b/Engine/Plugins/VirtualProduction/DMX/DMXEngine/Source/DMXEditor/Private/Widgets/FixtureType/SDMXFixtureTypeTree.h
@@ -75,6 +75,9 @@ private:
 	/** Called when the Add button was clicked */
 	FReply OnAddNewFixtureTypeClicked();
 
+	/** Adds a new Fixture Type to the Library */
+	void AddNewFixtureType();
+
 	/** Returns the row that corresponds to the node */
 	TSharedPtr<SDMXFixtureTypeTreeFixtureTypeRow> FindEntityRowByNode(const TSharedRef<FDMXEntityTreeEntityNode>& EntityNode);
 
