DMX - Allow users to bulk add Fixtures in Patch Window

diff --git a/Engine/Plugins/VirtualProduction/DMX/DMXEngine/Source/DMXEditor/Private/Widgets/FixturePatch/SDMXMVRFixtureListToolbar.cpp b/Engine/Plugins/VirtualProduction/DMX/DMXEngine/Source/DMXEditor/Private/Widgets/FixturePatch/SDMXMVRFixtureListToolbar.cpp
--- a/Engine/Plugins/VirtualProduction/DMX/DMXEngine/Source/DMXEditor/Private/Widgets/FixturePatch/SDMXMVRFixtureListToolbar.cpp
+++ b/Engine/Plugins/VirtualProduction/DMX/DMXEngine/Source/DMXEditor/Private/Widgets/FixturePatch/SDMXMVRFixtureListToolbar.cpp
@@ -3,6 +3,7 @@
 #include "SDMXMVRFixtureListToolbar.h"
 
 #include "DMXEditor.h"
+#include "DMXFixturePatchSharedData.h"
 #include "DMXMVRFixtureListItem.h"
 #include "Library/DMXEntityFixturePatch.h"
 #include "Library/DMXEntityFixtureType.h"
@@ -14,6 +15,7 @@
 #include "Internationalization/Regex.h"
 #include "Widgets/SBoxPanel.h"
 #include "Widgets/Input/SCheckBox.h"
+#include "Widgets/Input/SEditableTextBox.h"
 #include "Widgets/Input/SSearchBox.h"
 #include "Widgets/Layout/SSeparator.h"
 #include "Widgets/Layout/SWrapBox.h"
@@ -292,9 +294,68 @@ TSharedRef<SWidget> SDMXMVRFixtureListToolbar::GenerateFixtureTypeDropdownMenu()
 			]
 			.MenuContent()
 			[
-				SAssignNew(FixtureTypeDropdownMenu, SDMXEntityDropdownMenu<UDMXEntityFixtureType>)
-				.DMXEditor(WeakDMXEditor)
-				.OnEntitySelected(this, &SDMXMVRFixtureListToolbar::OnAddNewMVRFixtureClicked)
+
+				SNew(SVerticalBox)
+				
+				// Bulk Add Fixture Patches 
+				+ SVerticalBox::Slot()
+				.AutoHeight()
+				[
+					SNew(SHorizontalBox)
+					
+					+SHorizontalBox::Slot()
+					.HAlign(HAlign_Left)
+					.VAlign(VAlign_Center)
+					.Padding(8.f, 2.f, 4.f, 2.f)
+					[
+						SNew(STextBlock)
+						.Text(LOCTEXT("BulkAddPatchesLabel", "Quantity"))
+					]
+
+					+ SHorizontalBox::Slot()
+					.HAlign(HAlign_Fill)
+					.VAlign(VAlign_Center)
+					.Padding(4.f, 2.f, 4.f, 2.f)
+					[
+						SNew(SEditableTextBox)
+						.SelectAllTextWhenFocused(true)
+						.ClearKeyboardFocusOnCommit(false)
+						.SelectAllTextOnCommit(true)
+						.Text_Lambda([this]()
+							{
+								return FText::FromString(FString::FromInt(NumFixturePatchesToAdd));
+							})
+						.OnVerifyTextChanged_Lambda([](const FText& InNewText, FText& OutErrorMessage)
+							{
+								int32 Value;
+								if (!LexTryParseString<int32>(Value, *InNewText.ToString()) ||
+									Value < 1)
+								{
+									OutErrorMessage = LOCTEXT("BulkAddPatchesBadString", "Needs a numeric value > 0");
+									return false;
+								}
+								return true;
+							})
+						.OnTextCommitted_Lambda([this](const FText& Text, ETextCommit::Type CommitType)
+							{
+								int32 Value;
+								if (LexTryParseString<int32>(Value, *Text.ToString()) &&
+									Value > 0)
+								{
+									constexpr int32 MaxNumPatchesToBulkAdd = 512;
+									NumFixturePatchesToAdd = FMath::Min(Value, MaxNumPatchesToBulkAdd);
+								}
+							})
+					]
+				]
+
+				// Fixture Type Selection
+				+ SVerticalBox::Slot()
+				[
+					SAssignNew(FixtureTypeDropdownMenu, SDMXEntityDropdownMenu<UDMXEntityFixtureType>)
+					.DMXEditor(WeakDMXEditor)
+					.OnEntitySelected(this, &SDMXMVRFixtureListToolbar::OnAddNewMVRFixtureClicked)
+				]
 			]
 			.IsFocusable(true)
 			.ContentPadding(FMargin(5.0f, 1.0f))
@@ -314,6 +375,11 @@ TSharedRef<SWidget> SDMXMVRFixtureListToolbar::GenerateFixtureTypeDropdownMenu()
 
 void SDMXMVRFixtureListToolbar::OnAddNewMVRFixtureClicked(UDMXEntity* InSelectedFixtureType)
 {
+	if (!ensureMsgf(InSelectedFixtureType, TEXT("Trying to add Fixture Patches, but selected fixture type is invalid.")))
+	{
+		return;
+	}
+
 	TSharedPtr<FDMXEditor> DMXEditor = WeakDMXEditor.Pin();
 	if (!DMXEditor.IsValid())
 	{
@@ -328,42 +394,32 @@ void SDMXMVRFixtureListToolbar::OnAddNewMVRFixtureClicked(UDMXEntity* InSelected
 
 	// Find a universe and address
 	TArray<UDMXEntityFixturePatch*> FixturePatches = DMXLibrary->GetEntitiesTypeCast<UDMXEntityFixturePatch>();
-	FixturePatches.Sort([](const UDMXEntityFixturePatch& FixturePatchA, const UDMXEntityFixturePatch& FixturePatchB)
-		{
-			const bool bUniverseIsSmaller = FixturePatchA.GetUniverseID() < FixturePatchB.GetUniverseID();
-			const bool bUniverseIsEqual = FixturePatchA.GetUniverseID() == FixturePatchB.GetUniverseID();
-			const bool bLastAddressIsSmaller = FixturePatchA.GetEndingChannel() < FixturePatchB.GetEndingChannel();
-			
-			return bUniverseIsSmaller || (bUniverseIsEqual && bLastAddressIsSmaller);
-		});
 
-	int32 Universe = 1; 
-	int32 Address = 1;
-	if (FixturePatches.Num() > 0)
+	// Create a new fixture patches
+	const FScopedTransaction CreateFixturePatchTransaction(LOCTEXT("CreateFixturePatchTransaction", "Create Fixture Patch"));
+	TArray<UDMXEntityFixturePatch*> NewFixturePatches;
+	TArray<TWeakObjectPtr<UDMXEntityFixturePatch>> NewWeakFixturePatches;
+	DMXLibrary->PreEditChange(UDMXLibrary::StaticClass()->FindPropertyByName(UDMXLibrary::GetEntitiesPropertyName()));
+	for (int32 iNumFixturePatchesAdded = 0; iNumFixturePatchesAdded < NumFixturePatchesToAdd; iNumFixturePatchesAdded++)
 	{
-		const int32 ChannelSpan = FixtureType->Modes.Num() > 0 ? FixtureType->Modes[0].ChannelSpan : 1;
-		if (FixturePatches.Last()->GetEndingChannel() + ChannelSpan > DMX_MAX_ADDRESS)
-		{
-			Universe = FixturePatches.Last()->GetUniverseID() + 1;
-			Address = 1;
-		}
-		else
-		{
-			Universe = FixturePatches.Last()->GetUniverseID();
-			Address = FixturePatches.Last()->GetEndingChannel() + 1;
-		}
+		FDMXEntityFixturePatchConstructionParams FixturePatchConstructionParams;
+		FixturePatchConstructionParams.FixtureTypeRef = FDMXEntityFixtureTypeRef(FixtureType);
+		FixturePatchConstructionParams.ActiveMode = 0;
+
+		constexpr bool bMarkLibraryDirty = false;
+		UDMXEntityFixturePatch* NewFixturePatch = UDMXEntityFixturePatch::CreateFixturePatchInLibrary(FixturePatchConstructionParams, FixtureType->Name, bMarkLibraryDirty);
+		NewFixturePatches.Add(NewFixturePatch);
+		NewWeakFixturePatches.Add(NewFixturePatch);	
 	}
 
-	// Create a new fixture patch
-	FDMXEntityFixturePatchConstructionParams FixturePatchConstructionParams;
-	FixturePatchConstructionParams.FixtureTypeRef = FDMXEntityFixtureTypeRef(FixtureType);
-	FixturePatchConstructionParams.ActiveMode = 0;
-	FixturePatchConstructionParams.UniverseID = Universe;
-	FixturePatchConstructionParams.StartingAddress = Address;
+	// Auto assign
+	constexpr bool bAllowDecrementUniverse = false;
+	constexpr bool bAllowDecrementChannels = false;
+	FDMXEditorUtils::AutoAssignedChannels(bAllowDecrementUniverse, bAllowDecrementChannels, NewFixturePatches);
 
-	const FScopedTransaction CreateFixturePatchTransaction(LOCTEXT("CreateFixturePatchTransaction", "Create Fixture Patch"));
-	constexpr bool bMarkLibraryDirty = true;
-	UDMXEntityFixturePatch* NewFixturePatch = UDMXEntityFixturePatch::CreateFixturePatchInLibrary(FixturePatchConstructionParams, FixtureType->Name, bMarkLibraryDirty);
+	DMXLibrary->PostEditChange();
+
+	DMXEditor->GetFixturePatchSharedData()->SelectFixturePatches(NewWeakFixturePatches);
 }
 
 void SDMXMVRFixtureListToolbar::OnSearchTextChanged(const FText& SearchText)
diff --git a/Engine/Plugins/VirtualProduction/DMX/DMXEngine/Source/DMXEditor/Private/Widgets/FixturePatch/SDMXMVRFixtureListToolbar.h b/Engine/Plugins/VirtualProduction/DMX/DMXEngine/Source/DMXEditor/Private/Widgets/FixturePatch/SDMXMVRFixtureListToolbar.h
--- a/Engine/Plugins/VirtualProduction/DMX/DMXEngine/Source/DMXEditor/Private/Widgets/FixturePatch/SDMXMVRFixtureListToolbar.h
+++ b/Engine/Plugins/VirtualProduction/DMX/DMXEngine/Source/DMXEditor/Private/Widgets/FixturePatch/SDMXMVRFixtureListToolbar.h
@@ -52,6 +52,9 @@ private:
 	/** True if only conflicts are searched for */
 	bool bShowConfictsOnly = false;
 
+	/** The number of fixture patches to add when one is selected in the entity dropdown */
+	int32 NumFixturePatchesToAdd = 1;
+
 	/** Dropdown menu of Fixture Types */
 	TSharedPtr<SDMXEntityDropdownMenu<UDMXEntityFixtureType>> FixtureTypeDropdownMenu;
 
