DMX - Fix newly created  Fixture Patches may be assigned in front of the last patch

diff --git a/Engine/Plugins/VirtualProduction/DMX/DMXEngine/Source/DMXEditor/Private/Widgets/FixturePatch/SDMXMVRFixtureListToolbar.cpp b/Engine/Plugins/VirtualProduction/DMX/DMXEngine/Source/DMXEditor/Private/Widgets/FixturePatch/SDMXMVRFixtureListToolbar.cpp
--- a/Engine/Plugins/VirtualProduction/DMX/DMXEngine/Source/DMXEditor/Private/Widgets/FixturePatch/SDMXMVRFixtureListToolbar.cpp
+++ b/Engine/Plugins/VirtualProduction/DMX/DMXEngine/Source/DMXEditor/Private/Widgets/FixturePatch/SDMXMVRFixtureListToolbar.cpp
@@ -12,6 +12,7 @@
 
 #include "EditorStyleSet.h"
 #include "ScopedTransaction.h"
+#include "Algo/MaxElement.h"
 #include "Internationalization/Regex.h"
 #include "Widgets/SBoxPanel.h"
 #include "Widgets/Input/SCheckBox.h"
@@ -392,8 +393,28 @@ void SDMXMVRFixtureListToolbar::OnAddNewMVRFixtureClicked(UDMXEntity* InSelected
 	}
 	UDMXEntityFixtureType* FixtureType = CastChecked<UDMXEntityFixtureType>(InSelectedFixtureType);
 
-	// Find a universe and address
+	// Find a Universe and Address
 	TArray<UDMXEntityFixturePatch*> FixturePatches = DMXLibrary->GetEntitiesTypeCast<UDMXEntityFixturePatch>();
+	UDMXEntityFixturePatch* const* LastFixturePatchPtr = Algo::MaxElementBy(FixturePatches, [](const UDMXEntityFixturePatch* FixturePatch)
+		{
+			return (int64)FixturePatch->GetUniverseID() * DMX_MAX_ADDRESS + FixturePatch->GetStartingChannel();
+		});
+	int32 Universe = 1;
+	int32 Address = 1;
+	if (LastFixturePatchPtr && !FixtureType->Modes.IsEmpty())
+	{
+		const UDMXEntityFixturePatch& LastFixturePatch = **LastFixturePatchPtr;
+		if (LastFixturePatch.GetEndingChannel() + FixtureType->Modes[0].ChannelSpan > DMX_MAX_ADDRESS)
+		{
+			Universe = LastFixturePatch.GetUniverseID() + 1;
+			Address = 1;
+		}
+		else
+		{
+			Universe = LastFixturePatch.GetUniverseID();
+			Address = LastFixturePatch.GetEndingChannel() + 1;
+		}
+	}
 
 	// Create a new fixture patches
 	const FScopedTransaction CreateFixturePatchTransaction(LOCTEXT("CreateFixturePatchTransaction", "Create Fixture Patch"));
@@ -405,11 +426,24 @@ void SDMXMVRFixtureListToolbar::OnAddNewMVRFixtureClicked(UDMXEntity* InSelected
 		FDMXEntityFixturePatchConstructionParams FixturePatchConstructionParams;
 		FixturePatchConstructionParams.FixtureTypeRef = FDMXEntityFixtureTypeRef(FixtureType);
 		FixturePatchConstructionParams.ActiveMode = 0;
+		FixturePatchConstructionParams.UniverseID = Universe;
+		FixturePatchConstructionParams.StartingAddress = Address;
 
 		constexpr bool bMarkLibraryDirty = false;
 		UDMXEntityFixturePatch* NewFixturePatch = UDMXEntityFixturePatch::CreateFixturePatchInLibrary(FixturePatchConstructionParams, FixtureType->Name, bMarkLibraryDirty);
 		NewFixturePatches.Add(NewFixturePatch);
 		NewWeakFixturePatches.Add(NewFixturePatch);	
+
+		// Increment Universe and Address in steps by one. This is enough for auto assign while keeping order of the named patches
+		if (Address + FixtureType->Modes[0].ChannelSpan > DMX_MAX_ADDRESS)
+		{
+			Universe++;
+			Address = 1; 
+		}
+		else
+		{
+			Address++;
+		}
 	}
 
 	// Auto assign
