DMX - MVR improvements

* MVR Scenes can now be refreshed from a changed or reimported DMX Library
* Spawned fixture actors are now named after the patch instead of the Actor class
* All actors of a specific GDTF type can now be selected from Details
* Prevent from dropping the same DMX Library in the level twice
* Destroy attached actors when MVR Scene Actor is destroyed

diff --git a/Engine/Plugins/VirtualProduction/DMX/DMXEngine/Source/DMXEditor/Private/Customizations/DMXMVRSceneActorDetails.cpp b/Engine/Plugins/VirtualProduction/DMX/DMXEngine/Source/DMXEditor/Private/Customizations/DMXMVRSceneActorDetails.cpp
--- a/Engine/Plugins/VirtualProduction/DMX/DMXEngine/Source/DMXEditor/Private/Customizations/DMXMVRSceneActorDetails.cpp
+++ b/Engine/Plugins/VirtualProduction/DMX/DMXEngine/Source/DMXEditor/Private/Customizations/DMXMVRSceneActorDetails.cpp
@@ -12,6 +12,13 @@
 #include "DetailLayoutBuilder.h"
 #include "DetailWidgetRow.h"
 #include "IPropertyUtilities.h"
+#include "ISceneOutliner.h"
+#include "LevelEditor.h"
+#include "ScopedTransaction.h"
+#include "Modules/ModuleManager.h"
+#include "Subsystems/EditorActorSubsystem.h"
+#include "Widgets/Input/SButton.h"
+#include "Widgets/Layout/SWrapBox.h"
 
 
 #define LOCTEXT_NAMESPACE "DMXMVRSceneActorDetails"
@@ -35,6 +42,73 @@ void FDMXMVRSceneActorDetails::CustomizeDetails(IDetailLayoutBuilder& DetailBuil
 		}
 	}
 
+	CreateRefreshMVRSceneSection(DetailBuilder);
+	CreateGDTFToActorClassSection(DetailBuilder);
+
+	// Listen to map and actor changes
+	FEditorDelegates::MapChange.AddSP(this, &FDMXMVRSceneActorDetails::OnMapChange);
+
+	if (GEngine)
+	{
+		GEngine->OnLevelActorDeleted().AddSP(this, &FDMXMVRSceneActorDetails::OnActorDeleted);
+	}
+}
+
+void FDMXMVRSceneActorDetails::CreateRefreshMVRSceneSection(IDetailLayoutBuilder& DetailBuilder)
+{
+	IDetailCategoryBuilder& MVRCategory = DetailBuilder.EditCategory("MVR");
+
+	MVRCategory.AddProperty(DetailBuilder.GetProperty(ADMXMVRSceneActor::GetDMXLibraryPropertyNameChecked()));
+	MVRCategory.AddCustomRow(LOCTEXT("RefreshSceneFilterText", "Refresh from DMX Library"))
+		.WholeRowContent()
+		[			
+			SNew(SBorder)
+			.Padding(8.f, 0.f, 0.f, 0.f)
+			.BorderImage(FAppStyle::GetBrush("NoBorder"))
+			[
+				SNew(SButton)
+				.Text(LOCTEXT("RefreshActorsFromDMXLibraryButtonCaption", "Update Actors from DMX Library"))
+				.OnClicked(this, &FDMXMVRSceneActorDetails::OnRefreshActorsFromDMXLibraryClicked)
+			]
+		];
+
+	const TSharedRef<IPropertyHandle> RespawnDeletedActorHandle = DetailBuilder.GetProperty(GET_MEMBER_NAME_CHECKED(ADMXMVRSceneActor, bRespawnDeletedActorsOnRefresh));
+	RespawnDeletedActorHandle->MarkHiddenByCustomization();
+	MVRCategory.AddCustomRow(LOCTEXT("RespawnDeletedActorsFilterText", "Respawn Deleted Actors"))
+		.NameContent()
+		[
+			SNew(SBorder)
+			.Padding(8.f, 0.f, 0.f, 0.f)
+			.BorderImage(FAppStyle::GetBrush("NoBorder"))
+			[
+				RespawnDeletedActorHandle->CreatePropertyNameWidget()
+			]
+		]
+		.ValueContent()
+		[
+			RespawnDeletedActorHandle->CreatePropertyValueWidget()
+		];
+	
+	const TSharedRef<IPropertyHandle> UpdateTransformHandle = DetailBuilder.GetProperty(GET_MEMBER_NAME_CHECKED(ADMXMVRSceneActor, bUpdateTransformsOnRefresh));
+	UpdateTransformHandle->MarkHiddenByCustomization();
+	MVRCategory.AddCustomRow(LOCTEXT("UpdateTransformsFilterText", "Update Transform"))
+		.NameContent()
+		[
+			SNew(SBorder)
+			.Padding(8.f, 0.f, 0.f, 0.f)
+			.BorderImage(FAppStyle::GetBrush("NoBorder"))
+			[
+				UpdateTransformHandle->CreatePropertyNameWidget()
+			]
+		]
+		.ValueContent()
+		[
+			UpdateTransformHandle->CreatePropertyValueWidget()
+		];
+}
+
+void FDMXMVRSceneActorDetails::CreateGDTFToActorClassSection(IDetailLayoutBuilder& DetailBuilder)
+{
 	IDetailCategoryBuilder& ActorTypeForGDTFCategory = DetailBuilder.EditCategory("GDTF to Spawned Actor");
 	ActorTypeForGDTFCategory.InitiallyCollapsed(false);
 
@@ -42,7 +116,7 @@ void FDMXMVRSceneActorDetails::CustomizeDetails(IDetailLayoutBuilder& DetailBuil
 	GDTFToDefaultActorClassesHandle->MarkHiddenByCustomization();
 
 	const TSharedPtr<IPropertyHandleArray> GDTFToDefaultActorClassHandleArray = GDTFToDefaultActorClassesHandle->AsArray();
-	FSimpleDelegate GDTFToDefaultActorClassArrayChangedDelegate = FSimpleDelegate::CreateSP(this, &FDMXMVRSceneActorDetails::ForceRefresh);
+	FSimpleDelegate GDTFToDefaultActorClassArrayChangedDelegate = FSimpleDelegate::CreateSP(this, &FDMXMVRSceneActorDetails::RequestRefresh);
 	GDTFToDefaultActorClassHandleArray->SetOnNumElementsChanged(GDTFToDefaultActorClassArrayChangedDelegate);
 
 	uint32 NumGDTFToDefaultActorClassElements;
@@ -59,26 +133,119 @@ void FDMXMVRSceneActorDetails::CustomizeDetails(IDetailLayoutBuilder& DetailBuil
 		DefaultActorClassHandle->SetOnChildPropertyValuePreChange(FSimpleDelegate::CreateSP(this, &FDMXMVRSceneActorDetails::OnPreEditChangeActorClassInGDTFToActorClasses));
 		DefaultActorClassHandle->SetOnChildPropertyValueChanged(FSimpleDelegate::CreateSP(this, &FDMXMVRSceneActorDetails::OnPostEditChangeActorClassInGDTFToActorClasses));
 
+		UObject* GDTFObject;
+		if (!GDTFHandle->GetValue(GDTFObject))
+		{
+			return;
+		}
+		UDMXImportGDTF* GDTF = Cast<UDMXImportGDTF>(GDTFObject);
+		if (!GDTF)
+		{
+			continue;
+		}
+
+		if (!IsAnyActorUsingGDTF(GDTF))
+		{
+			continue;
+		}
+
 		ActorTypeForGDTFCategory.AddCustomRow(LOCTEXT("GDTFToDefaultActorClassFilter", "GDTF"))
-			.NameContent()
+		.NameContent()
+		[
+			GDTFHandle->CreatePropertyValueWidget()
+		]
+		.ValueContent()
+		[
+			SNew(SWrapBox)
+
+			+ SWrapBox::Slot()
 			[
-				GDTFHandle->CreatePropertyValueWidget()
+				DefaultActorClassHandle->CreatePropertyValueWidget()
 			]
-			.ValueContent()
+
+			+ SWrapBox::Slot()
 			[
-				DefaultActorClassHandle->CreatePropertyValueWidget()
-			];
+				SNew(SButton)
+				.OnClicked(this, &FDMXMVRSceneActorDetails::OnSelectGDTFToActorClassGroupClicked, GDTFObject)
+				.Text(LOCTEXT("SelectGDTFGroupButtonCaption", "Select"))
+			]
+		];
+	}
+}
+
+FReply FDMXMVRSceneActorDetails::OnRefreshActorsFromDMXLibraryClicked()
+{
+	const TArray<TWeakObjectPtr<UObject>> SelectedObjects = PropertyUtilities->GetSelectedObjects();
+
+	for (TWeakObjectPtr<UObject> SelectedObject : SelectedObjects)
+	{
+		if (ADMXMVRSceneActor* MVRSceneActor = Cast<ADMXMVRSceneActor>(SelectedObject.Get()))
+		{
+			const FScopedTransaction RefreshActorsFromDMXLibraryTransaction(LOCTEXT("RefreshActorsFromDMXLibraryTransaction", "Update MVR Scene form DMX Library"));
+			MVRSceneActor->PreEditChange(ADMXMVRSceneActor::StaticClass()->FindPropertyByName(ADMXMVRSceneActor::GetRelatedAcctorsPropertyNameChecked()));
+
+			MVRSceneActor->RefreshFromDMXLibrary();
+
+			MVRSceneActor->PostEditChange();
+		}
+	}
+
+	return FReply::Handled();
+}
+
+FReply FDMXMVRSceneActorDetails::OnSelectGDTFToActorClassGroupClicked(UObject* GDTFObject)
+{
+	UDMXImportGDTF* GDTF = Cast<UDMXImportGDTF>(GDTFObject);
+	if (!GDTF)
+	{
+		return FReply::Unhandled();
+	}
+
+	const TArray<TWeakObjectPtr<UObject>> SelectedObjects = PropertyUtilities->GetSelectedObjects();
+
+	for (TWeakObjectPtr<UObject> SelectedObject : SelectedObjects)
+	{
+		if (ADMXMVRSceneActor* MVRSceneActor = Cast<ADMXMVRSceneActor>(SelectedObject.Get()))
+		{
+			const TArray<AActor*> ActorsForThisGDTF = MVRSceneActor->GetActorsSpawnedForGDTF(GDTF);
+			UEditorActorSubsystem* EditorActorSubsystem = GEditor->GetEditorSubsystem<UEditorActorSubsystem>();
+
+			if (EditorActorSubsystem)
+			{
+				EditorActorSubsystem->SetSelectedLevelActors(ActorsForThisGDTF);
+			}
+		}
 	}
+
+	// Set focus on the Scene Outliner so the user can execute keyboard commands right away
+	const TWeakPtr<class ILevelEditor> LevelEditor = FModuleManager::GetModuleChecked<FLevelEditorModule>(TEXT("LevelEditor")).GetLevelEditorInstance();
+	const TSharedPtr<class ISceneOutliner> SceneOutliner = LevelEditor.IsValid() ? LevelEditor.Pin()->GetMostRecentlyUsedSceneOutliner() : nullptr;
+	if (SceneOutliner.IsValid())
+	{
+		SceneOutliner->SetKeyboardFocus();
+	}
+
+	return FReply::Handled();
 }
 
 void FDMXMVRSceneActorDetails::OnFixturePatchChanged(const UDMXEntityFixturePatch* FixturePatch)
 {
-	ForceRefresh();
+	RequestRefresh();
 }
 
 void FDMXMVRSceneActorDetails::OnFixtureTypeChanged(const UDMXEntityFixtureType* FixtureType)
 {
-	ForceRefresh();
+	RequestRefresh();
+}
+
+void FDMXMVRSceneActorDetails::OnMapChange(uint32 MapChangeFlags)
+{
+	RequestRefresh();
+}
+
+void FDMXMVRSceneActorDetails::OnActorDeleted(AActor* DeletedActor)
+{
+	RequestRefresh();
 }
 
 void FDMXMVRSceneActorDetails::OnPreEditChangeActorClassInGDTFToActorClasses()
@@ -103,9 +270,22 @@ void FDMXMVRSceneActorDetails::OnPostEditChangeActorClassInGDTFToActorClasses()
 	}
 }
 
-void FDMXMVRSceneActorDetails::ForceRefresh()
+bool FDMXMVRSceneActorDetails::IsAnyActorUsingGDTF(const UDMXImportGDTF* GDTF) const
+{
+	const TArray<TWeakObjectPtr<UObject>> SelectedObjects = PropertyUtilities->GetSelectedObjects();
+	for (TWeakObjectPtr<UObject> SelectedObject : SelectedObjects)
+	{
+		if (ADMXMVRSceneActor* MVRSceneActor = Cast<ADMXMVRSceneActor>(SelectedObject.Get()))
+		{
+			return !MVRSceneActor->GetActorsSpawnedForGDTF(GDTF).IsEmpty();
+		}
+	}
+	return false;
+}
+
+void FDMXMVRSceneActorDetails::RequestRefresh()
 {
-	PropertyUtilities->ForceRefresh();
+	PropertyUtilities->RequestRefresh();
 }
 
 #undef LOCTEXT_NAMESPACE
diff --git a/Engine/Plugins/VirtualProduction/DMX/DMXEngine/Source/DMXEditor/Private/Customizations/DMXMVRSceneActorDetails.h b/Engine/Plugins/VirtualProduction/DMX/DMXEngine/Source/DMXEditor/Private/Customizations/DMXMVRSceneActorDetails.h
--- a/Engine/Plugins/VirtualProduction/DMX/DMXEngine/Source/DMXEditor/Private/Customizations/DMXMVRSceneActorDetails.h
+++ b/Engine/Plugins/VirtualProduction/DMX/DMXEngine/Source/DMXEditor/Private/Customizations/DMXMVRSceneActorDetails.h
@@ -4,11 +4,14 @@
 
 #include "CoreMinimal.h"
 #include "IDetailCustomization.h"
+#include "Input/Reply.h"
 
 class ADMXMVRSceneActor;
 class UDMXEntityFixturePatch;
 class UDMXEntityFixtureType;
+class UDMXImportGDTF;
 
+enum class ECheckBoxState : uint8;
 class IPropertyHandle;
 class IPropertyUtilities;
 
@@ -27,20 +30,41 @@ protected:
 	//~ End IDetailCustomization interface
 
 private:
+	/** Creates the section where the user can refresh the MVR scene from a changed DMX Library */
+	void CreateRefreshMVRSceneSection(IDetailLayoutBuilder& DetailBuilder);
+
+	/** Creates the section where the user can select an actor class for each GDTF in the MVR Scene */
+	void CreateGDTFToActorClassSection(IDetailLayoutBuilder& DetailBuilder);
+
+	/** Called when the Refresh Actors from DMX Library button was clicked */
+	FReply OnRefreshActorsFromDMXLibraryClicked();
+
+	/** Called when a GDTF to Actor Class Group was clicked */
+	FReply OnSelectGDTFToActorClassGroupClicked(UObject* GDTFObject);
+
 	/** Called when a Fixture Patch changed */
 	void OnFixturePatchChanged(const UDMXEntityFixturePatch* FixturePatch);
 
 	/** Called when a Fixture Type changed */
 	void OnFixtureTypeChanged(const UDMXEntityFixtureType* FixtureType);
 
+	/** Called when a sub-level is loaded */
+	void OnMapChange(uint32 MapChangeFlags);
+
+	/** Called when an actor got deleted in editor */
+	void OnActorDeleted(AActor* DeletedActor);
+
 	/** Called before an Actor Class changed in the GDTFToActorClasses member of the Actor */
 	void OnPreEditChangeActorClassInGDTFToActorClasses();
 
 	/** Called after an Actor Class changed in the GDTFToActorClasses member of the Actor */
 	void OnPostEditChangeActorClassInGDTFToActorClasses();
 
-	/** Forces this Details Customization to refresh */
-	void ForceRefresh();
+	/** Returns true if any actor in the current level makes use of the specified GDTF */
+	bool IsAnyActorUsingGDTF(const UDMXImportGDTF* GDTF) const;
+
+	/** Requests this Details Customization to refresh */
+	void RequestRefresh();
 
 	/** Handle to the Actor Class Property in the GDTFToDefaultActorClasses struct */
 	TSharedPtr<IPropertyHandle> DefaultActorClassHandle;
diff --git a/Engine/Plugins/VirtualProduction/DMX/DMXEngine/Source/DMXEditor/Private/Factories/DMXMVRSceneActorFactory.cpp b/Engine/Plugins/VirtualProduction/DMX/DMXEngine/Source/DMXEditor/Private/Factories/DMXMVRSceneActorFactory.cpp
--- a/Engine/Plugins/VirtualProduction/DMX/DMXEngine/Source/DMXEditor/Private/Factories/DMXMVRSceneActorFactory.cpp
+++ b/Engine/Plugins/VirtualProduction/DMX/DMXEngine/Source/DMXEditor/Private/Factories/DMXMVRSceneActorFactory.cpp
@@ -7,8 +7,11 @@
 #include "Library/DMXLibrary.h"
 #include "MVR/DMXMVRSceneActor.h"
 
+#include "EngineUtils.h"
 #include "Components/StaticMeshComponent.h"
 #include "Engine/StaticMesh.h"
+#include "Framework/Notifications/NotificationManager.h"
+#include "Widgets/Notifications/SNotificationList.h"
 
 
 UDMXMVRSceneActorFactory::UDMXMVRSceneActorFactory()
@@ -65,6 +68,11 @@ AActor* UDMXMVRSceneActorFactory::SpawnActor(UObject* Asset, ULevel* InLevel, co
 			return nullptr;
 		}
 
+		if(IsDMXLibraryAlreadySpawned(World, DMXLibrary))
+		{ 
+			NotifyDMXLibraryAlreadySpawned();
+		}
+
 		ADMXMVRSceneActor* MVRSceneActor = World->SpawnActor<ADMXMVRSceneActor>(ADMXMVRSceneActor::StaticClass());
 		if (MVRSceneActor)
 		{
@@ -76,3 +84,29 @@ AActor* UDMXMVRSceneActorFactory::SpawnActor(UObject* Asset, ULevel* InLevel, co
 
 	return nullptr;
 }
+
+bool UDMXMVRSceneActorFactory::IsDMXLibraryAlreadySpawned(UWorld* World, UDMXLibrary* DMXLibrary) const
+{
+	if (World && DMXLibrary)
+	{
+		for (TActorIterator<ADMXMVRSceneActor> It(World, ADMXMVRSceneActor::StaticClass()); It; ++It)
+		{
+			ADMXMVRSceneActor* MVRSceneActor = *It;
+			if (MVRSceneActor && MVRSceneActor->GetDMXLibrary() == DMXLibrary)
+			{
+				return true;
+			}
+		}
+	}
+	return false;
+}
+
+void UDMXMVRSceneActorFactory::NotifyDMXLibraryAlreadySpawned()
+{
+	if (!DMXLibraryAlreadySpawnedNotification.IsValid())
+	{
+		FNotificationInfo NotificationInfo(NSLOCTEXT("DMXMVRSceneActorFactory", "DMXLibraryAlreadySpawnedInSceneNotification", "DMX Library is already spawned in this scene."));
+		NotificationInfo.ExpireDuration = 5.f;
+		DMXLibraryAlreadySpawnedNotification = FSlateNotificationManager::Get().AddNotification(NotificationInfo);
+	}
+}
diff --git a/Engine/Plugins/VirtualProduction/DMX/DMXEngine/Source/DMXEditor/Private/Factories/DMXMVRSceneActorFactory.h b/Engine/Plugins/VirtualProduction/DMX/DMXEngine/Source/DMXEditor/Private/Factories/DMXMVRSceneActorFactory.h
--- a/Engine/Plugins/VirtualProduction/DMX/DMXEngine/Source/DMXEditor/Private/Factories/DMXMVRSceneActorFactory.h
+++ b/Engine/Plugins/VirtualProduction/DMX/DMXEngine/Source/DMXEditor/Private/Factories/DMXMVRSceneActorFactory.h
@@ -9,8 +9,12 @@
 
 #include "DMXMVRSceneActorFactory.generated.h"
 
+class UDMXLibrary;
+
 class AActor;
 struct FAssetData;
+class SNotificationItem;
+class UWorld;
 
 
 UCLASS()
@@ -26,4 +30,14 @@ public:
 	virtual bool CanCreateActorFrom(const FAssetData& AssetData, FText& OutErrorMsg) override;
 	virtual AActor* SpawnActor(UObject* Asset, ULevel* InLevel, const FTransform& Transform, const FActorSpawnParameters& InSpawnParams) override;
 	//~ End UActorFactory Interface
+
+private:
+	/** Returns true if the DMX Library is already spawned as an MVR Scene Actor in the world */
+	bool IsDMXLibraryAlreadySpawned(UWorld* World, UDMXLibrary* DMXLibrary) const;
+
+	/** Creates a notification that the DMX Library is already spawned as an MVR Scene Actor in the world */
+	void NotifyDMXLibraryAlreadySpawned();
+
+	/** Notification displayed when a DMX Library was already spawned */
+	TWeakPtr<SNotificationItem> DMXLibraryAlreadySpawnedNotification;
 };
diff --git a/Engine/Plugins/VirtualProduction/DMX/DMXEngine/Source/DMXRuntime/Private/MVR/DMXMVRSceneActor.cpp b/Engine/Plugins/VirtualProduction/DMX/DMXEngine/Source/DMXRuntime/Private/MVR/DMXMVRSceneActor.cpp
--- a/Engine/Plugins/VirtualProduction/DMX/DMXEngine/Source/DMXRuntime/Private/MVR/DMXMVRSceneActor.cpp
+++ b/Engine/Plugins/VirtualProduction/DMX/DMXEngine/Source/DMXRuntime/Private/MVR/DMXMVRSceneActor.cpp
@@ -13,6 +13,9 @@
 #include "MVR/Types/DMXMVRFixtureNode.h"
 
 #include "DatasmithAssetUserData.h"
+#include "EngineUtils.h"
+#include "Algo/Find.h"
+#include "Algo/RemoveIf.h"
 #include "Components/SceneComponent.h"
 #include "Engine/World.h"
 #include "UObject/UObjectIterator.h"
@@ -122,25 +125,40 @@
 }
 #endif // WITH_EDITOR
 
-void ADMXMVRSceneActor::EnsureMVRUUIDsForRelatedActors()
+UDMXEntityFixturePatch* ADMXMVRSceneActor::GetFixturePatchFromActor(AActor* Actor) const
 {
-	for (const TSoftObjectPtr<AActor>& RelatedActor : RelatedActors)
+	TArray<UDMXComponent*> DMXComponents;
+	Actor->GetComponents<UDMXComponent>(DMXComponents);
+	if (!ensureAlwaysMsgf(!DMXComponents.IsEmpty(), TEXT("'%s' implements the DMXMVRFixtureActorInterface, but has no DMX component. The DMX Component is needed to patch and identify the fixture in the MVR Scene."), *Actor->GetName()))
 	{
-		if (AActor* Actor = RelatedActor.Get())
-		{			
-			const FString MVRFixtureUUID = UDMXMVRAssetUserData::GetMVRAssetUserDataValueForKey(*Actor, UDMXMVRAssetUserData::MVRFixtureUUIDMetaDataKey);
-			if (MVRFixtureUUID.IsEmpty())
-			{	
-				// Try to acquire the MVR Fixture UUID
-				if (UDMXEntityFixturePatch* FixturePatch = GetFixturePatch(Actor))
-				{
-					UDMXMVRAssetUserData::SetMVRAssetUserDataValueForKey(*Actor, UDMXMVRAssetUserData::MVRFixtureUUIDMetaDataKey, FixturePatch->GetMVRFixtureUUID().ToString());
-				}
-			}
-		}
+		return nullptr;
 	}
+	ensureAlwaysMsgf(DMXComponents.Num() == 1, TEXT("'%s' implements the DMXMVRFixtureActorInterface, but has more than one DMX component. A single DMX component is required to clearly identify the fixture by MVR UUID in the MVR Scene."), *Actor->GetName());
+	
+	UDMXEntityFixturePatch* FixturePatch = DMXComponents[0]->GetFixturePatch();
+
+	return FixturePatch;
 }
 
+void ADMXMVRSceneActor::SetFixturePatchOnActor(AActor* Actor, UDMXEntityFixturePatch* FixturePatch)
+{
+	if (!ensureMsgf(Actor && FixturePatch, TEXT("Trying to Set Fixture Patch on Actor, but Actor or Fixture Patch are invalid.")))
+	{
+		return;
+	}
+
+	TArray<UDMXComponent*> DMXComponents;
+	Actor->GetComponents<UDMXComponent>(DMXComponents);
+	if (!ensureAlwaysMsgf(!DMXComponents.IsEmpty(), TEXT("'%s' implements the DMXMVRFixtureActorInterface, but has no DMX component. The DMX Component is needed to patch and identify the fixture in the MVR Scene."), *Actor->GetName()))
+	{
+		return;
+	}
+	ensureAlwaysMsgf(DMXComponents.Num() == 1, TEXT("'%s' implements the DMXMVRFixtureActorInterface, but has more than one DMX component. A single DMX component is required to clearly identify the fixture by MVR UUID in the MVR Scene."), *Actor->GetName());
+
+	DMXComponents[0]->SetFixturePatch(FixturePatch);
+}
+
+
 #if WITH_EDITOR
 void ADMXMVRSceneActor::SetDMXLibrary(UDMXLibrary* NewDMXLibrary)
 {
@@ -155,34 +173,18 @@
 	}
 	DMXLibrary = NewDMXLibrary;
 
+	RefreshFromDMXLibrary();
+}
+#endif // WITH_EDITOR
 
-	const TSharedRef<FDMXMVRFixtureActorLibrary> MVRFixtureActorLibrary = MakeShared<FDMXMVRFixtureActorLibrary>();
-	const TArray<UDMXEntityFixturePatch*> FixturePatches = DMXLibrary->GetEntitiesTypeCast<UDMXEntityFixturePatch>();
-
-	// Build the GDTF to Actor Class Pair array
-	for (UDMXEntityFixturePatch* FixturePatch : FixturePatches)
+#if WITH_EDITOR
+void ADMXMVRSceneActor::RefreshFromDMXLibrary()
+{
+	if (!ensureMsgf(DMXLibrary, TEXT("Trying to update MVR Scene from DMX Library, but DMX Library was never set or no longer exists.")))
 	{
-		UDMXEntityFixtureType* FixtureType = FixturePatch->GetFixtureType();
-		if (FixtureType && FixtureType->DMXImport)
-		{
-			const bool bGDTFToDefaultActorClassPairAlreadyCreated = GDTFToDefaultActorClasses.ContainsByPredicate([FixtureType](const FDMXMVRSceneGDTFToActorClassPair& GDTFToActorClassPair)
-				{
-					return GDTFToActorClassPair.GDTF == FixtureType->DMXImport;
-				});
-			if (bGDTFToDefaultActorClassPairAlreadyCreated)
-			{
-				continue;
-			}
-
-			UClass* ActorClass = MVRFixtureActorLibrary->FindMostAppropriateActorClassForPatch(FixturePatch);
-			FDMXMVRSceneGDTFToActorClassPair GDTFToActorClassPair;
-			GDTFToActorClassPair.GDTF = FixtureType->DMXImport;
-			GDTFToActorClassPair.ActorClass = ActorClass;
-			GDTFToDefaultActorClasses.Add(GDTFToActorClassPair);
-		}
+		return;
 	}
 
-	// Spawn Fixture Actors
 	DMXLibrary->UpdateGeneralSceneDescription();
 	UDMXMVRGeneralSceneDescription* GeneralSceneDescription = DMXLibrary->GetLazyGeneralSceneDescription();
 	if (!GeneralSceneDescription)
@@ -190,8 +192,54 @@
 		return;
 	}
 
+	const UWorld* World = GetWorld();
+	if (!ensureMsgf(World, TEXT("Invalid world when trying to update MVR Scene from DMX Library.")))
+	{
+		return;
+	}
+
+	if (bRespawnDeletedActorsOnRefresh)
+	{
+		DeletedMVRFixtureUUIDs.Reset();
+	}
+
+	const TSharedRef<FDMXMVRFixtureActorLibrary> MVRFixtureActorLibrary = MakeShared<FDMXMVRFixtureActorLibrary>();
+	const TArray<UDMXEntityFixturePatch*> FixturePatches = DMXLibrary->GetEntitiesTypeCast<UDMXEntityFixturePatch>();
+
+	// Destroy actors no longer present in the DMX Library, remember any other spawned actor
+	TMap<UDMXEntityFixturePatch*, AActor*> FixturePatchToSpawnedActorMap;
+	for (const TSoftObjectPtr<AActor>& SoftActorPtr : RelatedActors)
+	{
+		if (!SoftActorPtr.IsValid())
+		{
+			continue;
+		}
+		AActor* Actor = SoftActorPtr.Get();
+
+		UDMXEntityFixturePatch* FixturePatch = GetFixturePatchFromActor(Actor);
+		if (!FixturePatch || !FixturePatches.Contains(FixturePatch))
+		{
+			Actor->Destroy();
+			continue;
+		}
+
+		FixturePatchToSpawnedActorMap.Add(FixturePatch, Actor);
+	}
+
+	// Spawn newly added and if requested, previously deleted actors
 	for (UDMXEntityFixturePatch* FixturePatch : FixturePatches)
 	{
+		if (FixturePatchToSpawnedActorMap.Contains(FixturePatch))
+		{
+			continue;
+		}
+
+		if (!bRespawnDeletedActorsOnRefresh &&
+			DeletedMVRFixtureUUIDs.Contains(FixturePatch->GetMVRFixtureUUID()))
+		{
+			continue;
+		}
+
 		UClass* ActorClass = MVRFixtureActorLibrary->FindMostAppropriateActorClassForPatch(FixturePatch);
 		if (!ActorClass)
 		{
@@ -208,9 +256,133 @@
 		const FTransform Transform = FixtureNode->GetTransformAbsolute();
 		SpawnMVRActor(ActorClass, FixturePatch, Transform);
 	}
+
+	// Update Transforms if requested
+	if (bUpdateTransformsOnRefresh)
+	{
+		for (const TSoftObjectPtr<AActor>& SoftRelatedActor : RelatedActors)
+		{
+			if (!SoftRelatedActor.IsValid())
+			{
+				continue;
+			}
+
+			AActor* RelatedActor = SoftRelatedActor.Get();
+			UDMXEntityFixturePatch* FixturePatch = GetFixturePatchFromActor(RelatedActor);
+			if (!FixturePatch)
+			{
+				continue;
+			}
+
+			const FGuid& MVRFixtureUUID = FixturePatch->GetMVRFixtureUUID();
+			const UDMXMVRFixtureNode* FixtureNode = GeneralSceneDescription->FindFixtureNode(MVRFixtureUUID);
+			if (!FixtureNode)
+			{
+				continue;
+			}
+
+			const FTransform Transform = FixtureNode->GetTransformAbsolute();
+			RelatedActor->SetActorTransform(Transform);
+		}
+	}
+
+	UpdateGDTFToDefaultActorClasses(MVRFixtureActorLibrary);
 }
 #endif // WITH_EDITOR
 
 #if WITH_EDITOR
+TArray<AActor*> ADMXMVRSceneActor::GetActorsSpawnedForGDTF(const UDMXImportGDTF* GDTF) const
+{
+	TArray<AActor*> Result;
+	if (!GDTF)
+	{
+		return Result;
+	}
+
+	for (const TSoftObjectPtr<AActor>& SoftActorPtr : RelatedActors)
+	{
+		if (!SoftActorPtr.IsValid())
+		{
+			continue;
+		}
+		AActor* Actor = SoftActorPtr.Get();
+
+		UDMXEntityFixturePatch* FixturePatch = GetFixturePatchFromActor(Actor);
+		if (!FixturePatch || !FixturePatch->GetFixtureType())
+		{
+			continue;
+		}
+
+		if (FixturePatch->GetFixtureType()->DMXImport == GDTF)
+		{
+			Result.Add(Actor);
+		}
+	}
+
+	return Result;
+}
+#endif // WITH_EDITOR
+
+void ADMXMVRSceneActor::EnsureMVRUUIDsForRelatedActors()
+{
+	for (const TSoftObjectPtr<AActor>& RelatedActor : RelatedActors)
+	{
+		if (AActor* Actor = RelatedActor.Get())
+		{
+			const FString MVRFixtureUUID = UDMXMVRAssetUserData::GetMVRAssetUserDataValueForKey(*Actor, UDMXMVRAssetUserData::MVRFixtureUUIDMetaDataKey);
+			if (MVRFixtureUUID.IsEmpty())
+			{
+				// Try to acquire the MVR Fixture UUID
+				if (UDMXEntityFixturePatch* FixturePatch = GetFixturePatchFromActor(Actor))
+				{
+					UDMXMVRAssetUserData::SetMVRAssetUserDataValueForKey(*Actor, UDMXMVRAssetUserData::MVRFixtureUUIDMetaDataKey, FixturePatch->GetMVRFixtureUUID().ToString());
+				}
+			}
+		}
+	}
+}
+
+void ADMXMVRSceneActor::UpdateGDTFToDefaultActorClasses(const TSharedRef<FDMXMVRFixtureActorLibrary>& MVRFixtureActorLibrary)
+{
+	const TArray<UDMXEntityFixturePatch*> FixturePatches = DMXLibrary->GetEntitiesTypeCast<UDMXEntityFixturePatch>();
+	for (UDMXEntityFixturePatch* FixturePatch : FixturePatches)
+	{
+		const UDMXEntityFixtureType* FixtureType = FixturePatch->GetFixtureType();
+		if (FixtureType && FixtureType->DMXImport)
+		{
+			const FDMXMVRSceneGDTFToActorClassPair* ExistingGDTFToActorClassPairPtr = Algo::FindByPredicate(GDTFToDefaultActorClasses, [FixtureType](const FDMXMVRSceneGDTFToActorClassPair& GDTFToActorClassPair)
+				{
+					return GDTFToActorClassPair.GDTF == FixtureType->DMXImport;
+				});
+			if (ExistingGDTFToActorClassPairPtr)
+			{
+				continue;
+			}
+
+			UClass* ActorClass = MVRFixtureActorLibrary->FindMostAppropriateActorClassForPatch(FixturePatch);
+			FDMXMVRSceneGDTFToActorClassPair GDTFToActorClassPair;
+			GDTFToActorClassPair.GDTF = FixtureType->DMXImport;
+			GDTFToActorClassPair.ActorClass = ActorClass;
+
+			GDTFToDefaultActorClasses.Add(GDTFToActorClassPair);
+		}
+	}
+
+	// Remove those GDTFs that are no longer used by any actor
+	Algo::RemoveIf(GDTFToDefaultActorClasses, [this](const FDMXMVRSceneGDTFToActorClassPair& GDTFToActorClassPair)
+		{
+			const TSoftObjectPtr<AActor>* ActorThatUsesGDTF = Algo::FindByPredicate(RelatedActors, [this, GDTFToActorClassPair](const TSoftObjectPtr<AActor>& Actor)
+				{
+					const UDMXEntityFixturePatch* FixturePatch = Actor.IsValid() ? GetFixturePatchFromActor(Actor.Get()) : nullptr;
+					const UDMXEntityFixtureType* FixtureType = FixturePatch ? FixturePatch->GetFixtureType() : nullptr;
+					const UDMXImport* GDTF = FixtureType ? FixtureType->DMXImport : nullptr;
+					
+					return GDTFToActorClassPair.GDTF == GDTF;
+				});
+			return ActorThatUsesGDTF == nullptr;
+		});
+}
+
+#if WITH_EDITOR
 void ADMXMVRSceneActor::OnMapChange(uint32 MapEventFlags)
 {
@@ -225,13 +397,33 @@
 #if WITH_EDITOR
 void ADMXMVRSceneActor::OnActorDeleted(AActor* DeletedActor)
 {
-	const int32 RelatedActorIndex = RelatedActors.Find(DeletedActor);
-	if (RelatedActorIndex != INDEX_NONE)
+	if (DeletedActor == this)
 	{
-		// This will add this actor to the transaction if there is one currently recording
-		Modify();
+		for (const TSoftObjectPtr<AActor>& RelatedActor : RelatedActors)
+		{
+			if (AActor* Actor = RelatedActor.IsValid() ? RelatedActor.Get() : nullptr)
+			{
+				Actor->Modify();
+				Actor->Destroy();
+			}
+		}
+	}
+	else
+	{
+		const int32 RelatedActorIndex = RelatedActors.Find(DeletedActor);
+		if (RelatedActorIndex != INDEX_NONE)
+		{
+			// This will add this actor to the transaction if there is one currently recording
+			Modify();
 
-		RelatedActors[RelatedActorIndex]->Reset();
+			const UDMXEntityFixturePatch* FixturePatch = GetFixturePatchFromActor(DeletedActor);
+			if (FixturePatch)
+			{
+				DeletedMVRFixtureUUIDs.Add(FixturePatch->GetMVRFixtureUUID());
+			}
+
+			RelatedActors[RelatedActorIndex]->Reset();
+		}
 	}
 }
 #endif // WITH_EDITOR
@@ -299,7 +491,7 @@
 			continue;
 		}
 
-		UDMXEntityFixturePatch* FixturePatch = GetFixturePatch(Actor);
+		UDMXEntityFixturePatch* FixturePatch = GetFixturePatchFromActor(Actor);
 		if (FixturePatch && 
 			FixturePatch->GetFixtureType() && 
 			FixturePatch->GetFixtureType()->DMXImport == GDTFToDefaultActorClasses[IndexOfChangedElement].GDTF)
@@ -310,6 +502,7 @@
 }
 #endif // WITH_EDITOR
 
+#if WITH_EDITOR
 AActor* ADMXMVRSceneActor::SpawnMVRActor(const TSubclassOf<AActor>&ActorClass, UDMXEntityFixturePatch* FixturePatch, const FTransform& Transform, AActor* Template)
 {
 	UWorld* World = GetWorld();
@@ -329,14 +522,22 @@
 		return nullptr;
 	}
 
+	FName ActorName = *FixturePatch->Name;
+	if (!IsGloballyUniqueObjectName(*FixturePatch->Name))
+	{
+		ActorName = MakeUniqueObjectName(World, AActor::StaticClass(), *FixturePatch->Name);
+	}
+
 	FActorSpawnParameters ActorSpawnParameters;
 	ActorSpawnParameters.Template = Template;
 	ActorSpawnParameters.SpawnCollisionHandlingOverride = ESpawnActorCollisionHandlingMethod::AlwaysSpawn;
+	ActorSpawnParameters.Name = ActorName;
 	AActor* NewFixtureActor = World->SpawnActor<AActor>(ActorClass, Transform, ActorSpawnParameters);
 	if (!NewFixtureActor)
 	{
 		return nullptr;
 	}
+	NewFixtureActor->SetActorLabel(ActorName.ToString());
 
 	NewFixtureActor->RegisterAllComponents();
 	USceneComponent* RootComponentOfChildActor = NewFixtureActor->GetRootComponent();
@@ -358,7 +559,7 @@
 	RootComponentOfChildActor->AttachToComponent(MVRSceneRoot, FAttachmentTransformRules::KeepWorldTransform);
 	const FGuid& MVRFixtureUUID = FixturePatch->GetMVRFixtureUUID();
 	UDMXMVRAssetUserData::SetMVRAssetUserDataValueForKey(*NewFixtureActor, UDMXMVRAssetUserData::MVRFixtureUUIDMetaDataKey, MVRFixtureUUID.ToString());
-	SetFixturePatch(NewFixtureActor, FixturePatch);
+	SetFixturePatchOnActor(NewFixtureActor, FixturePatch);
 	RelatedActors.Add(NewFixtureActor);
 
 #if WITH_EDITOR
@@ -368,7 +569,9 @@
 
 	return NewFixtureActor;
 }
+#endif // WITH_EDITOR
 
+#if WITH_EDITOR
 AActor* ADMXMVRSceneActor::ReplaceMVRActor(AActor* ActorToReplace, const TSubclassOf<AActor>& ClassOfNewActor)
 {
 	if (!ensureAlwaysMsgf(ActorToReplace, TEXT("Trying to replace MVR Fixture in MVR Scene, but the Actor to replace is not valid.")))
@@ -393,7 +596,7 @@
 		UDMXEntityFixturePatch* FixturePatch = nullptr;
 		if (IDMXMVRFixtureActorInterface* MVRFixtureActorInterface = Cast<IDMXMVRFixtureActorInterface>(ActorToReplace))
 		{
-			FixturePatch = GetFixturePatch(ActorToReplace);
+			FixturePatch = GetFixturePatchFromActor(ActorToReplace);
 		}
 		
 		if (!FixturePatch)
@@ -428,36 +631,4 @@
 
 	return nullptr;
 }
-
-UDMXEntityFixturePatch* ADMXMVRSceneActor::GetFixturePatch(AActor* Actor) const
-{
-	TArray<UDMXComponent*> DMXComponents;
-	Actor->GetComponents<UDMXComponent>(DMXComponents);
-	if (!ensureAlwaysMsgf(!DMXComponents.IsEmpty(), TEXT("'%s' implements the DMXMVRFixtureActorInterface, but has no DMX component. The DMX Component is needed to patch and identify the fixture in the MVR Scene."), *Actor->GetName()))
-	{
-		return nullptr;
-	}
-	ensureAlwaysMsgf(DMXComponents.Num() == 1, TEXT("'%s' implements the DMXMVRFixtureActorInterface, but has more than one DMX component. A single DMX component is required to clearly identify the fixture by MVR UUID in the MVR Scene."), *Actor->GetName());
-	
-	UDMXEntityFixturePatch* FixturePatch = DMXComponents[0]->GetFixturePatch();
-
-	return FixturePatch;
-}
-
-void ADMXMVRSceneActor::SetFixturePatch(AActor* Actor, UDMXEntityFixturePatch* FixturePatch)
-{
-	if (!ensureMsgf(Actor && FixturePatch, TEXT("Trying to Set Fixture Patch on Actor, but Actor or Fixture Patch are invalid.")))
-	{
-		return;
-	}
-
-	TArray<UDMXComponent*> DMXComponents;
-	Actor->GetComponents<UDMXComponent>(DMXComponents);
-	if (!ensureAlwaysMsgf(!DMXComponents.IsEmpty(), TEXT("'%s' implements the DMXMVRFixtureActorInterface, but has no DMX component. The DMX Component is needed to patch and identify the fixture in the MVR Scene."), *Actor->GetName()))
-	{
-		return;
-	}
-	ensureAlwaysMsgf(DMXComponents.Num() == 1, TEXT("'%s' implements the DMXMVRFixtureActorInterface, but has more than one DMX component. A single DMX component is required to clearly identify the fixture by MVR UUID in the MVR Scene."), *Actor->GetName());
-
-	DMXComponents[0]->SetFixturePatch(FixturePatch);
-}
+#endif // WITH_EDITOR
diff --git a/Engine/Plugins/VirtualProduction/DMX/DMXEngine/Source/DMXRuntime/Public/MVR/DMXMVRSceneActor.h b/Engine/Plugins/VirtualProduction/DMX/DMXEngine/Source/DMXRuntime/Public/MVR/DMXMVRSceneActor.h
--- a/Engine/Plugins/VirtualProduction/DMX/DMXEngine/Source/DMXRuntime/Public/MVR/DMXMVRSceneActor.h
+++ b/Engine/Plugins/VirtualProduction/DMX/DMXEngine/Source/DMXRuntime/Public/MVR/DMXMVRSceneActor.h
@@ -7,6 +7,7 @@
 
 #include "DMXMVRSceneActor.generated.h"
 
+class FDMXMVRFixtureActorLibrary;
 class UDMXEntityFixturePatch;
 class UDMXImportGDTF;
 class UDMXLibrary;
@@ -51,10 +52,15 @@ public:
 #endif
 	//~ End AActor interface
 
-public:
 #if WITH_EDITOR
 	/** Sets the dmx library for this MVR actor. Should only be called once, further calls will have no effect and hit an ensure condition */
 	void SetDMXLibrary(UDMXLibrary* NewDMXLibrary);
+
+	/** Refreshes the scene from the DMX Library */
+	void RefreshFromDMXLibrary();
+
+	/** Returns an array of actors that are currently spawned in the scene for specified GDTF */
+	TArray<AActor*> GetActorsSpawnedForGDTF(const UDMXImportGDTF* GDTF) const;
 #endif
 	/** Returns the DMX Library of this MVR Scene Actor */
 	FORCEINLINE UDMXLibrary* GetDMXLibrary() const { return DMXLibrary; }
@@ -70,10 +76,29 @@ public:
 	static FName GetMVRSceneRootPropertyNameChecked() { return GET_MEMBER_NAME_CHECKED(ADMXMVRSceneActor, MVRSceneRoot); }
 #endif
 
+#if WITH_EDITORONLY_DATA
+	/** If checked, respawns Fixture Actors deleted from the MVR Scene */
+	UPROPERTY(EditAnywhere, Category = "MVR", Meta = (DisplayName = "Respawn Deleted Actors"))
+	bool bRespawnDeletedActorsOnRefresh = false;
+
+	/** If checked, updates transforms from the DMX Library */
+	UPROPERTY(EditAnywhere, Category = "MVR", Meta = (DisplayName = "Update Transforms"))
+	bool bUpdateTransformsOnRefresh = false;
+#endif
+
 private:
 	/** Set MVR UUIDs for related Actors */
 	void EnsureMVRUUIDsForRelatedActors();
 
+	/** Updates the GDTFToDefaultActorClasses map so it only contains those entries present both in the DMX Libraryand the MVR Scene */
+	void UpdateGDTFToDefaultActorClasses(const TSharedRef<FDMXMVRFixtureActorLibrary>& MVRFixtureActorLibrary);
+
+	/** Gets the Fixture Patch of an Actor, or nullptr if the Fixture Patch cannot be retrieved anymore */
+	UDMXEntityFixturePatch* GetFixturePatchFromActor(AActor* Actor) const;
+
+	/** Sets the Fixture Patch on an Actor if possible or fails silently */
+	void SetFixturePatchOnActor(AActor* Actor, UDMXEntityFixturePatch* FixturePatch);
+
 #if WITH_EDITOR
 	/** Called when a sub-level is loaded */
 	void OnMapChange(uint32 MapChangeFlags);
@@ -86,13 +111,13 @@ private:
 
 	/** Replaces related Actors that use the Default Actor Class for the GDTF with an instance of the new Default Actor Class */
 	void HandleDefaultActorClassForGDTFChanged();
-#endif // WITH_EDITOR
 
 	/** Spawns an MVR Actor in this Scene. Returns the newly spawned Actor, or nullptr if no Actor could be spawned. */
 	AActor* SpawnMVRActor(const TSubclassOf<AActor>&ActorClass, UDMXEntityFixturePatch * FixturePatch, const FTransform & Transform, AActor * Template = nullptr);
 
 	/** Replaces an MVR Actor in this Scene with another. Returns the newly spawned Actor, or nullptr if no Actor could be spawned. */
 	AActor* ReplaceMVRActor(AActor* ActorToReplace, const TSubclassOf<AActor>& ClassOfNewActor);
+#endif // WITH_EDITOR
 
 	/** The DMX Library this Scene Actor uses */
 	UPROPERTY(BlueprintReadOnly, VisibleAnywhere, Category = "MVR", Meta = (AllowPrivateAccess = true))
@@ -106,6 +131,10 @@ private:
 	UPROPERTY(EditAnywhere, Category = "MVR", Meta = (DispayName = "Default Actor used for GDTF"))
 	TArray<FDMXMVRSceneGDTFToActorClassPair> GDTFToDefaultActorClasses;
 
+	/** Array holding MVR Fixture UUIDs that were explicitly removed from the scene */
+	UPROPERTY()
+	TArray<FGuid> DeletedMVRFixtureUUIDs;
+
 #if WITH_EDITORONLY_DATA
 	/** The GDTFToDefaultActorClassMap cached of PreEditChange, to find changes. */
 	TArray<FDMXMVRSceneGDTFToActorClassPair> GDTFToDefaultActorClasses_PreEditChange;
@@ -114,11 +143,4 @@ private:
 	/** The root component to which all actors are attached initially */
 	UPROPERTY(VisibleAnywhere, Category = "Actor", AdvancedDisplay, Meta = (AllowPrivateAccess = true))
 	TObjectPtr<USceneComponent> MVRSceneRoot;
-
-private:
-	/** Gets the Fixture Patch of an Actor, or nullptr if the Fixture Patch cannot be retrieved anymore */
-	UDMXEntityFixturePatch* GetFixturePatch(AActor* Actor) const;
-
-	/** Sets the Fixture Patch on an Actor if possible or fails silently */
-	void SetFixturePatch(AActor* Actor, UDMXEntityFixturePatch* FixturePatch);
 };
