DMX - Change default naming scheme so the first Fixture Patch or Fixture Type created uses _000

diff --git a/Engine/Plugins/VirtualProduction/DMX/DMXEngine/Source/DMXRuntime/Private/DMXRuntimeUtils.cpp b/Engine/Plugins/VirtualProduction/DMX/DMXEngine/Source/DMXRuntime/Private/DMXRuntimeUtils.cpp
--- a/Engine/Plugins/VirtualProduction/DMX/DMXEngine/Source/DMXRuntime/Private/DMXRuntimeUtils.cpp
+++ b/Engine/Plugins/VirtualProduction/DMX/DMXEngine/Source/DMXRuntime/Private/DMXRuntimeUtils.cpp
@@ -5,6 +5,9 @@
 #include "Library/DMXLibrary.h"
 #include "Library/DMXEntityFixturePatch.h"
 
+#include "Algo/MaxElement.h"
+#include "Internationalization/Regex.h"
+
 
 #define LOCTEXT_NAMESPACE "FDMXRuntimeUtils"
 
@@ -141,46 +144,96 @@ FString FDMXRuntimeUtils::ConvertTransformToGDTF4x3MatrixString(FTransform Trans
 	return Result;
 }
 
-FString FDMXRuntimeUtils::GenerateUniqueNameFromExisting(const TSet<FString>& InExistingNames, const FString& InBaseName)
+FString FDMXRuntimeUtils::GenerateUniqueNameFromExisting(const TSet<FString>& InExistingNames, const FString& InDesiredName)
 {
-	if (!InBaseName.IsEmpty() && !InExistingNames.Contains(InBaseName))
+	auto GetBaseNameLambda([](const FString InName) -> FString
+		{
+			const FRegexPattern BaseNamePattern(TEXT("(\\w*)_\\d*"));
+			FRegexMatcher Regex(BaseNamePattern, *InName);
+			if (Regex.FindNext())
+			{
+				return Regex.GetCaptureGroup(1);
+			}
+			return InName;
+		});
+
+	auto GetIndexStringLambda([](const FString& InName, FString& OutIndexString) -> bool
 	{
-		return InBaseName;
-	}
+		const FRegexPattern IndexPattern(TEXT("\\w*_(\\d*)"));
+		FRegexMatcher Regex(IndexPattern, *InName);
+		if (Regex.FindNext())
+		{
+			OutIndexString = Regex.GetCaptureGroup(1);
+			return true;
+		}
+		return false;
+	});
 
-	FString FinalName;
-	FString BaseName;
+	FString DesiredName = InDesiredName;
+	DesiredName.RemoveSpacesInline();
+	const FString DesiredBaseName = GetBaseNameLambda(DesiredName);
 
-	int32 Index = 0;
-	if (InBaseName.IsEmpty())
+	TSet<FString> ConflictingNames;
+	ConflictingNames.Reserve(InExistingNames.Num());
+	for (const FString& OtherName : InExistingNames)
 	{
-		BaseName = TEXT("Default name");
+		const FString OtherBaseName = GetBaseNameLambda(OtherName);
+		if (OtherBaseName.Equals(DesiredBaseName))
+		{
+			ConflictingNames.Add(OtherName);
+		}
 	}
-	else
+
+	TSet<FString> ExistingIndexStrings;
+	ExistingIndexStrings.Reserve(ConflictingNames.Num());
+	for (const FString& ExistingName : ConflictingNames)
 	{
-		// If there's an index at the end of the name, start from there
-		FDMXRuntimeUtils::GetNameAndIndexFromString(InBaseName, BaseName, Index);
+		FString IndexString;
+		if (GetIndexStringLambda(ExistingName, IndexString))
+		{
+			ExistingIndexStrings.Add(IndexString);
+		}
 	}
 
-	int32 Count = (Index == 0) ? 1 : Index;
-	FinalName = BaseName;
-	// Add Count to the BaseName, increasing Count, until it's a non-existent name
-	do
-	{
-		// Calculate the number of digits in the number, adding 2 (1 extra to correctly count digits, another to account for the '_' that will be added to the name
-		int32 CountLength = Count > 0 ? (int32)FGenericPlatformMath::LogX(10.0f, Count) + 2 : 2;
+	int32 NewIndex = 0;
+	const FString* const MaxIndexStringPtr = Algo::MaxElementBy(ExistingIndexStrings, [](const FString& IndexString)
+		{
+			int32 Index = INDEX_NONE;
+			LexTryParseString(Index, *IndexString);
+			return Index;
+		});
 
-		// If the length of the final string will be too long, cut off the end so we can fit the number
-		if (CountLength + BaseName.Len() >= NAME_SIZE)
+	if (MaxIndexStringPtr &&
+		LexTryParseString(NewIndex, **MaxIndexStringPtr))
+	{
+		// Increment the trailing index string
+		NewIndex++;
+		FString NewIndexString = FString::FromInt(NewIndex);
+		while(NewIndexString.Len() < 3)
 		{
-			BaseName = BaseName.Left(NAME_SIZE - CountLength - 1);
+			NewIndexString = FString(TEXT("0")) + NewIndexString;
 		}
 
-		FinalName = FString::Printf(TEXT("%s_%d"), *BaseName, Count);
-		++Count;
-	} while (InExistingNames.Contains(FinalName));
+		const FRegexPattern BaseNamePattern(TEXT("(\\w*)_\\d*"));
+		FRegexMatcher Regex(BaseNamePattern, *DesiredName);
 
-	return FinalName;
+		const FString BaseName = Regex.FindNext() ? Regex.GetCaptureGroup(1) : DesiredName;
+		return BaseName + TEXT("_") + NewIndexString;
+	}
+	else
+	{
+		// Use the alreay appended index
+		FString DesiredIndexString;
+		if (GetIndexStringLambda(DesiredName, DesiredIndexString) &&
+			LexTryParseString(NewIndex, *DesiredIndexString))
+		{
+			return DesiredName;
+		}
+		else
+		{
+			return DesiredName.EndsWith(TEXT("_")) ? DesiredName + FString(TEXT("000")) : DesiredName + FString(TEXT("_000"));
+		}
+	}
 }
 
 FString FDMXRuntimeUtils::FindUniqueEntityName(const UDMXLibrary* InLibrary, TSubclassOf<UDMXEntity> InEntityClass, const FString& InBaseName /*= TEXT("")*/)
