DMX - Allow to filter universes in patch view by typing a range. Example: 'Uni 1-3, 5' displays all patches in Universe 1, 2, 3 and 5

diff --git a/Engine/Plugins/VirtualProduction/DMX/DMXEngine/Source/DMXEditor/Private/Widgets/FixturePatch/SDMXMVRFixtureListToolbar.cpp b/Engine/Plugins/VirtualProduction/DMX/DMXEngine/Source/DMXEditor/Private/Widgets/FixturePatch/SDMXMVRFixtureListToolbar.cpp
--- a/Engine/Plugins/VirtualProduction/DMX/DMXEngine/Source/DMXEditor/Private/Widgets/FixturePatch/SDMXMVRFixtureListToolbar.cpp
+++ b/Engine/Plugins/VirtualProduction/DMX/DMXEngine/Source/DMXEditor/Private/Widgets/FixturePatch/SDMXMVRFixtureListToolbar.cpp
@@ -26,46 +26,85 @@
 
 namespace UE::DMXRuntime::SDMXMVRFixtureListToolbar::Private
 {
-	UE_NODISCARD bool ParseUniverse(const FString& InputString, int32& OutUniverse)
+	TArray<int32> ParseUniverses(const FString& InputString)
 	{
+		TArray<int32> OutUniverses;
+
 		// Try to match addresses formating, e.g. '1.', '1:' etc.
-		static const TCHAR* ParamDelimiters[] =
+		static const TCHAR* UniverseAddressParamDelimiters[] =
 		{
 			TEXT("."),
 			TEXT(","),
 			TEXT(":"),
 			TEXT(";")
 		};
-
-		TArray<FString> ValueStringArray;
-		constexpr bool bParseEmpty = true;
-		const FString InputStringWithSpace = InputString + TEXT(" "); // So ValueStringArray will be lenght of 2 if Address is empty, e.g. '0.'
-		InputStringWithSpace.ParseIntoArray(ValueStringArray, ParamDelimiters, 4, bParseEmpty);
-		if (ValueStringArray.Num() == 2)
+		if (InputString.EndsWith(UniverseAddressParamDelimiters[0]) ||
+			InputString.EndsWith(UniverseAddressParamDelimiters[1]) ||
+			InputString.EndsWith(UniverseAddressParamDelimiters[2]) ||
+			InputString.EndsWith(UniverseAddressParamDelimiters[3]))
 		{
-			if (LexTryParseString<int32>(OutUniverse, *ValueStringArray[0]))
+			TArray<FString> UniverseAddressStringArray;
+
+			constexpr bool bCullEmpty = true;
+			InputString.ParseIntoArray(UniverseAddressStringArray, UniverseAddressParamDelimiters, 4, bCullEmpty);
+			if (UniverseAddressStringArray.Num() == 1)
 			{
-				return true;
+				int32 Universe;
+				if (LexTryParseString(Universe, *UniverseAddressStringArray[0]))
+				{
+					OutUniverses.Add(Universe);
+					return OutUniverses;
+				}
 			}
 		}
 
-		// Try to match strings starting with Uni, e.g. 'Uni 1', 'Universe 1', 'Universe1'
-		if (InputString.StartsWith(TEXT("Uni")))
+		// Try to match strings starting with Uni, e.g. 'Uni 1', 'uni1', 'Uni 1, 2', 'universe 1', 'Universe 1, 2 - 3, 4'
+		const FRegexPattern UniversesPattern(TEXT("^(?:universe|uni)\\s*(.*)"), ERegexPatternFlags::CaseInsensitive);
+		FRegexMatcher Regex(UniversesPattern, *InputString);
+		if (Regex.FindNext())
 		{
-			const FRegexPattern SequenceOfDigitsPattern(TEXT("^[^\\d]*(\\d+)"));
-			FRegexMatcher Regex(SequenceOfDigitsPattern, *InputString);
-			if (Regex.FindNext())
+			FString UniversesString = Regex.GetCaptureGroup(1);
+			UniversesString.RemoveSpacesInline();
+
+			static const TCHAR* UniversesDelimiter[] =
+			{
+				TEXT(",")
+			};
+			TArray<FString> UniveresStringArray;
+
+			constexpr bool bCullEmpty = true;
+			UniversesString.ParseIntoArray(UniveresStringArray, UniversesDelimiter, 1, bCullEmpty);
+			for (const FString& UniversesSubstring : UniveresStringArray)
 			{
-				const FString UniverseString = Regex.GetCaptureGroup(1);
-				if (LexTryParseString<int32>(OutUniverse, *UniverseString))
+				static const TCHAR* UniverseRangeDelimiter[] =
 				{
-					return true;
+					TEXT("-")
+				};
+
+				TArray<FString> UniverseRangeStringArray;
+				UniversesSubstring.ParseIntoArray(UniverseRangeStringArray, UniverseRangeDelimiter, 1, bCullEmpty);
+
+				int32 UniverseStart;
+				int32 UniverseEnd;
+				int32 Universe;
+				if (UniverseRangeStringArray.Num() == 2 &&
+					LexTryParseString(UniverseStart, *UniverseRangeStringArray[0]) &&
+					LexTryParseString(UniverseEnd, *UniverseRangeStringArray[1]) &&
+					UniverseStart < UniverseEnd)
+				{
+					for (Universe = UniverseStart; Universe <= UniverseEnd; Universe++)
+					{
+						OutUniverses.Add(Universe);
+					}
+				}
+				else if (LexTryParseString(Universe, *UniversesSubstring))
+				{
+					OutUniverses.Add(Universe);
 				}
 			}
 		}
 
-		OutUniverse = -1;
-		return false;
+		return OutUniverses;
 	}
 
 	UE_NODISCARD bool ParseAddress(const FString& InputString, int32& OutAddress)
@@ -165,6 +204,7 @@ void SDMXMVRFixtureListToolbar::Construct(const FArguments& InArgs, TWeakPtr<FDM
 					SNew(SSearchBox)
 					.MinDesiredWidth(400.f)
 					.OnTextChanged(this, &SDMXMVRFixtureListToolbar::OnSearchTextChanged)
+					.ToolTipText(LOCTEXT("SearchBarTooltip", "Examples:\n\n* PatchName\n* FixtureTypeName\n* SomeMode\n* 1.\n* 1.1\n* Universe 1\n* Uni 1-3\n* Uni 1, 3\n* Uni 1, 4-5'."))
 				]
 									
 				+ SWrapBox::Slot()
@@ -219,22 +259,24 @@ TArray<TSharedPtr<FDMXMVRFixtureListItem>> SDMXMVRFixtureListToolbar::FilterItem
 		return Result;
 	}
 
-	int32 Universe;
-	if (ParseUniverse(SearchString, Universe))
+	const TArray<int32> Universes = ParseUniverses(SearchString);
+	if(!Universes.IsEmpty())
 	{
-		Result.RemoveAll([Universe](const TSharedPtr<FDMXMVRFixtureListItem>& Item)
+		Result.RemoveAll([Universes](const TSharedPtr<FDMXMVRFixtureListItem>& Item)
 			{
-				return Item->GetUniverse() != Universe;
+				return !Universes.Contains(Item->GetUniverse());
 			});
 
-		int32 Address;
-		if (ParseAddress(SearchString, Address))
-		{
-			Result.RemoveAll([Address](const TSharedPtr<FDMXMVRFixtureListItem>& Item)
-				{
-					return Item->GetAddress() != Address;
-				});
-		}
+		return Result;
+	}
+
+	int32 Address;
+	if (ParseAddress(SearchString, Address))
+	{
+		Result.RemoveAll([Address](const TSharedPtr<FDMXMVRFixtureListItem>& Item)
+			{
+				return Item->GetAddress() != Address;
+			});
 
 		return Result;
 	}
@@ -401,8 +443,10 @@ void SDMXMVRFixtureListToolbar::OnAddNewMVRFixtureClicked(UDMXEntity* InSelected
 		});
 	int32 Universe = 1;
 	int32 Address = 1;
+	int32 ChannelSpan = 0;
 	if (LastFixturePatchPtr && !FixtureType->Modes.IsEmpty())
 	{
+		ChannelSpan = FixtureType->Modes[0].ChannelSpan;
 		const UDMXEntityFixturePatch& LastFixturePatch = **LastFixturePatchPtr;
 		if (LastFixturePatch.GetEndingChannel() + FixtureType->Modes[0].ChannelSpan > DMX_MAX_ADDRESS)
 		{
@@ -435,7 +479,7 @@ void SDMXMVRFixtureListToolbar::OnAddNewMVRFixtureClicked(UDMXEntity* InSelected
 		NewWeakFixturePatches.Add(NewFixturePatch);	
 
 		// Increment Universe and Address in steps by one. This is enough for auto assign while keeping order of the named patches
-		if (Address + FixtureType->Modes[0].ChannelSpan > DMX_MAX_ADDRESS)
+		if (Address + ChannelSpan > DMX_MAX_ADDRESS)
 		{
 			Universe++;
 			Address = 1; 
