DMX - Add MVR Import options

diff --git a/Engine/Plugins/VirtualProduction/DMX/DMXEngine/Source/DMXEditor/Private/Customizations/DMXLibraryFromMVRImportOptionsDetails.cpp b/Engine/Plugins/VirtualProduction/DMX/DMXEngine/Source/DMXEditor/Private/Customizations/DMXLibraryFromMVRImportOptionsDetails.cpp
new file mode 100644
--- /dev/null
+++ b/Engine/Plugins/VirtualProduction/DMX/DMXEngine/Source/DMXEditor/Private/Customizations/DMXLibraryFromMVRImportOptionsDetails.cpp
@@ -0,0 +1,98 @@
+// Copyright Epic Games, Inc. All Rights Reserved.
+
+#include "Customizations/DMXLibraryFromMVRImportOptionsDetails.h"
+
+#include "DMXProtocolSettings.h"
+#include "Factories/DMXLibraryFromMVRImporter.h"
+#include "Factories/DMXLibraryFromMVRImportOptions.h"
+#include "IO/DMXPortManager.h"
+#include "Library/DMXEntityFixturePatch.h"
+#include "Library/DMXLibrary.h"
+
+#include "DetailCategoryBuilder.h"
+#include "DetailLayoutBuilder.h"
+#include "DetailWidgetRow.h"
+#include "Misc/Paths.h"
+#include "Widgets/Layout/SWrapBox.h"
+#include "Widgets/Text/STextBlock.h"
+
+
+#define LOCTEXT_NAMESPACE "DMXLibraryFromMVROptionsDetails"
+
+TSharedRef<IDetailCustomization> FDMXLibraryFromMVRImportOptionsDetails::MakeInstance()
+{
+	return MakeShared<FDMXLibraryFromMVRImportOptionsDetails>();
+}
+
+void FDMXLibraryFromMVRImportOptionsDetails::CustomizeDetails(IDetailLayoutBuilder& DetailBuilder)
+{
+	TArray<TWeakObjectPtr<UObject>> ObjectsBeingCustomized;
+	DetailBuilder.GetObjectsBeingCustomized(ObjectsBeingCustomized);
+	if (!ensureAlwaysMsgf(ObjectsBeingCustomized.Num() == 1 || !ObjectsBeingCustomized[0].IsValid(), TEXT("Multi-Editing MVR Import Options is not supported, or invalid customized Object.")))
+	{
+		return;
+	}
+
+	ImportOptions = CastChecked<UDMXLibraryFromMVRImportOptions>(ObjectsBeingCustomized[0]);
+	if (!ensureAlwaysMsgf(FPaths::FileExists(ImportOptions->Filename), TEXT("Invalid MVR File '%s'."), *ImportOptions->Filename))
+	{
+		return;
+	}
+
+	IDetailCategoryBuilder& PortsCategoryBuilder = DetailBuilder.EditCategory("Ports");
+	if (!HasAnyInputPorts())
+	{
+		const TSharedRef<IPropertyHandle> InputPortToUpdateHandle = DetailBuilder.GetProperty(GET_MEMBER_NAME_CHECKED(UDMXLibraryFromMVRImportOptions, InputPortToUpdate));
+		InputPortToUpdateHandle->MarkHiddenByCustomization();
+
+		// Show 'Generate Input Port' instead of the default property name 'Update Input Port' if there's no port available. 
+		const TSharedRef<IPropertyHandle> UpdateInputPortHandle = DetailBuilder.GetProperty(GET_MEMBER_NAME_CHECKED(UDMXLibraryFromMVRImportOptions, bUpdateInputPort));
+		UpdateInputPortHandle->MarkHiddenByCustomization();
+
+		PortsCategoryBuilder.AddCustomRow(FText::GetEmpty())
+			.NameContent()
+			[
+				SNew(STextBlock)
+				.Font(DetailBuilder.GetDetailFont())
+				.Text(LOCTEXT("GeneratePortDisplayName", "Generate Input Port"))
+			]
+			.ValueContent()
+			[
+				UpdateInputPortHandle->CreatePropertyValueWidget()
+			];
+	}
+
+	if (!HasAnyOutputPorts())
+	{
+		const TSharedRef<IPropertyHandle> OutputPortToUpdateHandle = DetailBuilder.GetProperty(GET_MEMBER_NAME_CHECKED(UDMXLibraryFromMVRImportOptions, OutputPortToUpdate));
+		OutputPortToUpdateHandle->MarkHiddenByCustomization();
+
+		// Show 'Generate Output Port' instead of the default property name 'Update Output Port' if there's no port available. 
+		const TSharedRef<IPropertyHandle> UpdateOutputPortHandle = DetailBuilder.GetProperty(GET_MEMBER_NAME_CHECKED(UDMXLibraryFromMVRImportOptions, bUpdateOutputPort));
+		UpdateOutputPortHandle->MarkHiddenByCustomization();
+
+		PortsCategoryBuilder.AddCustomRow(FText::GetEmpty())
+			.NameContent()
+			[
+				SNew(STextBlock)
+				.Font(DetailBuilder.GetDetailFont())
+				.Text(LOCTEXT("GeneratePortDisplayName", "Generate Output Port"))
+			]
+			.ValueContent()
+			[
+				UpdateOutputPortHandle->CreatePropertyValueWidget()
+			];
+	}
+}
+
+bool FDMXLibraryFromMVRImportOptionsDetails::HasAnyInputPorts() const
+{
+	return !FDMXPortManager::Get().GetInputPorts().IsEmpty();
+}
+
+bool FDMXLibraryFromMVRImportOptionsDetails::HasAnyOutputPorts() const
+{
+	return !FDMXPortManager::Get().GetOutputPorts().IsEmpty();
+}
+
+#undef LOCTEXT_NAMESPACE
diff --git a/Engine/Plugins/VirtualProduction/DMX/DMXEngine/Source/DMXEditor/Private/Customizations/DMXLibraryFromMVRImportOptionsDetails.h b/Engine/Plugins/VirtualProduction/DMX/DMXEngine/Source/DMXEditor/Private/Customizations/DMXLibraryFromMVRImportOptionsDetails.h
new file mode 100644
--- /dev/null
+++ b/Engine/Plugins/VirtualProduction/DMX/DMXEngine/Source/DMXEditor/Private/Customizations/DMXLibraryFromMVRImportOptionsDetails.h
@@ -0,0 +1,39 @@
+// Copyright Epic Games, Inc. All Rights Reserved.
+
+#pragma once
+
+#include "CoreMinimal.h"
+#include "IDetailCustomization.h"
+
+class UDMXLibrary;
+class UDMXLibraryFromMVRImportOptions;
+
+class IPropertyUtilities;
+
+
+/** Details customization fro the MVR Scene Actor Factory Import Options */
+class FDMXLibraryFromMVRImportOptionsDetails final
+	: public IDetailCustomization
+{
+public:
+	/** Creates an instance of this details customization */
+	static TSharedRef<IDetailCustomization> MakeInstance();
+
+protected:
+	//~ Begin IDetailCustomization interface
+	virtual void CustomizeDetails(IDetailLayoutBuilder& DetailBuilder) override;
+	//~ End IDetailCustomization interface
+
+private:
+	/** Returns true if the project has any input ports defined */
+	bool HasAnyInputPorts() const;
+
+	/** Returns true if the project has any output ports defined */
+	bool HasAnyOutputPorts() const;
+
+	/** For comparison when merging, create a temporary DMX Library from the imported MVR before actually merging changes when reimporting */
+	UDMXLibrary* TempDMXLibrary = nullptr;
+
+	/** Import Options Object */
+	UDMXLibraryFromMVRImportOptions* ImportOptions = nullptr;
+};
diff --git a/Engine/Plugins/VirtualProduction/DMX/DMXEngine/Source/DMXEditor/Private/Factories/DMXLibraryFromMVRFactory.cpp b/Engine/Plugins/VirtualProduction/DMX/DMXEngine/Source/DMXEditor/Private/Factories/DMXLibraryFromMVRFactory.cpp
--- a/Engine/Plugins/VirtualProduction/DMX/DMXEngine/Source/DMXEditor/Private/Factories/DMXLibraryFromMVRFactory.cpp
+++ b/Engine/Plugins/VirtualProduction/DMX/DMXEngine/Source/DMXEditor/Private/Factories/DMXLibraryFromMVRFactory.cpp
@@ -2,65 +2,23 @@
 
 #include "Factories/DMXLibraryFromMVRFactory.h"
 
-#include "DMXEditorLog.h"
-#include "DMXEditorModule.h"
 #include "DMXEditorSettings.h"
-#include "DMXEditorUtils.h"
-#include "DMXInitializeFixtureTypeFromGDTFHelper.h"
-#include "DMXZipper.h"
-#include "Factories/DMXGDTFFactory.h"
-#include "Library/DMXEntityFixturePatch.h"
-#include "Library/DMXEntityFixtureType.h"
-#include "Library/DMXGDTFAssetImportData.h"
-#include "Library/DMXImportGDTF.h"
+#include "DMXEditorLog.h"
+#include "Factories/DMXLibraryFromMVRImporter.h"
+#include "Factories/DMXLibraryFromMVRImportOptions.h"
 #include "Library/DMXLibrary.h"
 #include "MVR/DMXMVRAssetImportData.h"
 #include "MVR/DMXMVRGeneralSceneDescription.h"
-#include "MVR/Types/DMXMVRFixtureNode.h"
 
-#include "Algo/Find.h"
+#include "Editor.h"
+#include "EditorReimportHandler.h"
 #include "AssetRegistry/AssetRegistryModule.h"
-#include "AssetToolsModule.h"
-#include "ObjectTools.h"
-#include "XmlFile.h"
-#include "EditorFramework/AssetImportData.h"
-#include "Misc/FileHelper.h"
-#include "Misc/MessageDialog.h"
+#include "Misc/Paths.h"
+#include "Subsystems/ImportSubsystem.h"
 
 
 #define LOCTEXT_NAMESPACE "DMXLibraryFromMVRFactory"
 
-namespace UE::DMX::DMXLibraryFromMVRFactory::Private
-{
-	/** Finds the file corresponding to the GDTF Spec in an MVR Zip file. Returns true if the file was found.  */
-	bool FindGDTFFilenameInMVRZip(const TSharedRef<FDMXZipper>& MVRZip, const FString& GDTFSpec, FString& OutGDTFFilename)
-	{
-		OutGDTFFilename.Reset();
-
-		// Some softwares store the GDTF spec with, some without the .gdtf extension. 
-		TArray<FString> Filenames = MVRZip->GetFiles();
-		const FString* FilenamePtr = Algo::FindByPredicate(Filenames, [GDTFSpec](const FString& Filename)
-			{
-				if (GDTFSpec == Filename)
-				{
-					return true;
-				}
-				else if (GDTFSpec + TEXT(".gdtf") == Filename)
-				{
-					return true;
-				}
-				return false;
-			});
-
-		if (FilenamePtr)
-		{
-			OutGDTFFilename = *FilenamePtr;
-			return true;
-		}
-		return false;
-	}
-}
-
 const FString UDMXLibraryFromMVRFactory::MVRFileExtension = TEXT("MVR");
 
 UDMXLibraryFromMVRFactory::UDMXLibraryFromMVRFactory()
@@ -72,9 +30,8 @@ UDMXLibraryFromMVRFactory::UDMXLibraryFromMVRFactory()
 	Formats.Add(TEXT("mvr;My Virtual Rig"));
 }
 
-UObject* UDMXLibraryFromMVRFactory::FactoryCreateFile(UClass* InClass, UObject* Parent, FName InName, EObjectFlags Flags, const FString& InFilename, const TCHAR* Parms, FFeedbackContext* Warn, bool& bOutOperationCanceled)
+UObject* UDMXLibraryFromMVRFactory::FactoryCreateFile(UClass* InClass, UObject* Parent, FName InName, EObjectFlags Flags, const FString& InFilename, const TCHAR* Parms, FFeedbackContext* Warn, bool& bOutOperationCancelled)
 {
-	bOutOperationCanceled = false;
 	CurrentFilename = InFilename;
 
 	if (!FPaths::FileExists(InFilename))
@@ -86,32 +43,29 @@ UObject* UDMXLibraryFromMVRFactory::FactoryCreateFile(UClass* InClass, UObject*
 	check(DMXEditorSettings);
 	DMXEditorSettings->LastMVRImportPath = FPaths::GetPath(InFilename);
 
-	UDMXLibrary* NewDMXLibrary = CreateDMXLibraryAsset(Parent, InName, Flags, InFilename);
-	if (!NewDMXLibrary)
+	FDMXLibraryFromMVRImporter MVRImporter;
+	if (!MVRImporter.LoadMVRFile(InFilename))
 	{
-		UE_LOG(LogDMXEditor, Error, TEXT("Failed to create DMX Library for MVR '%s'."), *InFilename);
 		return nullptr;
 	}
 
-	const TSharedRef<FDMXZipper> Zip = MakeShared<FDMXZipper>();
-	if (!Zip->LoadFromFile(InFilename))
+	UDMXLibraryFromMVRImportOptions* ImportOptions = NewObject<UDMXLibraryFromMVRImportOptions>();
+	MVRImporter.UpdateImportOptionsFromModalWindow(ImportOptions, bOutOperationCancelled);
+	if (ImportOptions->bCancelled)
 	{
-		UE_LOG(LogDMXEditor, Error, TEXT("Cannot read MVR '%s'. File is not a valid MVR."), *InFilename);
 		return nullptr;
 	}
 
-	UDMXMVRGeneralSceneDescription* GeneralSceneDescription = CreateGeneralSsceneDescription(*NewDMXLibrary, Zip);
-	if (!GeneralSceneDescription)
-	{
-		UE_LOG(LogDMXEditor, Error, TEXT("Cannot read General Scene Description from MVR '%s'. File is not a valid MVR."), *InFilename);
-		return nullptr;
-	}
+	UImportSubsystem* ImportSubsystem = GEditor->GetEditorSubsystem<UImportSubsystem>();
+	ImportSubsystem->BroadcastAssetPreImport(this, UDMXLibrary::StaticClass(), Parent, InName, *MVRFileExtension);
+
+	UDMXLibrary* NewDMXLibrary = nullptr;
+	TArray<UDMXImportGDTF*> GDTFs;
+	MVRImporter.Import(Parent, InName, Flags | RF_Public, ImportOptions, NewDMXLibrary, GDTFs, bOutOperationCancelled);
 
-	UDMXMVRAssetImportData* MVRAssetImportData = GeneralSceneDescription->GetMVRAssetImportData();
-	MVRAssetImportData->SetSourceFile(InFilename);
+	ImportSubsystem->BroadcastAssetPostImport(this, NewDMXLibrary);
 
-	TArray<UDMXImportGDTF*> GDTFs = CreateGDTFAssets(Parent, Flags, Zip, *GeneralSceneDescription);
-	InitDMXLibrary(NewDMXLibrary, GDTFs, GeneralSceneDescription);
+	ImportOptions->ApplyOptions(NewDMXLibrary);
 
 	return NewDMXLibrary;
 }
@@ -163,7 +117,7 @@ EReimportResult::Type UDMXLibraryFromMVRFactory::Reimport(UObject* Obj)
 	{
 		return EReimportResult::Failed;
 	}
-	const FString SourceFilename = MVRAssetImportData->GetFilePathAndName();
+	FString SourceFilename = MVRAssetImportData->GetFilePathAndName();
 
 	UDMXLibrary* DMXLibrary = Cast<UDMXLibrary>(Obj);
 	if (!ensureMsgf(DMXLibrary, TEXT("Invalid DMX Library when trying to reimport DMX Library.")))
@@ -171,315 +125,68 @@ EReimportResult::Type UDMXLibraryFromMVRFactory::Reimport(UObject* Obj)
 		return EReimportResult::Failed;
 	}
 
-	if (!DMXLibrary->GetEntities().IsEmpty())
-	{
-		const FText MessageText = LOCTEXT("MVRReimportDialog", "DMX Library already contains data. Reimporting MVR will clear existing data. Do you want to proceed?");
-		if (FMessageDialog::Open(EAppMsgType::YesNo, MessageText) == EAppReturnType::No)
-		{
-			return EReimportResult::Cancelled;
-		}
-	}
-
-	const TSharedRef<FDMXZipper> Zip = MakeShared<FDMXZipper>();
-	if (!Zip->LoadFromFile(SourceFilename))
-	{
-		UE_LOG(LogDMXEditor, Error, TEXT("Cannot read MVR '%s'. File is not a valid MVR."), *SourceFilename);
-		return EReimportResult::Failed;
-	}
-
-	UDMXMVRGeneralSceneDescription* GeneralSceneDescription = CreateGeneralSsceneDescription(*DMXLibrary, Zip);
-	if (!GeneralSceneDescription)
+	FDMXLibraryFromMVRImporter MVRImporter;
+	if (!MVRImporter.LoadMVRFile(SourceFilename))
 	{
-		UE_LOG(LogDMXEditor, Error, TEXT("Cannot read General Scene Description from MVR '%s'. File is not a valid MVR."), *SourceFilename);
 		return EReimportResult::Failed;
 	}
 
-	// Since a new General Scene Description was created, the source file has to be set anew, 
-	MVRAssetImportData->SetSourceFile(SourceFilename);
-
-	TArray<UDMXImportGDTF*> GDTFs = CreateGDTFAssets(DMXLibrary->GetOuter(), RF_NoFlags, Zip, *GeneralSceneDescription);
-	InitDMXLibrary(DMXLibrary, GDTFs, GeneralSceneDescription);
-
-	return EReimportResult::Succeeded;
-}
-
-int32 UDMXLibraryFromMVRFactory::GetPriority() const
-{
-	return ImportPriority;
-}
-
-UDMXLibrary* UDMXLibraryFromMVRFactory::CreateDMXLibraryAsset(UObject* Parent, const FName& Name, EObjectFlags Flags, const FString& InFilename)
-{
-	UImportSubsystem* ImportSubsystem = GEditor->GetEditorSubsystem<UImportSubsystem>();
-	ImportSubsystem->BroadcastAssetPreImport(this, UDMXLibrary::StaticClass(), Parent, Name, *MVRFileExtension);
-
-	UDMXLibrary* NewDMXLibrary = NewObject<UDMXLibrary>(Parent, Name, Flags | RF_Public);
-	if (!NewDMXLibrary)
-	{
-		return nullptr;
-	}
-
-	ImportSubsystem->BroadcastAssetPostImport(this, NewDMXLibrary);
-	
-	return NewDMXLibrary;
-}
+	UDMXLibraryFromMVRImportOptions* ImportOptions = NewObject<UDMXLibraryFromMVRImportOptions>();
+	ImportOptions->ImportIntoDMXLibrary = CastChecked<UDMXLibrary>(Obj);
+	ImportOptions->bIsReimport = true;
+	ImportOptions->bCreateNewDMXLibrary = false;
 
-UDMXMVRGeneralSceneDescription* UDMXLibraryFromMVRFactory::CreateGeneralSsceneDescription(UDMXLibrary& DMXLibrary, const TSharedRef<FDMXZipper>& Zip) const
-{
-	TArray64<uint8> XMLData;
-	if (!Zip->GetFileContent(TEXT("GeneralSceneDescription.xml"), XMLData))
+	bool bOperationCancelled;
+	MVRImporter.UpdateImportOptionsFromModalWindow(ImportOptions, bOperationCancelled);
+	if (bOperationCancelled)
 	{
-		return nullptr;
+		return EReimportResult::Cancelled;
 	}
 
-	// MVR implicitly adpots UTF-8 encoding of Xml Files by adopting the GDTF standard (DIN-15800).
-	// Content is NOT null-terminated; we need to specify lengths here.
-	const FUTF8ToTCHAR TCHARData(reinterpret_cast<const ANSICHAR*>(XMLData.GetData()), XMLData.Num());
-
-	const TSharedRef<FXmlFile> GeneralSceneDescriptionXml = MakeShared<FXmlFile>(FString(TCHARData.Length(), TCHARData.Get()), EConstructMethod::ConstructFromBuffer);
-
-	const FName GeneralSceneDescriptionName = FName(GetName() + TEXT("_MVRGeneralSceneDescription"));
-	UDMXMVRGeneralSceneDescription* GeneralSceneDescription = UDMXMVRGeneralSceneDescription::CreateFromXmlFile(GeneralSceneDescriptionXml, &DMXLibrary, GeneralSceneDescriptionName);
-	
-	return GeneralSceneDescription;
-}
-
-TArray<UDMXImportGDTF*> UDMXLibraryFromMVRFactory::CreateGDTFAssets(UObject* Parent, EObjectFlags Flags, const TSharedRef<FDMXZipper>& Zip, const UDMXMVRGeneralSceneDescription& GeneralSceneDescription)
-{
-	const FString Path = FPaths::GetPath(Parent->GetName()) + TEXT("/GDTFs");
-
-	TArray<FAssetData> ExistingGDTFAssets;
-	FAssetRegistryModule& AssetRegistryModule = FModuleManager::LoadModuleChecked<FAssetRegistryModule>(TEXT("AssetRegistry"));
-	AssetRegistryModule.Get().GetAssetsByPath(FName(Path), ExistingGDTFAssets);
-
-	TArray<UDMXImportGDTF*> ExistingGDTFs;
-	for (const FAssetData& AssetData : ExistingGDTFAssets)
-	{
-		if (UDMXImportGDTF* GDTF = Cast<UDMXImportGDTF>(AssetData.GetAsset()))
-		{
-			ExistingGDTFs.AddUnique(GDTF);
-		}
-	}
+	if (ImportOptions->bCreateNewDMXLibrary)
+	{		
+		// Case where the user wants to create a new DMX Library instead of reimporting
+		const FName NewDMXLibraryName = MakeUniqueObjectName(DMXLibrary->GetOuter(), UDMXLibrary::StaticClass(), *DMXLibrary->GetName(), EUniqueObjectNameOptions::GloballyUnique);
 
-	UDMXGDTFFactory* GDTFFactory = NewObject<UDMXGDTFFactory>();
-	TArray<UDMXImportGDTF*> ImportedGDTFs;
-	TArray<FString> ImportedGDTFNames;
-	if (ExistingGDTFs.Num() > 0)
-	{
-		const FText MessageText = LOCTEXT("MVRImportReimportsGDTFDialog", "MVR contains existing GDTFs. Do you want to reimport the existing GDTF assets?");
-		if (FMessageDialog::Open(EAppMsgType::YesNo, MessageText) == EAppReturnType::Yes)
-		{
-			for (UDMXImportGDTF* GDTFAsset : ExistingGDTFs)
-			{
-				if (GDTFFactory->Reimport(GDTFAsset) == EReimportResult::Succeeded)
-				{
-					ImportedGDTFs.Add(GDTFAsset);
-
-					if (UDMXGDTFAssetImportData* GDTFAssetImportData = GDTFAsset->GetGDTFAssetImportData())
-					{
-						const FString SourceFilename = FPaths::GetCleanFilename(GDTFAssetImportData->GetFilePathAndName());
-						ImportedGDTFNames.Add(SourceFilename);
-					}
-				}
-			}
-		}
-	}
+		UImportSubsystem* ImportSubsystem = GEditor->GetEditorSubsystem<UImportSubsystem>();
+		ImportSubsystem->BroadcastAssetPreImport(this, UDMXLibrary::StaticClass(), DMXLibrary->GetOuter(), NewDMXLibraryName, *MVRFileExtension);
 
-	// Import GDTF Assets that aren't yet imported in the reimport procedure
-	TArray<UDMXMVRFixtureNode*> FixtureNodes;
-	GeneralSceneDescription.GetFixtureNodes(FixtureNodes);
-	for (const UDMXMVRFixtureNode* FixtureNode : FixtureNodes)
-	{
-		// Don't import the same GDTF twice
-		if (!ImportedGDTFNames.Contains(FixtureNode->GDTFSpec))
-		{
-			using namespace UE::DMX::DMXLibraryFromMVRFactory::Private;
-
-			FString GDTFFilename;
-			if (!FindGDTFFilenameInMVRZip(Zip, FixtureNode->GDTFSpec, GDTFFilename))
-			{
-				UE_LOG(LogDMXEditor, Warning, TEXT("Cannot find a GDTF file that corresponds to GDTF Spec '%s' of Fixture Node '%s'. Ignoring Fixture."), *FixtureNode->GDTFSpec, *FixtureNode->Name)
-				{
-					continue;
-				}
-			}
-
-			const FDMXZipper::FDMXScopedUnzipToTempFile ScopedUnzipGDTF(Zip, GDTFFilename);
-			if (!ScopedUnzipGDTF.TempFilePathAndName.IsEmpty())
-			{
-				constexpr bool bRemovePathFromDesiredName = true;
-				const FString BaseFileName = FPaths::GetBaseFilename(ScopedUnzipGDTF.TempFilePathAndName, bRemovePathFromDesiredName);
-				FString AssetName = ObjectTools::SanitizeObjectName(BaseFileName);
-
-				FString GDTFPackageName;
-				FString GDTFAssetName;
-				FAssetToolsModule& AssetToolsModule = FModuleManager::LoadModuleChecked<FAssetToolsModule>("AssetTools");
-				AssetToolsModule.Get().CreateUniqueAssetName(Path / AssetName, TEXT(""), GDTFPackageName, GDTFAssetName);
-
-				UPackage* Package = CreatePackage(*GDTFPackageName);
-				check(Package);
-				Package->FullyLoad();
-
-				bool bCanceled;
-				UObject* NewGDTFObject = GDTFFactory->FactoryCreateFile(UDMXImportGDTF::StaticClass(), Package, FName(*FixtureNode->GDTFSpec), Flags | RF_Public, ScopedUnzipGDTF.TempFilePathAndName, nullptr, GWarn, bCanceled);
-				
-				if (UDMXImportGDTF* NewGDTF = Cast<UDMXImportGDTF>(NewGDTFObject))
-				{
-					ImportedGDTFs.Add(NewGDTF);
-					ImportedGDTFNames.Add(FixtureNode->GDTFSpec);
-
-					FAssetRegistryModule::AssetCreated(NewGDTF);
-					Package->MarkPackageDirty();
-				}
-			}
-		}
-	}
+		UDMXLibrary* NewDMXLibrary = DuplicateObject(DMXLibrary, DMXLibrary->GetOuter());
+		NewDMXLibrary->Rename(*NewDMXLibraryName.ToString());
 
-	return ImportedGDTFs;
-}
+		TArray<UDMXImportGDTF*> GDTFs;
+		MVRImporter.Reimport(DMXLibrary, ImportOptions, GDTFs);
 
-void UDMXLibraryFromMVRFactory::InitDMXLibrary(UDMXLibrary* DMXLibrary, const TArray<UDMXImportGDTF*>& GDTFAssets, UDMXMVRGeneralSceneDescription* GeneralSceneDescription) const
-{
-	if (!ensureAlwaysMsgf(DMXLibrary, TEXT("Trying to initialize a DMX Library from MVR, but the DMX Library is not valid.")))
-	{
-		return;
-	}
-	if (!ensureAlwaysMsgf(GeneralSceneDescription, TEXT("Trying to initialize a DMX Library '%s' from MVR, but the DMX Library is not valid."), *DMXLibrary->GetName()))
-	{
-		return;
-	}
-	DMXLibrary->SetMVRGeneralSceneDescription(GeneralSceneDescription);
+		ImportSubsystem->BroadcastAssetPostImport(this, NewDMXLibrary);
+		FAssetRegistryModule::AssetCreated(NewDMXLibrary);
 
-	// Get or create a Fixture Type for each GDTF
-	TArray<UDMXEntityFixtureType*> FixtureTypes = DMXLibrary->GetEntitiesTypeCast<UDMXEntityFixtureType>();
-	TMap<FString, UDMXEntityFixtureType*> GDTFSpecToFixtureTypeMap;
-	TMap<UDMXEntityFixtureType*, FLinearColor> FixtureTypeToColorMap;
-	for (UDMXImportGDTF* GDTF : GDTFAssets)
-	{
-		UDMXGDTFAssetImportData* GDTFAssetImportData = GDTF->GetGDTFAssetImportData();
-		if (!GDTFAssetImportData)
-		{
-			continue;
-		}
-		const FString GDTFSourceFilename = GDTFAssetImportData->GetFilePathAndName();
-		const FString GDTFFilename = FPaths::GetCleanFilename(GDTFSourceFilename);
-
-		UDMXEntityFixtureType** ExistingFixtureTypePtr = Algo::FindByPredicate(FixtureTypes, [GDTF](const UDMXEntityFixtureType* FixtureType)
-			{
-				return FixtureType->DMXImport == GDTF;
-			});
-		if (ExistingFixtureTypePtr)
-		{
-			GDTFSpecToFixtureTypeMap.Add(GDTFFilename, *ExistingFixtureTypePtr);
-		}
-		else
-		{
-			FDMXEntityFixtureTypeConstructionParams FixtureTypeConstructionParams;
-			FixtureTypeConstructionParams.DMXCategory = FDMXFixtureCategory(FDMXFixtureCategory::GetFirstValue());
-			FixtureTypeConstructionParams.ParentDMXLibrary = DMXLibrary;
-
-			UDMXEntityFixtureType* NewFixtureType = UDMXEntityFixtureType::CreateFixtureTypeInLibrary(FixtureTypeConstructionParams, FPaths::GetBaseFilename(GDTFFilename));
-			const bool bAdvancedImportSuccess = FDMXInitializeFixtureTypeFromGDTFHelper::GenerateModesFromGDTF(*NewFixtureType, *GDTF);
-			if (!bAdvancedImportSuccess)
-			{
-				UE_LOG(LogDMXEditor, Warning, TEXT("Failed to initialize Fixture Type '%s', falling back to legacy method that doesn't support matrix fixtures."), *NewFixtureType->GetName());
-				NewFixtureType->SetModesFromDMXImport(GDTF);
-			}
-			NewFixtureType->DMXImport = GDTF;
-
-			GDTFSpecToFixtureTypeMap.Add(GDTFFilename, NewFixtureType);
-		}
-	}
+		ImportOptions->ApplyOptions(NewDMXLibrary);
 
-	// Remove Fixture Patches that have no corresponding Fixture Node in the General Scene Description
-	TArray<UDMXMVRFixtureNode*> FixtureNodes;
-	GeneralSceneDescription->GetFixtureNodes(FixtureNodes);
-	TArray<UDMXEntityFixturePatch*> FixturePatches = DMXLibrary->GetEntitiesTypeCast<UDMXEntityFixturePatch>();
-	for (UDMXEntityFixturePatch* FixturePatch : FixturePatches)
-	{
-		UDMXMVRFixtureNode** CorrespondingFixtureNodePtr = Algo::FindByPredicate(FixtureNodes, [FixturePatch](const UDMXMVRFixtureNode* FixtureNode)
-			{
-				return FixturePatch->GetMVRFixtureUUID() == FixtureNode->UUID;
-			});
-		if (!CorrespondingFixtureNodePtr)
-		{
-			checkNoEntry();
-			UDMXEntityFixturePatch::RemoveFixturePatchFromLibrary(FDMXEntityFixturePatchRef(FixturePatch));
-		}
+		// Return 'Cancelled' so the reimport succes notification doesn't show
+		return EReimportResult::Cancelled;
 	}
-	FixturePatches = DMXLibrary->GetEntitiesTypeCast<UDMXEntityFixturePatch>();
-
-	// Update or create new Fixture Patches for the MVR Fixtures
-	for (const UDMXMVRFixtureNode* FixtureNode : FixtureNodes)
+	else
 	{
-		// Find the GDTF Filename, irregardless if it has the .gdtf extension
-		FString GDTFFilename = FixtureNode->GDTFSpec;
-		if (!GDTFSpecToFixtureTypeMap.Contains(GDTFFilename))
-		{
-			GDTFFilename = FixtureNode->GDTFSpec + TEXT(".gdtf");
-		}
-		if (!GDTFSpecToFixtureTypeMap.Contains(GDTFFilename))
-		{
-			continue;
+		if (ImportOptions->ImportIntoDMXLibrary &&
+			ImportOptions->ImportIntoDMXLibrary != DMXLibrary)
+		{		
+			// Case where user selected a different DMX Library to reimport into
+			DMXLibrary = ImportOptions->ImportIntoDMXLibrary;
 		}
 
-		UDMXEntityFixtureType* FixtureType = GDTFSpecToFixtureTypeMap[GDTFFilename];
-		if (!FixtureTypeToColorMap.Contains(FixtureType))
-		{
-			FLinearColor FixtureTypeColor = FLinearColor::MakeRandomColor();
+		// Reimport
+		TArray<UDMXImportGDTF*> GDTFs;
+		MVRImporter.Reimport(DMXLibrary, ImportOptions, GDTFs);
 
-			// Avoid dominant red values for a bit more of a professional feel
-			if (FixtureTypeColor.R > 0.75f)
-			{
-				FixtureTypeColor.R = FMath::Abs(FixtureTypeColor.R - 1.0f);
-			}
+		ImportOptions->ApplyOptions(DMXLibrary);
 
-			FixtureTypeToColorMap.Add(FixtureType, FixtureTypeColor);
-		}
-
-		int32 ActiveModeIndex = FixtureType->Modes.IndexOfByPredicate([FixtureNode](const FDMXFixtureMode& Mode)
-			{
-				return Mode.ModeName == FixtureNode->GDTFMode;
-			});
-
-		if (ActiveModeIndex != INDEX_NONE)
-		{
-			UDMXEntityFixturePatch* const* ExistingFixturePatchPtr = Algo::FindByPredicate(FixturePatches, [FixtureNode](const UDMXEntityFixturePatch* FixturePatch)
-				{
-					return FixturePatch->GetMVRFixtureUUID() == FixtureNode->UUID;
-				});
-			if (ExistingFixturePatchPtr)
-			{			
-				// Update the existing patch (matching MVR UUID) if possible
-				UDMXEntityFixturePatch& ExistingFixturePatch = **ExistingFixturePatchPtr;
-
-				ExistingFixturePatch.SetFixtureType(FixtureType);
-				ExistingFixturePatch.SetActiveModeIndex(ActiveModeIndex);
-				ExistingFixturePatch.SetUniverseID(FixtureNode->GetUniverseID());
-				ExistingFixturePatch.SetStartingChannel(FixtureNode->GetStartingChannel());
-				ExistingFixturePatch.EditorColor = FixtureTypeToColorMap.FindChecked(FixtureType);
-			}
-			else
-			{
-				// Create a new patch
-				FDMXEntityFixturePatchConstructionParams FixturePatchConstructionParams;
-				FixturePatchConstructionParams.ActiveMode = ActiveModeIndex;
-				FixturePatchConstructionParams.FixtureTypeRef = FDMXEntityFixtureTypeRef(FixtureType);
-				FixturePatchConstructionParams.UniverseID = FixtureNode->GetUniverseID();
-				FixturePatchConstructionParams.StartingAddress = FixtureNode->GetStartingChannel();
-				FixturePatchConstructionParams.MVRFixtureUUID = FixtureNode->UUID;
-
-				UDMXEntityFixturePatch* FixturePatch = UDMXEntityFixturePatch::CreateFixturePatchInLibrary(FixturePatchConstructionParams, FixtureNode->Name);
-				FixturePatch->EditorColor = FixtureTypeToColorMap.FindChecked(FixtureType);
-			}
-		}
-		else
-		{
-			UE_LOG(LogDMXEditor, Warning, TEXT("Skipped creating a Fixture Patch for '%s', cannot find Mode '%s' in '%s'."), *FixtureNode->Name, *FixtureNode->GDTFMode, *FixtureNode->GDTFSpec)
-		}
+		return EReimportResult::Succeeded;
 	}
+}
 
-	DMXLibrary->UpdateGeneralSceneDescription();
+int32 UDMXLibraryFromMVRFactory::GetPriority() const
+{
+	return ImportPriority;
 }
 
 UDMXMVRAssetImportData* UDMXLibraryFromMVRFactory::GetMVRAssetImportData(UObject* DMXLibraryObject) const
diff --git a/Engine/Plugins/VirtualProduction/DMX/DMXEngine/Source/DMXEditor/Private/Factories/DMXLibraryFromMVRImportOptions.cpp b/Engine/Plugins/VirtualProduction/DMX/DMXEngine/Source/DMXEditor/Private/Factories/DMXLibraryFromMVRImportOptions.cpp
new file mode 100644
--- /dev/null
+++ b/Engine/Plugins/VirtualProduction/DMX/DMXEngine/Source/DMXEditor/Private/Factories/DMXLibraryFromMVRImportOptions.cpp
@@ -0,0 +1,120 @@
+// Copyright Epic Games, Inc. All Rights Reserved.
+
+#include "Factories/DMXLibraryFromMVRImportOptions.h"
+
+#include "IO/DMXInputPort.h"
+#include "IO/DMXInputPortConfig.h"
+#include "IO/DMXPortManager.h"
+#include "IO/DMXOutputPort.h"
+#include "IO/DMXOutputPortConfig.h"
+#include "Library/DMXEntityFixturePatch.h"
+#include "Library/DMXLibrary.h"
+
+#include "Algo/Find.h"
+#include "Algo/MaxElement.h"
+#include "Algo/MinElement.h"
+
+
+void UDMXLibraryFromMVRImportOptions::ApplyOptions(UDMXLibrary* DMXLibrary)
+{
+	if (!ensureAlwaysMsgf(DMXLibrary, TEXT("Trying to apply MVR Import Options, but DMX Library is invalid.")))
+	{
+		return;
+	}
+
+	if (!ensureAlwaysMsgf(!bCancelled, TEXT("Trying to apply MVR Import Options, but import was Cancelled while the options were presented.")))
+	{
+		return;
+	}
+
+	const TArray<UDMXEntityFixturePatch*> TempFixturePatches = DMXLibrary->GetEntitiesTypeCast<UDMXEntityFixturePatch>();
+	const UDMXEntityFixturePatch* const* FirstPatchPtr = Algo::MinElementBy(TempFixturePatches, [](const UDMXEntityFixturePatch* FixturePatch)
+		{
+			return FixturePatch->GetUniverseID();
+		});
+	const UDMXEntityFixturePatch* const* LastPatchPtr = Algo::MaxElementBy(TempFixturePatches, [](const UDMXEntityFixturePatch* FixturePatch)
+		{
+			return FixturePatch->GetUniverseID();
+		});
+
+	if (!FirstPatchPtr && !LastPatchPtr)
+	{
+		return;
+	}
+	const int32 FirstUniverse = (*FirstPatchPtr)->GetUniverseID();
+	const int32 NumUniverses = FMath::Max(1, FMath::Clamp((*LastPatchPtr)->GetUniverseID() - FirstUniverse + 1, 0, TNumericLimits<int32>::Max()));
+
+	if (bUpdateInputPort)
+	{
+		const TSharedPtr<FDMXInputPort> InputPort = FDMXPortManager::Get().FindInputPortByGuid(InputPortToUpdate.GetPortGuid());
+		if (InputPort.IsValid())
+		{
+			UDMXProtocolSettings* ProtocolSettings = GetMutableDefault<UDMXProtocolSettings>();
+			FDMXInputPortConfig* InputPortConfigPtr = Algo::FindByPredicate(ProtocolSettings->InputPortConfigs, [&InputPort](const FDMXInputPortConfig& InputPortConfig)
+				{
+					return InputPortConfig.GetPortGuid() == InputPort->GetPortGuid();
+				});
+			if (ensureAlwaysMsgf(InputPortConfigPtr, TEXT("Cannot find config of Input Port '%' in Project Settings."), *InputPort->GetPortName()))
+			{
+				FDMXInputPortConfig& PortConfig = *InputPortConfigPtr;
+
+				FDMXInputPortConfigParams ConfigParams(PortConfig);
+				ConfigParams.LocalUniverseStart = FMath::Min(PortConfig.GetLocalUniverseStart(), FirstUniverse);
+				ConfigParams.NumUniverses = FMath::Max(PortConfig.GetNumUniverses(), NumUniverses);
+				
+				PortConfig = FDMXInputPortConfig(PortConfig.GetPortGuid(), ConfigParams);
+			}
+		}
+		else
+		{
+			FDMXInputPortConfig NewConfig;
+			FDMXInputPortConfigParams ConfigParams(NewConfig);
+			ConfigParams.LocalUniverseStart = FMath::Min(NewConfig.GetLocalUniverseStart(), FirstUniverse);
+			ConfigParams.NumUniverses = FMath::Max(NewConfig.GetNumUniverses(), NumUniverses);
+
+			UDMXProtocolSettings* ProtocolSettings = GetMutableDefault<UDMXProtocolSettings>();
+			ProtocolSettings->PreEditChange(UDMXProtocolSettings::StaticClass()->FindPropertyByName(GET_MEMBER_NAME_CHECKED(UDMXProtocolSettings, InputPortConfigs)));
+			ProtocolSettings->InputPortConfigs.Add(NewConfig);
+			ProtocolSettings->PostEditChange();
+		}
+
+		FDMXPortManager::Get().UpdateFromProtocolSettings();
+	}
+
+	if (bUpdateOutputPort)
+	{
+		const TSharedPtr<FDMXOutputPort> OutputPort = FDMXPortManager::Get().FindOutputPortByGuid(OutputPortToUpdate.GetPortGuid());
+		if (OutputPort.IsValid())
+		{
+			UDMXProtocolSettings* ProtocolSettings = GetMutableDefault<UDMXProtocolSettings>();
+			FDMXOutputPortConfig* OutputPortConfigPtr = Algo::FindByPredicate(ProtocolSettings->OutputPortConfigs, [&OutputPort](const FDMXOutputPortConfig& OutputPortConfig)
+				{
+					return OutputPortConfig.GetPortGuid() == OutputPort->GetPortGuid();
+				});
+			if (ensureAlwaysMsgf(OutputPortConfigPtr, TEXT("Cannot find config of Output Port '%' in Project Settings."), *OutputPort->GetPortName()))
+			{
+				FDMXOutputPortConfig& PortConfig = *OutputPortConfigPtr;
+
+				FDMXOutputPortConfigParams ConfigParams(PortConfig);
+				ConfigParams.LocalUniverseStart = FMath::Min(PortConfig.GetLocalUniverseStart(), FirstUniverse);
+				ConfigParams.NumUniverses = FMath::Max(PortConfig.GetNumUniverses(), NumUniverses);
+
+				PortConfig = FDMXOutputPortConfig(PortConfig.GetPortGuid(), ConfigParams);
+			}
+		}
+		else
+		{
+			FDMXOutputPortConfig NewConfig;
+			FDMXOutputPortConfigParams ConfigParams(NewConfig);
+			ConfigParams.LocalUniverseStart = FMath::Min(NewConfig.GetLocalUniverseStart(), FirstUniverse);
+			ConfigParams.NumUniverses = FMath::Max(NewConfig.GetNumUniverses(), NumUniverses);
+
+			UDMXProtocolSettings* ProtocolSettings = GetMutableDefault<UDMXProtocolSettings>();
+			ProtocolSettings->PreEditChange(UDMXProtocolSettings::StaticClass()->FindPropertyByName(GET_MEMBER_NAME_CHECKED(UDMXProtocolSettings, OutputPortConfigs)));
+			ProtocolSettings->OutputPortConfigs.Add(NewConfig);
+			ProtocolSettings->PostEditChange();
+		}
+
+		FDMXPortManager::Get().UpdateFromProtocolSettings();
+	}
+}
diff --git a/Engine/Plugins/VirtualProduction/DMX/DMXEngine/Source/DMXEditor/Private/Factories/DMXLibraryFromMVRImportOptions.h b/Engine/Plugins/VirtualProduction/DMX/DMXEngine/Source/DMXEditor/Private/Factories/DMXLibraryFromMVRImportOptions.h
new file mode 100644
--- /dev/null
+++ b/Engine/Plugins/VirtualProduction/DMX/DMXEngine/Source/DMXEditor/Private/Factories/DMXLibraryFromMVRImportOptions.h
@@ -0,0 +1,70 @@
+// Copyright Epic Games, Inc. All Rights Reserved.
+
+#pragma once
+
+#include "CoreMinimal.h"
+#include "UObject/Object.h"
+
+#include "IO/DMXInputPortReference.h"
+#include "IO/DMXOutputPortReference.h"
+
+#include "DMXLibraryFromMVRImportOptions.generated.h"
+
+class UDMXLibrary;
+
+
+/** 
+ * Import Options when importing an MVR File as DMX Library.
+ *
+ * Note:
+ * - A MVR File Path and Name is expected to be set.
+ * - bCancelled has to be set to true when the user cancels import while Options here are presented.
+ */
+UCLASS(Config = EditorPerProjectUserSettings)
+class UDMXLibraryFromMVRImportOptions
+    : public UObject
+{
+	GENERATED_BODY()
+
+public:
+	/** Applies the options to the project (e.g. port changes). */
+	void ApplyOptions(UDMXLibrary* DMXLibrary);
+
+	/** If true, show options for Reimport */
+	UPROPERTY()
+	bool bIsReimport = false;
+
+	/** Set to true if the user decided to cancel the import while being presented Options here */
+	bool bCancelled = false;
+
+	/** File Path and Name of the MVR file for which these options are displayed */
+	FString Filename;
+
+	/** If checked, creates a new DMX Library */
+	UPROPERTY(EditAnywhere, Category = "Reimport")
+	bool bCreateNewDMXLibrary = true;
+
+	/** Imports the MVR into the specified DMX Library. Note this will not merge, but clear Patches in the DMX Library. */
+	UPROPERTY(EditAnywhere, Category = "Reimport", Meta = (DisplayName = "Import into DMX Library", EditCondition = "!bCreateNewDMXLibrary", EditConditionHides))
+	TObjectPtr<UDMXLibrary> ImportIntoDMXLibrary;
+
+	/** If checked, reimports GDTF that were previously imported */
+	UPROPERTY(EditAnywhere, Category = "Reimport")
+	bool bReimportExisitingGDTFs = true;
+
+	/** If checked, updates the specified Input Port to span the Universe range of all Patches in the MVR file, or creates a new Input Port if none exist. */
+	UPROPERTY(Config, EditAnywhere, Category = "Ports")
+	bool bUpdateInputPort = true;
+
+	/** The Input Port which is updated */
+	UPROPERTY(EditAnywhere, Category = "Ports", Meta = (ShowOnlyInnerProperties, EditCondition = "bUpdateInputPort", EditConditionHides))
+	FDMXInputPortReference InputPortToUpdate;
+
+	/** If checked, updates the specified Ouput Port to span the Universe range of all Patches in the MVR file, or creates a new Output Port if none exist. */
+	UPROPERTY(Config, EditAnywhere, Category = "Ports")
+	bool bUpdateOutputPort = true;
+
+	/** The Output Port which is updated */
+	UPROPERTY(EditAnywhere, Category = "Ports", Meta = (ShowOnlyInnerProperties, EditCondition = "bUpdateOutputPort", EditConditionHides))
+	FDMXOutputPortReference OutputPortToUpdate;
+};
diff --git a/Engine/Plugins/VirtualProduction/DMX/DMXEngine/Source/DMXEditor/Private/Factories/DMXLibraryFromMVRImporter.cpp b/Engine/Plugins/VirtualProduction/DMX/DMXEngine/Source/DMXEditor/Private/Factories/DMXLibraryFromMVRImporter.cpp
new file mode 100644
--- /dev/null
+++ b/Engine/Plugins/VirtualProduction/DMX/DMXEngine/Source/DMXEditor/Private/Factories/DMXLibraryFromMVRImporter.cpp
@@ -0,0 +1,463 @@
+// Copyright Epic Games, Inc. All Rights Reserved.
+
+#include "Factories/DMXLibraryFromMVRImporter.h"
+
+#include "DMXEditorLog.h"
+#include "DMXInitializeFixtureTypeFromGDTFHelper.h"
+#include "DMXZipper.h"
+#include "Factories/DMXGDTFFactory.h"
+#include "Factories/DMXLibraryFromMVRImportOptions.h"
+#include "Library/DMXEntityFixtureType.h"
+#include "Library/DMXEntityFixturePatch.h"
+#include "Library/DMXGDTFAssetImportData.h"
+#include "Library/DMXImportGDTF.h"
+#include "Library/DMXLibrary.h"
+#include "MVR/DMXMVRAssetImportData.h"
+#include "MVR/DMXMVRGeneralSceneDescription.h"
+#include "MVR/Types/DMXMVRFixtureNode.h"
+#include "Widgets/SDMXLibraryFromMVRImportOptions.h"
+
+#include "AssetToolsModule.h"
+#include "ObjectTools.h"
+#include "XmlFile.h"
+#include "Algo/Find.h"
+#include "AssetRegistry/AssetRegistryModule.h"
+#include "Framework/Application/SlateApplication.h"
+#include "Interfaces/IMainFrameModule.h"
+#include "Layout/SlateRect.h"
+#include "Misc/MessageDialog.h"
+
+
+#define LOCTEXT_NAMESPACE "DMXLibraryFromMVRImporter"
+
+namespace UE::DMX::DMXLibraryFromMVRImporter::Private
+{
+	/** Finds the file corresponding to the GDTF Spec in an MVR Zip file. Returns true if the file was found.  */
+	UE_NODISCARD bool FindGDTFFilenameInMVRZip(const TSharedRef<FDMXZipper>& MVRZip, const FString& GDTFSpec, FString& OutGDTFFilename)
+	{
+		OutGDTFFilename.Reset();
+
+		// Some softwares store the GDTF spec with, some without the .gdtf extension. 
+		TArray<FString> Filenames = MVRZip->GetFiles();
+		const FString* FilenamePtr = Algo::FindByPredicate(Filenames, [GDTFSpec](const FString& Filename)
+			{
+				if (GDTFSpec == Filename)
+				{
+					return true;
+				}
+				else if (GDTFSpec + TEXT(".gdtf") == Filename)
+				{
+					return true;
+				}
+				return false;
+			});
+
+		if (FilenamePtr)
+		{
+			OutGDTFFilename = *FilenamePtr;
+			return true;
+		}
+		return false;
+	}
+
+	/** Window that displays the MVR import options */
+	class FDMXLibraryFromMVROptionsWindow
+	{
+	public:
+		/** Constructor */
+		FDMXLibraryFromMVROptionsWindow(UDMXLibraryFromMVRImportOptions* InOptions)
+		{
+			if (!ensureAlwaysMsgf(InOptions, TEXT("Cannot create MVR import options window, options are invalid.")))
+			{
+				return;
+			}
+
+			TSharedPtr<SWindow> ParentWindow;
+			if (FModuleManager::Get().IsModuleLoaded("MainFrame"))
+			{
+				IMainFrameModule& MainFrame = FModuleManager::LoadModuleChecked<IMainFrameModule>("MainFrame");
+				ParentWindow = MainFrame.GetParentWindow();
+			}
+
+			const FVector2D WindowSize = FVector2D(512.f, 338.f);
+
+			const FSlateRect WorkAreaRect = FSlateApplicationBase::Get().GetPreferredWorkArea();
+			const FVector2D DisplayTopLeft(WorkAreaRect.Left, WorkAreaRect.Top);
+			const FVector2D DisplaySize(WorkAreaRect.Right - WorkAreaRect.Left, WorkAreaRect.Bottom - WorkAreaRect.Top);
+
+			const FVector2D WindowPosition = (DisplayTopLeft + (DisplaySize - WindowSize) / 2.0f);
+			
+			const FText Caption = InOptions->bIsReimport ?
+				LOCTEXT("MVRSceneActorFactoryOptionsWindowCaption", "MVR Reimport Options") :
+				LOCTEXT("MVRSceneActorFactoryOptionsWindowCaption", "MVR Import Options");
+
+			const TSharedRef<SWindow> Window = SNew(SWindow)
+				.Title(Caption)
+				.SizingRule(ESizingRule::FixedSize)
+				.AutoCenter(EAutoCenter::None)
+				.ClientSize(WindowSize)
+				.ScreenPosition(WindowPosition);
+
+			Window->SetContent
+			(
+				SNew(SDMXLibraryFromMVRImportOptions, Window, InOptions)
+				.MaxWidgetSize(WindowPosition)
+			);
+
+			FSlateApplication::Get().AddModalWindow(Window, ParentWindow, false);
+		}
+	};
+}
+
+bool FDMXLibraryFromMVRImporter::LoadMVRFile(const FString& InFilename)
+{
+	Filename = InFilename;
+	if (!ensureAlwaysMsgf(FPaths::FileExists(InFilename), TEXT("Cannot import MVR File '%s'. File does not exist.")))
+	{
+		return false;
+	}
+
+	// Unzip MVR
+	Zip = MakeShared<FDMXZipper>();
+	if (!Zip->LoadFromFile(InFilename))
+	{
+		UE_LOG(LogDMXEditor, Error, TEXT("Cannot read MVR '%s'. File is not a valid MVR."), *InFilename);
+		return false;
+	}
+
+	// Create the General Scene Description
+	GeneralSceneDescription = CreateGeneralSsceneDescription(GetTransientPackage());
+	if (!GeneralSceneDescription)
+	{
+		UE_LOG(LogDMXEditor, Error, TEXT("Cannot read General Scene Description from MVR '%s'. File is not a valid MVR."), *Filename);
+		return false;
+	}
+
+	// Remember source data
+	UDMXMVRAssetImportData* MVRAssetImportData = GeneralSceneDescription->GetMVRAssetImportData();
+	MVRAssetImportData->SetSourceFile(Filename);
+
+	return true;
+}
+
+void FDMXLibraryFromMVRImporter::UpdateImportOptionsFromModalWindow(UDMXLibraryFromMVRImportOptions* InOutImportOptions, bool& bOutCancelled)
+{
+	if (!ensureAlwaysMsgf(InOutImportOptions, TEXT("Invalid Import Options when trying to display an Import Options Window.")))
+	{
+		return;
+	}
+
+	InOutImportOptions->Filename = Filename;
+
+	// The Modal Window blocks the Game Thread and allows users to edit the InOutImportOptions
+	using namespace UE::DMX::DMXLibraryFromMVRImporter::Private;
+	FDMXLibraryFromMVROptionsWindow OptionsWindow(InOutImportOptions);
+
+	bOutCancelled = InOutImportOptions->bCancelled;
+}
+
+void FDMXLibraryFromMVRImporter::Import(UObject* InOuter, const FName& InName, EObjectFlags InFlags, UDMXLibraryFromMVRImportOptions* InImportOptions, UDMXLibrary*& OutDMXLibrary, TArray<UDMXImportGDTF*>& OutGDTFs, bool& bOutCancelled)
+{
+	check(GeneralSceneDescription);
+
+	OutDMXLibrary = nullptr;
+	OutGDTFs.Reset();
+
+	UDMXLibrary* NewDMXLibrary = NewObject<UDMXLibrary>(InOuter, InName, InFlags | RF_Public);
+	check(NewDMXLibrary);
+
+	GeneralSceneDescription->Rename(nullptr, NewDMXLibrary);
+	NewDMXLibrary->SetMVRGeneralSceneDescription(GeneralSceneDescription);
+
+	const bool bReimportExisitingGDTFs = InImportOptions ? InImportOptions->bReimportExisitingGDTFs : true;
+	const TArray<UDMXImportGDTF*> NewGDTFs = CreateGDTFs(InOuter, InFlags, bReimportExisitingGDTFs);
+
+	InitializeDMXLibrary(NewDMXLibrary, NewGDTFs);
+
+	OutDMXLibrary = NewDMXLibrary;
+	OutGDTFs = NewGDTFs;
+}
+
+void FDMXLibraryFromMVRImporter::Reimport(UDMXLibrary* InDMXLibrary, UDMXLibraryFromMVRImportOptions* InImportOptions, TArray<UDMXImportGDTF*>& OutGDTFs)
+{
+	check(GeneralSceneDescription);
+
+	OutGDTFs.Reset();
+
+	GeneralSceneDescription->Rename(nullptr, InDMXLibrary);
+	InDMXLibrary->SetMVRGeneralSceneDescription(GeneralSceneDescription);
+
+	const bool bReimportExistingGDTFs = InImportOptions ? InImportOptions->bReimportExisitingGDTFs : true;
+	const TArray<UDMXImportGDTF*> NewGDTFs = CreateGDTFs(InDMXLibrary->GetOuter(), InDMXLibrary->GetFlags(), bReimportExistingGDTFs);
+
+	InitializeDMXLibrary(InDMXLibrary, NewGDTFs);
+
+	OutGDTFs = NewGDTFs;
+}
+
+void FDMXLibraryFromMVRImporter::InitializeDMXLibrary(UDMXLibrary* DMXLibrary, const TArray<UDMXImportGDTF*>& GDTFAssets)
+{
+	check(DMXLibrary);
+	check(Zip.IsValid());
+	check(GeneralSceneDescription);
+
+	// Move the General Scene Description to the DMX Library and set it
+	GeneralSceneDescription->Rename(nullptr, DMXLibrary);
+	DMXLibrary->SetMVRGeneralSceneDescription(GeneralSceneDescription);
+
+	// Get or create a Fixture Type for each GDTF
+	const TArray<UDMXEntityFixtureType*> FixtureTypes = DMXLibrary->GetEntitiesTypeCast<UDMXEntityFixtureType>();
+	TMap<FString, UDMXEntityFixtureType*> GDTFSpecToFixtureTypeMap;
+	TMap<UDMXEntityFixtureType*, FLinearColor> FixtureTypeToColorMap;
+	for (UDMXImportGDTF* GDTF : GDTFAssets)
+	{
+		UDMXGDTFAssetImportData* GDTFAssetImportData = GDTF->GetGDTFAssetImportData();
+		if (!GDTFAssetImportData)
+		{
+			continue;
+		}
+		const FString GDTFSourceFilename = GDTFAssetImportData->GetFilePathAndName();
+		const FString GDTFFilename = FPaths::GetCleanFilename(GDTFSourceFilename);
+
+		UDMXEntityFixtureType* const* ExistingFixtureTypePtr = Algo::FindByPredicate(FixtureTypes, [GDTF](const UDMXEntityFixtureType* FixtureType)
+			{
+				return FixtureType->DMXImport == GDTF;
+			});
+		if (ExistingFixtureTypePtr)
+		{
+			GDTFSpecToFixtureTypeMap.Add(GDTFFilename, *ExistingFixtureTypePtr);
+		}
+		else
+		{
+			FDMXEntityFixtureTypeConstructionParams FixtureTypeConstructionParams;
+			FixtureTypeConstructionParams.DMXCategory = FDMXFixtureCategory(FDMXFixtureCategory::GetFirstValue());
+			FixtureTypeConstructionParams.ParentDMXLibrary = DMXLibrary;
+
+			UDMXEntityFixtureType* NewFixtureType = UDMXEntityFixtureType::CreateFixtureTypeInLibrary(FixtureTypeConstructionParams, FPaths::GetBaseFilename(GDTFFilename));
+			const bool bAdvancedImportSuccess = FDMXInitializeFixtureTypeFromGDTFHelper::GenerateModesFromGDTF(*NewFixtureType, *GDTF);
+			if (!bAdvancedImportSuccess)
+			{
+				UE_LOG(LogDMXEditor, Warning, TEXT("Failed to initialize Fixture Type '%s', falling back to legacy method that doesn't support matrix fixtures."), *NewFixtureType->GetName());
+				NewFixtureType->SetModesFromDMXImport(GDTF);
+			}
+			NewFixtureType->DMXImport = GDTF;
+
+			GDTFSpecToFixtureTypeMap.Add(GDTFFilename, NewFixtureType);
+		}
+	}
+
+	// Remove Fixture Patches that have no corresponding Fixture Node in the General Scene Description
+	TArray<UDMXMVRFixtureNode*> FixtureNodes;
+	GeneralSceneDescription->GetFixtureNodes(FixtureNodes);
+	TArray<UDMXEntityFixturePatch*> FixturePatches = DMXLibrary->GetEntitiesTypeCast<UDMXEntityFixturePatch>();
+	for (UDMXEntityFixturePatch* FixturePatch : FixturePatches)
+	{
+		UDMXMVRFixtureNode** CorrespondingFixtureNodePtr = Algo::FindByPredicate(FixtureNodes, [FixturePatch](const UDMXMVRFixtureNode* FixtureNode)
+			{
+				return FixturePatch->GetMVRFixtureUUID() == FixtureNode->UUID;
+			});
+		if (!CorrespondingFixtureNodePtr)
+		{
+			UDMXEntityFixturePatch::RemoveFixturePatchFromLibrary(FDMXEntityFixturePatchRef(FixturePatch));
+		}
+	}
+	FixturePatches = DMXLibrary->GetEntitiesTypeCast<UDMXEntityFixturePatch>();
+
+	// Update or create new Fixture Patches for the MVR Fixtures
+	for (const UDMXMVRFixtureNode* FixtureNode : FixtureNodes)
+	{
+		// Find the GDTF Filename, irregardless if it has the .gdtf extension
+		FString GDTFFilename = FixtureNode->GDTFSpec;
+		if (!GDTFSpecToFixtureTypeMap.Contains(GDTFFilename))
+		{
+			GDTFFilename = FixtureNode->GDTFSpec + TEXT(".gdtf");
+		}
+		if (!GDTFSpecToFixtureTypeMap.Contains(GDTFFilename))
+		{
+			continue;
+		}
+
+		UDMXEntityFixtureType* FixtureType = GDTFSpecToFixtureTypeMap[GDTFFilename];
+		if (!FixtureTypeToColorMap.Contains(FixtureType))
+		{
+			FLinearColor FixtureTypeColor = FLinearColor::MakeRandomColor();
+
+			// Avoid dominant red values for a bit more of a professional feel
+			if (FixtureTypeColor.R > 0.75f)
+			{
+				FixtureTypeColor.R = FMath::Abs(FixtureTypeColor.R - 1.0f);
+			}
+
+			FixtureTypeToColorMap.Add(FixtureType, FixtureTypeColor);
+		}
+
+		int32 ActiveModeIndex = FixtureType->Modes.IndexOfByPredicate([FixtureNode](const FDMXFixtureMode& Mode)
+			{
+				return Mode.ModeName == FixtureNode->GDTFMode;
+			});
+
+		if (ActiveModeIndex != INDEX_NONE)
+		{
+			UDMXEntityFixturePatch* const* ExistingFixturePatchPtr = Algo::FindByPredicate(FixturePatches, [FixtureNode](const UDMXEntityFixturePatch* FixturePatch)
+				{
+					return FixturePatch->GetMVRFixtureUUID() == FixtureNode->UUID;
+				});
+			if (ExistingFixturePatchPtr)
+			{
+				// Update the existing patch (matching MVR UUID) if possible
+				UDMXEntityFixturePatch& ExistingFixturePatch = **ExistingFixturePatchPtr;
+
+				ExistingFixturePatch.SetFixtureType(FixtureType);
+				ExistingFixturePatch.SetActiveModeIndex(ActiveModeIndex);
+				ExistingFixturePatch.SetUniverseID(FixtureNode->GetUniverseID());
+				ExistingFixturePatch.SetStartingChannel(FixtureNode->GetStartingChannel());
+				ExistingFixturePatch.EditorColor = FixtureTypeToColorMap.FindChecked(FixtureType);
+			}
+			else
+			{
+				// Create a new patch
+				FDMXEntityFixturePatchConstructionParams FixturePatchConstructionParams;
+				FixturePatchConstructionParams.ActiveMode = ActiveModeIndex;
+				FixturePatchConstructionParams.FixtureTypeRef = FDMXEntityFixtureTypeRef(FixtureType);
+				FixturePatchConstructionParams.UniverseID = FixtureNode->GetUniverseID();
+				FixturePatchConstructionParams.StartingAddress = FixtureNode->GetStartingChannel();
+				FixturePatchConstructionParams.MVRFixtureUUID = FixtureNode->UUID;
+
+				UDMXEntityFixturePatch* FixturePatch = UDMXEntityFixturePatch::CreateFixturePatchInLibrary(FixturePatchConstructionParams, FixtureNode->Name);
+				FixturePatch->EditorColor = FixtureTypeToColorMap.FindChecked(FixtureType);
+			}
+		}
+		else
+		{
+			UE_LOG(LogDMXEditor, Warning, TEXT("Skipped creating a Fixture Patch for '%s', cannot find Mode '%s' in '%s'."), *FixtureNode->Name, *FixtureNode->GDTFMode, *FixtureNode->GDTFSpec)
+		}
+	}
+
+	DMXLibrary->UpdateGeneralSceneDescription();
+}
+
+void FDMXLibraryFromMVRImporter::AddReferencedObjects(FReferenceCollector& Collector)
+{
+	Collector.AddReferencedObject(GeneralSceneDescription);
+}
+
+UDMXMVRGeneralSceneDescription* FDMXLibraryFromMVRImporter::CreateGeneralSsceneDescription(UObject* Outer) const
+{
+	if (!Zip.IsValid())
+	{
+		return nullptr;
+	}
+
+	TArray64<uint8> XMLData;
+	if (!Zip->GetFileContent(TEXT("GeneralSceneDescription.xml"), XMLData))
+	{
+		return nullptr;
+	}
+
+	// MVR implicitly adpots UTF-8 encoding of Xml Files by adopting the GDTF standard (DIN-15800).
+	// Content is NOT null-terminated; we need to specify lengths here.
+	const FUTF8ToTCHAR TCHARData(reinterpret_cast<const ANSICHAR*>(XMLData.GetData()), XMLData.Num());
+
+	const TSharedRef<FXmlFile> GeneralSceneDescriptionXml = MakeShared<FXmlFile>(FString(TCHARData.Length(), TCHARData.Get()), EConstructMethod::ConstructFromBuffer);
+
+	UDMXMVRGeneralSceneDescription* NewGeneralSceneDescription = UDMXMVRGeneralSceneDescription::CreateFromXmlFile(GeneralSceneDescriptionXml, Outer, NAME_None);
+
+	return NewGeneralSceneDescription;
+}
+
+TArray<UDMXImportGDTF*> FDMXLibraryFromMVRImporter::CreateGDTFs(UObject* InParent, EObjectFlags InFlags, bool bReimportExisting)
+{
+	if (!GeneralSceneDescription || !Zip.IsValid())
+	{
+		return TArray<UDMXImportGDTF*>();
+	}
+
+	const FString Path = FPaths::GetPath(InParent->GetName()) + TEXT("/GDTFs");
+
+	TArray<FAssetData> ExistingGDTFAssets;
+	FAssetRegistryModule& AssetRegistryModule = FModuleManager::LoadModuleChecked<FAssetRegistryModule>(TEXT("AssetRegistry"));
+	AssetRegistryModule.Get().GetAssetsByPath(FName(Path), ExistingGDTFAssets);
+
+	TArray<UDMXImportGDTF*> ExistingGDTFs;
+	for (const FAssetData& AssetData : ExistingGDTFAssets)
+	{
+		if (UDMXImportGDTF* GDTF = Cast<UDMXImportGDTF>(AssetData.GetAsset()))
+		{
+			ExistingGDTFs.AddUnique(GDTF);
+		}
+	}
+
+	UDMXGDTFFactory* GDTFFactory = NewObject<UDMXGDTFFactory>();
+	TArray<UDMXImportGDTF*> ImportedGDTFs;
+	TArray<FString> ImportedGDTFNames;
+	if (ExistingGDTFs.Num() > 0 && bReimportExisting)
+	{
+		for (UDMXImportGDTF* GDTFAsset : ExistingGDTFs)
+		{
+			if (GDTFFactory->Reimport(GDTFAsset) == EReimportResult::Succeeded)
+			{
+				ImportedGDTFs.Add(GDTFAsset);
+
+				if (UDMXGDTFAssetImportData* GDTFAssetImportData = GDTFAsset->GetGDTFAssetImportData())
+				{
+					const FString SourceFilename = FPaths::GetCleanFilename(GDTFAssetImportData->GetFilePathAndName());
+					ImportedGDTFNames.Add(SourceFilename);
+				}
+			}
+		}
+	}
+
+	// Import GDTFs that aren't yet imported in the reimport procedure
+	TArray<UDMXMVRFixtureNode*> FixtureNodes;
+	GeneralSceneDescription->GetFixtureNodes(FixtureNodes);
+	for (const UDMXMVRFixtureNode* FixtureNode : FixtureNodes)
+	{
+		// Don't import the same GDTF twice
+		if (!ImportedGDTFNames.Contains(FixtureNode->GDTFSpec))
+		{
+			using namespace UE::DMX::DMXLibraryFromMVRImporter::Private;
+
+			FString GDTFFilename;
+			if (!FindGDTFFilenameInMVRZip(Zip.ToSharedRef(), FixtureNode->GDTFSpec, GDTFFilename))
+			{
+				UE_LOG(LogDMXEditor, Warning, TEXT("Cannot find a GDTF file that corresponds to GDTF Spec '%s' of Fixture Node '%s'. Ignoring Fixture."), *FixtureNode->GDTFSpec, *FixtureNode->Name)
+				{
+					continue;
+				}
+			}
+
+			const FDMXZipper::FDMXScopedUnzipToTempFile ScopedUnzipGDTF(Zip.ToSharedRef(), GDTFFilename);
+			if (!ScopedUnzipGDTF.TempFilePathAndName.IsEmpty())
+			{
+				constexpr bool bRemovePathFromDesiredName = true;
+				const FString BaseFileName = FPaths::GetBaseFilename(ScopedUnzipGDTF.TempFilePathAndName, bRemovePathFromDesiredName);
+				FString AssetName = ObjectTools::SanitizeObjectName(BaseFileName);
+
+				FString GDTFPackageName;
+				FString GDTFName;
+				FAssetToolsModule& AssetToolsModule = FModuleManager::LoadModuleChecked<FAssetToolsModule>("AssetTools");
+				AssetToolsModule.Get().CreateUniqueAssetName(Path / AssetName, TEXT(""), GDTFPackageName, GDTFName);
+
+				UPackage* Package = CreatePackage(*GDTFPackageName);
+				check(Package);
+				Package->FullyLoad();
+
+				bool bCancelled;
+				UObject* NewGDTFObject = GDTFFactory->FactoryCreateFile(UDMXImportGDTF::StaticClass(), Package, FName(*FixtureNode->GDTFSpec), InFlags | RF_Public, ScopedUnzipGDTF.TempFilePathAndName, nullptr, GWarn, bCancelled);
+
+				if (UDMXImportGDTF* NewGDTF = Cast<UDMXImportGDTF>(NewGDTFObject))
+				{
+					ImportedGDTFs.Add(NewGDTF);
+					ImportedGDTFNames.Add(FixtureNode->GDTFSpec);
+
+					FAssetRegistryModule::AssetCreated(NewGDTF);
+					Package->MarkPackageDirty();
+				}
+			}
+		}
+	}
+
+	return ImportedGDTFs;
+}
+
+#undef LOCTEXT_NAMESPACE
diff --git a/Engine/Plugins/VirtualProduction/DMX/DMXEngine/Source/DMXEditor/Private/Factories/DMXLibraryFromMVRImporter.h b/Engine/Plugins/VirtualProduction/DMX/DMXEngine/Source/DMXEditor/Private/Factories/DMXLibraryFromMVRImporter.h
new file mode 100644
--- /dev/null
+++ b/Engine/Plugins/VirtualProduction/DMX/DMXEngine/Source/DMXEditor/Private/Factories/DMXLibraryFromMVRImporter.h
@@ -0,0 +1,92 @@
+// Copyright Epic Games, Inc. All Rights Reserved.
+
+#pragma once
+
+#include "CoreMinimal.h"
+#include "UObject/GCObject.h"
+
+class UDMXImportGDTF;
+class UDMXLibrary;
+class UDMXMVRAssetImportData;
+class UDMXMVRGeneralSceneDescription;
+class UDMXLibraryFromMVRImportOptions;
+class FDMXZipper;
+
+enum EObjectFlags;
+
+
+/**
+ * Imports an MVR file into a DMX Library. 
+ * 
+ * 1. LoadMVRFile.
+ * 2. CreateGDTFs
+ * 3. InitializeDMXLibrary
+ * 
+ * Does not offer any merge options with an existing DMX Library.
+ * To merge an MVR file into an existing DMX Library, import the MVR file into a new DMX Library, then merge it with the existing.
+ */
+class FDMXLibraryFromMVRImporter :
+	public FGCObject
+{
+public:
+	/** Loads the MVR File */
+	UE_NODISCARD bool LoadMVRFile(const FString& InFilename);
+
+	/** Creates a modal window that allows the user to change import options. */
+	void UpdateImportOptionsFromModalWindow(UDMXLibraryFromMVRImportOptions* InOutImportOptions, bool& bOutCancelled);
+
+	/** 
+	 * Imports an MVR File into a newly created DMX Library. Creates GDTF assets and reimports if specified so in ImportOptions.
+	 *
+	 * @param InOuter			The outer in which the new objects will be created
+	 * @param InName			The name of the DMX Library (GDTFs will be auto named)
+	 * @param InFlag			Flags with which the new objects are constructed
+	 * @param ImportOptions		Can be null. If set, displays a dialog where the user can configure the import.
+	 * @param OutDMXLibrary		The resulting DMX Library
+	 * @param OutGDTFs			The resulting GDTFs
+	 * #param bOutCancelled		True if the import was cancelled
+	 */
+	void Import(UObject* InOuter, const FName& InName, EObjectFlags InFlags, UDMXLibraryFromMVRImportOptions* ImportOptions, UDMXLibrary*& OutDMXLibrary, TArray<UDMXImportGDTF*>& OutGDTFs, bool& bOutCancelled);
+
+	/**
+	 * Imports an MVR File into an existing DMX Library
+	 *
+	 * @param InOuter			The outer in which the new objects will be created
+	 * @param InName			The name of the DMX Library (GDTFs will be auto named)
+	 * @param InFlag			Flags with which the new objects are constructed
+	 * @param ImportOptions		Can be null. If set, displays a dialog where the user can configure the import.
+	 * @param OutDMXLibrary		The resulting DMX Library
+	 * @param OutGDTFs			The resulting GDTFs
+	 * #param bOutCancelled		True if the import was cancelled
+	 */
+
+	void Reimport(UDMXLibrary* InDMXLibrary, UDMXLibraryFromMVRImportOptions* InImportOptions, TArray<UDMXImportGDTF*>& OutGDTFs);
+
+protected:
+	//~ Begin FGCObject interface
+	virtual void AddReferencedObjects(FReferenceCollector& Collector) override;
+	virtual FString GetReferencerName() const override
+	{
+		return TEXT("FDMXLibraryFromMVRImporter");
+	}
+	//~ End FGCObject interface
+
+private:
+	/** Creates a General Scene Description from the MVR zip. Returns the General Scene Description or nullptr if not a valid MVR zip. */
+	UDMXMVRGeneralSceneDescription* CreateGeneralSsceneDescription(UObject* Outer) const;
+
+	/** Creates GDTF assets from the MVR */
+	TArray<UDMXImportGDTF*> CreateGDTFs(UObject* InParent, EObjectFlags InFlags, bool bReimportExisting);
+
+	/** Initializes the DMX Library from the MVR data present and specified GDTFs */
+	void InitializeDMXLibrary(UDMXLibrary* DMXLibrary, const TArray<UDMXImportGDTF*>& GDTFs);
+
+	/** Temporary General Scene Description of the MVR File, or nullptr if not a valid MVR */
+	UDMXMVRGeneralSceneDescription* GeneralSceneDescription = nullptr;
+
+	/** Zip that holds the data */
+	TSharedPtr<FDMXZipper> Zip;
+
+	/** File name of the MVR File */
+	FString Filename;
+};
diff --git a/Engine/Plugins/VirtualProduction/DMX/DMXEngine/Source/DMXEditor/Private/Widgets/SDMXLibraryFromMVRImportOptions.cpp b/Engine/Plugins/VirtualProduction/DMX/DMXEngine/Source/DMXEditor/Private/Widgets/SDMXLibraryFromMVRImportOptions.cpp
new file mode 100644
--- /dev/null
+++ b/Engine/Plugins/VirtualProduction/DMX/DMXEngine/Source/DMXEditor/Private/Widgets/SDMXLibraryFromMVRImportOptions.cpp
@@ -0,0 +1,116 @@
+// Copyright Epic Games, Inc. All Rights Reserved.
+
+#include "Widgets/SDMXLibraryFromMVRImportOptions.h"
+
+#include "Customizations/DMXLibraryFromMVRImportOptionsDetails.h"
+#include "Factories/DMXLibraryFromMVRImportOptions.h"
+
+#include "IDetailsView.h"
+#include "PropertyEditorModule.h"
+#include "SPrimaryButton.h"
+#include "Framework/Application/SlateApplication.h"
+#include "Modules/ModuleManager.h"
+
+
+#define LOCTEXT_NAMESPACE "SDMXLibraryFromMVRFactoryOptions"
+
+void SDMXLibraryFromMVRImportOptions::Construct(const FArguments& InArgs, const TSharedRef<SWindow>& ParentWindow, UDMXLibraryFromMVRImportOptions* DMXLibraryFromMVRImportOptions)
+{
+	if (!ensureAlwaysMsgf(DMXLibraryFromMVRImportOptions, TEXT("Invalid Options object for DMX Library from MVR Import Options widget.")))
+	{
+		return;
+	}
+
+	FPropertyEditorModule& PropertyEditorModule = FModuleManager::GetModuleChecked<FPropertyEditorModule>("PropertyEditor");
+	FDetailsViewArgs DetailsViewArgs;
+	DetailsViewArgs.bAllowSearch = false;
+	DetailsViewArgs.NameAreaSettings = FDetailsViewArgs::HideNameArea;
+	const TSharedRef<IDetailsView> DetailsView = PropertyEditorModule.CreateDetailView(DetailsViewArgs);
+
+	DetailsView->RegisterInstancedCustomPropertyLayout(UDMXLibraryFromMVRImportOptions::StaticClass(), FOnGetDetailCustomizationInstance::CreateStatic(FDMXLibraryFromMVRImportOptionsDetails::MakeInstance));
+
+	DetailsView->SetObject(DMXLibraryFromMVRImportOptions);
+
+	// Default to Cancelled, only clicking the 'Import' button should proceed.
+	DMXLibraryFromMVRImportOptions->bCancelled = true;
+
+	ChildSlot
+		[
+			SNew(SVerticalBox)
+	
+			+ SVerticalBox::Slot()
+			.HAlign(HAlign_Fill)
+			.VAlign(VAlign_Fill)
+			.FillHeight(1.f)
+			[
+				DetailsView
+			]
+				
+			+ SVerticalBox::Slot()
+			.HAlign(HAlign_Right)
+			.VAlign(VAlign_Bottom)
+			.AutoHeight()
+			[
+				SNew(SHorizontalBox)
+					
+				+ SHorizontalBox::Slot()
+				.HAlign(HAlign_Right)
+				.VAlign(VAlign_Center)
+				.AutoWidth()
+				.Padding(4.f, 2.f)
+				[
+					SAssignNew(ImportButton, SPrimaryButton)
+					.Text(LOCTEXT("OptionWindow_Import", "Import"))
+					.OnClicked_Lambda([ParentWindow, DMXLibraryFromMVRImportOptions]()
+						{
+							DMXLibraryFromMVRImportOptions->bCancelled = false;
+							ParentWindow->RequestDestroyWindow();
+							return FReply::Handled();
+						})
+				]
+
+				+ SHorizontalBox::Slot()
+				.HAlign(HAlign_Right)
+				.VAlign(VAlign_Bottom)
+				.AutoWidth()
+				.Padding(4.f, 2.f)
+				[
+					SNew(SButton)
+					.Text(LOCTEXT("OptionWindow_Cancel", "Cancel"))
+					.ToolTipText(LOCTEXT("OptionWindow_Cancel_ToolTip", "Cancels importing this file"))
+					.OnClicked_Lambda([ParentWindow]()
+						{
+							ParentWindow->RequestDestroyWindow();
+							return FReply::Handled();
+						})
+				]
+			]
+		];
+
+	RegisterActiveTimer(0.f, FWidgetActiveTimerDelegate::CreateSP(this, &SDMXLibraryFromMVRImportOptions::SetFocusPostConstruct));
+}
+
+EActiveTimerReturnType SDMXLibraryFromMVRImportOptions::SetFocusPostConstruct(double InCurrentTime, float InDeltaTime)
+{
+	if (ImportButton.IsValid())
+	{
+		FSlateApplication::Get().SetKeyboardFocus(ImportButton, EFocusCause::SetDirectly);
+	}
+
+	return EActiveTimerReturnType::Stop;
+}
+
+FReply SDMXLibraryFromMVRImportOptions::OnKeyDown(const FGeometry& MyGeometry, const FKeyEvent& InKeyEvent)
+{
+	if (InKeyEvent.GetKey() == EKeys::Escape)
+	{
+		if (UDMXLibraryFromMVRImportOptions* ImportOptions = WeakImportOptions.Get())
+		{
+			ImportOptions->bCancelled = true;
+		}
+	}
+
+	return FReply::Unhandled();
+}
+
+#undef LOCTEXT_NAMESPACE
diff --git a/Engine/Plugins/VirtualProduction/DMX/DMXEngine/Source/DMXEditor/Private/Widgets/SDMXLibraryFromMVRImportOptions.h b/Engine/Plugins/VirtualProduction/DMX/DMXEngine/Source/DMXEditor/Private/Widgets/SDMXLibraryFromMVRImportOptions.h
new file mode 100644
--- /dev/null
+++ b/Engine/Plugins/VirtualProduction/DMX/DMXEngine/Source/DMXEditor/Private/Widgets/SDMXLibraryFromMVRImportOptions.h
@@ -0,0 +1,42 @@
+// Copyright Epic Games, Inc. All Rights Reserved.
+
+#pragma once
+
+#include "CoreMinimal.h"
+#include "Widgets/SCompoundWidget.h"
+
+class UDMXLibraryFromMVRImportOptions;
+
+class SPrimaryButton;
+
+
+/** Widget that displays the DMX Library from MVR Import Options */
+class SDMXLibraryFromMVRImportOptions
+	: public SCompoundWidget
+{
+public:
+	SLATE_BEGIN_ARGS(SDMXLibraryFromMVRImportOptions)
+	{}
+		SLATE_ARGUMENT(FVector2D, MaxWidgetSize)
+
+    SLATE_END_ARGS()
+
+	/** Constructs this widget */
+    void Construct(const FArguments& InArgs, const TSharedRef<SWindow>& ParentWindow, UDMXLibraryFromMVRImportOptions* DMXLibraryFromMVRImportOptions);
+
+protected:
+	//~ Begin SWidget interface
+    virtual bool SupportsKeyboardFocus() const override { return true; }
+	virtual FReply OnKeyDown(const FGeometry& MyGeometry, const FKeyEvent& InKeyEvent) override;
+	//~ End SWidget interface
+
+private:
+	/** Sets focus after the widget was constructed (next tick) */
+	EActiveTimerReturnType SetFocusPostConstruct(double InCurrentTime, float InDeltaTime);
+
+	/** The Import Button widget */
+	TSharedPtr<SPrimaryButton> ImportButton;
+
+	/** Import options being displayed */
+	TWeakObjectPtr<UDMXLibraryFromMVRImportOptions> WeakImportOptions;
+};
diff --git a/Engine/Plugins/VirtualProduction/DMX/DMXEngine/Source/DMXEditor/Public/Factories/DMXLibraryFromMVRFactory.h b/Engine/Plugins/VirtualProduction/DMX/DMXEngine/Source/DMXEditor/Public/Factories/DMXLibraryFromMVRFactory.h
--- a/Engine/Plugins/VirtualProduction/DMX/DMXEngine/Source/DMXEditor/Public/Factories/DMXLibraryFromMVRFactory.h
+++ b/Engine/Plugins/VirtualProduction/DMX/DMXEngine/Source/DMXEditor/Public/Factories/DMXLibraryFromMVRFactory.h
@@ -3,16 +3,11 @@
 #pragma once
 
 #include "Factories/Factory.h"
-
 #include "EditorReimportHandler.h"
 
 #include "DMXLibraryFromMVRFactory.generated.h"
 
-class FDMXZipper;
-class UDMXImportGDTF;
-class UDMXLibrary;
 class UDMXMVRAssetImportData;
-class UDMXMVRGeneralSceneDescription;
 
 
 UCLASS()
@@ -23,10 +18,11 @@ class DMXEDITOR_API UDMXLibraryFromMVRFactory
 	GENERATED_BODY()
 
 public:
+	/** Constructor */
 	UDMXLibraryFromMVRFactory();
 
 	//~ Begin UFactory Interface
-	virtual UObject* FactoryCreateFile(UClass* InClass, UObject* Parent, FName InName, EObjectFlags Flags, const FString& Filename, const TCHAR* Parms, FFeedbackContext* Warn, bool& bOutOperationCanceled);
+	virtual UObject* FactoryCreateFile(UClass* InClass, UObject* Parent, FName InName, EObjectFlags Flags, const FString& Filename, const TCHAR* Parms, FFeedbackContext* Warn, bool& bOutOperationCancelled);
 	virtual bool FactoryCanImport(const FString& Filename) override;
 	//~ End UFactory Interface	
 
@@ -42,18 +38,6 @@ protected:
 	//~ End FReimportHandler Interface
 
 private:
-	/** Creates a DMX Library asset. Returns nullptr if the library could not be created */
-	UDMXLibrary* CreateDMXLibraryAsset(UObject* Parent, const FName& Name, EObjectFlags Flags, const FString& InFilename);
-
-	/** Creates a General Scene Description from the MVR zip. Returns the General Scene Description or nullptr if not a valid MVR zip. */
-	UDMXMVRGeneralSceneDescription* CreateGeneralSsceneDescription(UDMXLibrary& DMXLibrary, const TSharedRef<FDMXZipper>& Zip) const;
-
-	/** Creates GDTF assets from the MVR */
-	TArray<UDMXImportGDTF*> CreateGDTFAssets(UObject* Parent, EObjectFlags Flags, const TSharedRef<FDMXZipper>& Zip, const UDMXMVRGeneralSceneDescription& GeneralSceneDescription);
-
-	/** Initializes the DMX Library from the General Scene Description and GDTF assets. Updates the library in case it already contains patches */
-	void InitDMXLibrary(UDMXLibrary* DMXLibrary, const TArray<UDMXImportGDTF*>& GDTFAssets, UDMXMVRGeneralSceneDescription* GeneralSceneDescription) const;
-
 	/** Returns the MVR Asset Import Data for the Object, or nullptr if it cannot be obtained */
 	UDMXMVRAssetImportData* GetMVRAssetImportData(UObject* DMXLibraryObject) const;
 };
diff --git a/Engine/Plugins/VirtualProduction/DMX/DMXEngine/Source/DMXRuntime/Private/MVR/DMXMVRGeneralSceneDescription.cpp b/Engine/Plugins/VirtualProduction/DMX/DMXEngine/Source/DMXRuntime/Private/MVR/DMXMVRGeneralSceneDescription.cpp
--- a/Engine/Plugins/VirtualProduction/DMX/DMXEngine/Source/DMXRuntime/Private/MVR/DMXMVRGeneralSceneDescription.cpp
+++ b/Engine/Plugins/VirtualProduction/DMX/DMXEngine/Source/DMXRuntime/Private/MVR/DMXMVRGeneralSceneDescription.cpp
@@ -59,9 +59,12 @@ UDMXMVRFixtureNode* UDMXMVRGeneralSceneDescription::FindFixtureNode(const FGuid&
 UDMXMVRGeneralSceneDescription* UDMXMVRGeneralSceneDescription::CreateFromXmlFile(TSharedRef<FXmlFile> GeneralSceneDescriptionXml, UObject* Outer, FName Name, EObjectFlags Flags)
 {
 	UDMXMVRGeneralSceneDescription* GeneralSceneDescription = NewObject<UDMXMVRGeneralSceneDescription>(Outer, Name, Flags);
-	GeneralSceneDescription->ParseGeneralSceneDescriptionXml(GeneralSceneDescriptionXml);
+	if (GeneralSceneDescription->ParseGeneralSceneDescriptionXml(GeneralSceneDescriptionXml))
+	{
+		return GeneralSceneDescription;
+	}
 
-	return GeneralSceneDescription;
+	return nullptr;
 }
 #endif // WITH_EDITOR
 
@@ -223,11 +226,11 @@ void UDMXMVRGeneralSceneDescription::WriteFixturePatchToGeneralSceneDescription(
 #endif // WITH_EDITOR
 
 #if WITH_EDITOR
-void UDMXMVRGeneralSceneDescription::ParseGeneralSceneDescriptionXml(const TSharedRef<FXmlFile>& GeneralSceneDescriptionXml)
+bool UDMXMVRGeneralSceneDescription::ParseGeneralSceneDescriptionXml(const TSharedRef<FXmlFile>& GeneralSceneDescriptionXml)
 {
 	checkf(RootNode, TEXT("Unexpected: MVR General Scene Description Root Node is invalid."));
 
-	RootNode->InitializeFromGeneralSceneDescriptionXml(GeneralSceneDescriptionXml);
+	return RootNode->InitializeFromGeneralSceneDescriptionXml(GeneralSceneDescriptionXml);
 }
 #endif // WITH_EDITOR
 
diff --git a/Engine/Plugins/VirtualProduction/DMX/DMXEngine/Source/DMXRuntime/Public/MVR/DMXMVRGeneralSceneDescription.h b/Engine/Plugins/VirtualProduction/DMX/DMXEngine/Source/DMXRuntime/Public/MVR/DMXMVRGeneralSceneDescription.h
--- a/Engine/Plugins/VirtualProduction/DMX/DMXEngine/Source/DMXRuntime/Public/MVR/DMXMVRGeneralSceneDescription.h
+++ b/Engine/Plugins/VirtualProduction/DMX/DMXEngine/Source/DMXRuntime/Public/MVR/DMXMVRGeneralSceneDescription.h
@@ -63,7 +63,7 @@ private:
 	void WriteFixturePatchToGeneralSceneDescription(const UDMXEntityFixturePatch& FixturePatch);
 
 	/** Parses a General Scene Description Xml File. Only ment to be used for initialization (ensured) */
-	void ParseGeneralSceneDescriptionXml(const TSharedRef<FXmlFile>& GeneralSceneDescriptionXml);
+	UE_NODISCARD bool ParseGeneralSceneDescriptionXml(const TSharedRef<FXmlFile>& GeneralSceneDescriptionXml);
 #endif
 
 	/** Returns the Fixture IDs currently in use, as number, sorted from lowest to highest */
diff --git a/Engine/Plugins/VirtualProduction/DMX/DMXEngine/Source/DMXRuntime/Public/MVR/DMXMVRVersion.h b/Engine/Plugins/VirtualProduction/DMX/DMXEngine/Source/DMXRuntime/Public/MVR/DMXMVRVersion.h
--- a/Engine/Plugins/VirtualProduction/DMX/DMXEngine/Source/DMXRuntime/Public/MVR/DMXMVRVersion.h
+++ b/Engine/Plugins/VirtualProduction/DMX/DMXEngine/Source/DMXRuntime/Public/MVR/DMXMVRVersion.h
@@ -10,7 +10,7 @@ struct DMXRUNTIME_API FDMXMVRVersion
 {
 	static const int32 MajorVersion = 1;
 
-	static const int32 MinorVersion = 4;
+	static const int32 MinorVersion = 5;
 
 	static const FString GetMajorVersionAsString();
 	static const FString GetMinorVersionAsString();
diff --git a/Engine/Plugins/VirtualProduction/DMX/DMXEngine/Source/DMXRuntime/Public/MVR/Types/DMXMVRRootNode.h b/Engine/Plugins/VirtualProduction/DMX/DMXEngine/Source/DMXRuntime/Public/MVR/Types/DMXMVRRootNode.h
--- a/Engine/Plugins/VirtualProduction/DMX/DMXEngine/Source/DMXRuntime/Public/MVR/Types/DMXMVRRootNode.h
+++ b/Engine/Plugins/VirtualProduction/DMX/DMXEngine/Source/DMXRuntime/Public/MVR/Types/DMXMVRRootNode.h
@@ -36,7 +36,7 @@ public:
 	 * Does not merge, clears previous entries. 
 	 * Returns false if not a valid MVR General Scene Description Xml 
 	 */
-	bool InitializeFromGeneralSceneDescriptionXml(const TSharedRef<FXmlFile>& GeneralSceneDescriptionXml);
+	UE_NODISCARD bool InitializeFromGeneralSceneDescriptionXml(const TSharedRef<FXmlFile>& GeneralSceneDescriptionXml);
 
 	/** Creates a DMX Xml Node as per MVR standard, returns nullptr if the Xml File cannot be created. */
 	TSharedPtr<FXmlFile> CreateXmlFile();
diff --git a/Engine/Plugins/VirtualProduction/DMX/DMXProtocol/Source/DMXProtocolEditor/Private/DetailsCustomizations/DMXPortReferenceCustomizationBase.cpp b/Engine/Plugins/VirtualProduction/DMX/DMXProtocol/Source/DMXProtocolEditor/Private/DetailsCustomizations/DMXPortReferenceCustomizationBase.cpp
--- a/Engine/Plugins/VirtualProduction/DMX/DMXProtocol/Source/DMXProtocolEditor/Private/DetailsCustomizations/DMXPortReferenceCustomizationBase.cpp
+++ b/Engine/Plugins/VirtualProduction/DMX/DMXProtocol/Source/DMXProtocolEditor/Private/DetailsCustomizations/DMXPortReferenceCustomizationBase.cpp
@@ -27,11 +27,20 @@ FDMXPortReferenceCustomizationBase::FDMXPortReferenceCustomizationBase()
 
 void FDMXPortReferenceCustomizationBase::CustomizeHeader(TSharedRef<IPropertyHandle> StructPropertyHandle, class FDetailWidgetRow& HeaderRow, IPropertyTypeCustomizationUtils& StructCustomizationUtils)
 {
-	HeaderRow
-	.NameContent()
-	[
-		StructPropertyHandle->CreatePropertyNameWidget()
-	];
+	constexpr TCHAR ShowOnlyInnerPropertiesMetaDataSpecifier[] = TEXT("ShowOnlyInnerProperties");
+	bool bShowHeader = !StructPropertyHandle->HasMetaData(ShowOnlyInnerPropertiesMetaDataSpecifier);
+	if (bShowHeader)
+	{
+		HeaderRow
+			.NameContent()
+			[
+				StructPropertyHandle->CreatePropertyNameWidget()
+			]
+			.ValueContent()
+			[
+				StructPropertyHandle->CreatePropertyValueWidget()
+			];
+	}
 }
 
 void FDMXPortReferenceCustomizationBase::CustomizeChildren(TSharedRef<IPropertyHandle> StructPropertyHandle, class IDetailChildrenBuilder& ChildBuilder, IPropertyTypeCustomizationUtils& StructCustomizationUtils)
