DMX - Fix CIS Localization warnings

diff --git a/Engine/Plugins/VirtualProduction/DMX/DMXEngine/Source/DMXEditor/Private/Customizations/DMXLibraryFromMVRImportOptionsDetails.cpp b/Engine/Plugins/VirtualProduction/DMX/DMXEngine/Source/DMXEditor/Private/Customizations/DMXLibraryFromMVRImportOptionsDetails.cpp
--- a/Engine/Plugins/VirtualProduction/DMX/DMXEngine/Source/DMXEditor/Private/Customizations/DMXLibraryFromMVRImportOptionsDetails.cpp
+++ b/Engine/Plugins/VirtualProduction/DMX/DMXEngine/Source/DMXEditor/Private/Customizations/DMXLibraryFromMVRImportOptionsDetails.cpp
@@ -54,7 +54,7 @@ void FDMXLibraryFromMVRImportOptionsDetails::CustomizeDetails(IDetailLayoutBuild
 			[
 				SNew(STextBlock)
 				.Font(DetailBuilder.GetDetailFont())
-				.Text(LOCTEXT("GeneratePortDisplayName", "Generate Input Port"))
+				.Text(LOCTEXT("GenerateInputPortDisplayName", "Generate Input Port"))
 			]
 			.ValueContent()
 			[
@@ -76,7 +76,7 @@ void FDMXLibraryFromMVRImportOptionsDetails::CustomizeDetails(IDetailLayoutBuild
 			[
 				SNew(STextBlock)
 				.Font(DetailBuilder.GetDetailFont())
-				.Text(LOCTEXT("GeneratePortDisplayName", "Generate Output Port"))
+				.Text(LOCTEXT("GenerateOutputPortDisplayName", "Generate Output Port"))
 			]
 			.ValueContent()
 			[
diff --git a/Engine/Plugins/VirtualProduction/DMX/DMXEngine/Source/DMXEditor/Private/Factories/DMXLibraryFromMVRImporter.cpp b/Engine/Plugins/VirtualProduction/DMX/DMXEngine/Source/DMXEditor/Private/Factories/DMXLibraryFromMVRImporter.cpp
--- a/Engine/Plugins/VirtualProduction/DMX/DMXEngine/Source/DMXEditor/Private/Factories/DMXLibraryFromMVRImporter.cpp
+++ b/Engine/Plugins/VirtualProduction/DMX/DMXEngine/Source/DMXEditor/Private/Factories/DMXLibraryFromMVRImporter.cpp
@@ -88,8 +88,8 @@ namespace UE::DMX::DMXLibraryFromMVRImporter::Private
 			const FVector2D WindowPosition = (DisplayTopLeft + (DisplaySize - WindowSize) / 2.0f);
 			
 			const FText Caption = InOptions->bIsReimport ?
-				LOCTEXT("MVRSceneActorFactoryOptionsWindowCaption", "MVR Reimport Options") :
-				LOCTEXT("MVRSceneActorFactoryOptionsWindowCaption", "MVR Import Options");
+				LOCTEXT("ReimportWindowCaption", "MVR Reimport Options") :
+				LOCTEXT("ImportWindowCaption", "MVR Import Options");
 
 			const TSharedRef<SWindow> Window = SNew(SWindow)
 				.Title(Caption)
