DMX - Fix CIS Static Analysis Win64 (MSVC)

diff --git a/Engine/Plugins/VirtualProduction/DMX/DMXEngine/Source/DMXEditor/Private/Factories/DMXLibraryFromMVRImportOptions.cpp b/Engine/Plugins/VirtualProduction/DMX/DMXEngine/Source/DMXEditor/Private/Factories/DMXLibraryFromMVRImportOptions.cpp
--- a/Engine/Plugins/VirtualProduction/DMX/DMXEngine/Source/DMXEditor/Private/Factories/DMXLibraryFromMVRImportOptions.cpp
+++ b/Engine/Plugins/VirtualProduction/DMX/DMXEngine/Source/DMXEditor/Private/Factories/DMXLibraryFromMVRImportOptions.cpp
@@ -28,6 +28,10 @@ void UDMXLibraryFromMVRImportOptions::ApplyOptions(UDMXLibrary* DMXLibrary)
 	}
 
 	const TArray<UDMXEntityFixturePatch*> TempFixturePatches = DMXLibrary->GetEntitiesTypeCast<UDMXEntityFixturePatch>();
+	if (TempFixturePatches.IsEmpty())
+	{
+		return;
+	}
 	const UDMXEntityFixturePatch* const* FirstPatchPtr = Algo::MinElementBy(TempFixturePatches, [](const UDMXEntityFixturePatch* FixturePatch)
 		{
 			return FixturePatch->GetUniverseID();
@@ -36,11 +40,8 @@ void UDMXLibraryFromMVRImportOptions::ApplyOptions(UDMXLibrary* DMXLibrary)
 		{
 			return FixturePatch->GetUniverseID();
 		});
+	checkf(FirstPatchPtr && LastPatchPtr, TEXT("Unexpected, cannot find Min and Max element in an array that is not empty."));
 
-	if (!FirstPatchPtr && !LastPatchPtr)
-	{
-		return;
-	}
 	const int32 FirstUniverse = (*FirstPatchPtr)->GetUniverseID();
 	const int32 NumUniverses = FMath::Max(1, FMath::Clamp((*LastPatchPtr)->GetUniverseID() - FirstUniverse + 1, 0, TNumericLimits<int32>::Max()));
 
