DMX - Small MVR improvements

* Fixture Actors should not append an extra number when MVR scene is updated
* Reimporting MVR results in different Fixture Patch names, but shouldn't
* In some cases the same DMX Library still can be spawned twice where it shouldn't
* Sort by [REDACTED_PERSON] by default in library
* Rename 'Update Transforms' to 'Reset Transforms'

diff --git a/Engine/Plugins/VirtualProduction/DMX/DMXEngine/Source/DMXEditor/Private/Customizations/DMXMVRSceneActorDetails.cpp b/Engine/Plugins/VirtualProduction/DMX/DMXEngine/Source/DMXEditor/Private/Customizations/DMXMVRSceneActorDetails.cpp
--- a/Engine/Plugins/VirtualProduction/DMX/DMXEngine/Source/DMXEditor/Private/Customizations/DMXMVRSceneActorDetails.cpp
+++ b/Engine/Plugins/VirtualProduction/DMX/DMXEngine/Source/DMXEditor/Private/Customizations/DMXMVRSceneActorDetails.cpp
@@ -91,7 +91,7 @@ void FDMXMVRSceneActorDetails::CreateRefreshMVRSceneSection(IDetailLayoutBuilder
 	
 	const TSharedRef<IPropertyHandle> UpdateTransformHandle = DetailBuilder.GetProperty(GET_MEMBER_NAME_CHECKED(ADMXMVRSceneActor, bUpdateTransformsOnRefresh));
 	UpdateTransformHandle->MarkHiddenByCustomization();
-	MVRCategory.AddCustomRow(LOCTEXT("UpdateTransformsFilterText", "Update Transform"))
+	MVRCategory.AddCustomRow(LOCTEXT("ResetTransformsFilterText", "Reset Transforms"))
 		.NameContent()
 		[
 			SNew(SBorder)
diff --git a/Engine/Plugins/VirtualProduction/DMX/DMXEngine/Source/DMXEditor/Private/Factories/DMXMVRSceneActorFactory.cpp b/Engine/Plugins/VirtualProduction/DMX/DMXEngine/Source/DMXEditor/Private/Factories/DMXMVRSceneActorFactory.cpp
--- a/Engine/Plugins/VirtualProduction/DMX/DMXEngine/Source/DMXEditor/Private/Factories/DMXMVRSceneActorFactory.cpp
+++ b/Engine/Plugins/VirtualProduction/DMX/DMXEngine/Source/DMXEditor/Private/Factories/DMXMVRSceneActorFactory.cpp
@@ -71,6 +71,8 @@ AActor* UDMXMVRSceneActorFactory::SpawnActor(UObject* Asset, ULevel* InLevel, co
 		if(IsDMXLibraryAlreadySpawned(World, DMXLibrary))
 		{ 
 			NotifyDMXLibraryAlreadySpawned();
+
+			return nullptr;
 		}
 
 		ADMXMVRSceneActor* MVRSceneActor = World->SpawnActor<ADMXMVRSceneActor>(ADMXMVRSceneActor::StaticClass());
diff --git a/Engine/Plugins/VirtualProduction/DMX/DMXEngine/Source/DMXEditor/Private/Widgets/FixturePatch/SDMXMVRFixtureList.cpp b/Engine/Plugins/VirtualProduction/DMX/DMXEngine/Source/DMXEditor/Private/Widgets/FixturePatch/SDMXMVRFixtureList.cpp
--- a/Engine/Plugins/VirtualProduction/DMX/DMXEngine/Source/DMXEditor/Private/Widgets/FixturePatch/SDMXMVRFixtureList.cpp
+++ b/Engine/Plugins/VirtualProduction/DMX/DMXEngine/Source/DMXEditor/Private/Widgets/FixturePatch/SDMXMVRFixtureList.cpp
@@ -563,7 +563,7 @@ const FName FDMXMVRFixtureListCollumnIDs::Patch = "Patch";
 
 SDMXMVRFixtureList::SDMXMVRFixtureList()
 	: SortMode(EColumnSortMode::Ascending)
-	, SortedByColumnID(FDMXMVRFixtureListCollumnIDs::Patch)
+	, SortedByColumnID(FDMXMVRFixtureListCollumnIDs::FixtureID)
 {}
 
 SDMXMVRFixtureList::~SDMXMVRFixtureList()
@@ -756,7 +756,7 @@ void SDMXMVRFixtureList::RefreshList()
 	// Update and sort the list and its widgets
 	ListView->RebuildList();
 
-	SortByColumnID(EColumnSortPriority::Max, FDMXMVRFixtureListCollumnIDs::Patch, EColumnSortMode::Ascending);
+	SortByColumnID(EColumnSortPriority::Max, SortedByColumnID, EColumnSortMode::Ascending);
 
 	AdoptSelectionFromFixturePatchSharedData();
 }
diff --git a/Engine/Plugins/VirtualProduction/DMX/DMXEngine/Source/DMXRuntime/Private/MVR/DMXMVRSceneActor.cpp b/Engine/Plugins/VirtualProduction/DMX/DMXEngine/Source/DMXRuntime/Private/MVR/DMXMVRSceneActor.cpp
--- a/Engine/Plugins/VirtualProduction/DMX/DMXEngine/Source/DMXRuntime/Private/MVR/DMXMVRSceneActor.cpp
+++ b/Engine/Plugins/VirtualProduction/DMX/DMXEngine/Source/DMXRuntime/Private/MVR/DMXMVRSceneActor.cpp
@@ -537,7 +537,7 @@ AActor* ADMXMVRSceneActor::SpawnMVRActor(const TSubclassOf<AActor>&ActorClass, U
 	{
 		return nullptr;
 	}
-	NewFixtureActor->SetActorLabel(ActorName.ToString());
+	NewFixtureActor->SetActorLabel(FixturePatch->Name);
 
 	NewFixtureActor->RegisterAllComponents();
 	USceneComponent* RootComponentOfChildActor = NewFixtureActor->GetRootComponent();
