DMX - Fix Auto assigning a fixture patch after adding Matrix function is not working on the first try

diff --git a/Engine/Plugins/VirtualProduction/DMX/DMXEngine/Source/DMXEditor/Private/DMXEditorUtils.cpp b/Engine/Plugins/VirtualProduction/DMX/DMXEngine/Source/DMXEditor/Private/DMXEditorUtils.cpp
--- a/Engine/Plugins/VirtualProduction/DMX/DMXEngine/Source/DMXEditor/Private/DMXEditorUtils.cpp
+++ b/Engine/Plugins/VirtualProduction/DMX/DMXEngine/Source/DMXEditor/Private/DMXEditorUtils.cpp
@@ -870,6 +870,13 @@ void FDMXEditorUtils::AutoAssignedChannels(bool bAllowDecrementUniverse, bool bA
 			const bool bChannelIsSmallerOrEqual = FixturePatchA.GetStartingChannel() <= FixturePatchB.GetStartingChannel();
 			return bUniverseIsSmaller || (bUniverseIsEqual && bChannelIsSmallerOrEqual);
 		});
+	// Move all to the universe and address of the first patch
+	for (UDMXEntityFixturePatch* FixturePatch : FixturePatches)
+	{
+		FixturePatch->Modify();
+		FixturePatch->SetUniverseID(FixturePatches[0]->GetUniverseID());
+		FixturePatch->SetStartingChannel(FixturePatches[0]->GetStartingChannel());
+	}
 
 	// Using absolute channels (Universe * DMX_UIVERSE_SIZE + Channel)
 	auto FindAbsoluteChannelLambda([bAllowDecrementUniverse, bAllowDecrementChannels, &FixturePatches, &AllFixturePatchesInLibrary](UDMXEntityFixturePatch* FixturePatch, const TArray<UDMXEntityFixturePatch*>& OtherPendingAutoAssignFixturePatches) -> int32
