DMX - Fix DMX Library and PixelMapping in some cases take needlessly long to redraw

diff --git a/Engine/Plugins/VirtualProduction/DMX/DMXEngine/Source/DMXEditor/Private/Widgets/FixturePatch/SDMXMVRFixtureList.cpp b/Engine/Plugins/VirtualProduction/DMX/DMXEngine/Source/DMXEditor/Private/Widgets/FixturePatch/SDMXMVRFixtureList.cpp
--- a/Engine/Plugins/VirtualProduction/DMX/DMXEngine/Source/DMXEditor/Private/Widgets/FixturePatch/SDMXMVRFixtureList.cpp
+++ b/Engine/Plugins/VirtualProduction/DMX/DMXEngine/Source/DMXEditor/Private/Widgets/FixturePatch/SDMXMVRFixtureList.cpp
@@ -711,19 +711,20 @@ void SDMXMVRFixtureList::RefreshList()
 		return;
 	}
 
-	DMXLibrary->UpdateGeneralSceneDescription();
 	UDMXMVRGeneralSceneDescription* GeneralSceneDescription = DMXLibrary->GetLazyGeneralSceneDescription();
 	if (!GeneralSceneDescription)
 	{
 		return;
 	}
 
+	// Remove all that are no longer a patch
 	const TArray<UDMXEntityFixturePatch*> FixturePatches = DMXLibrary->GetEntitiesTypeCast<UDMXEntityFixturePatch>();
 	ListSource.RemoveAll([&FixturePatches](const TSharedPtr<FDMXMVRFixtureListItem>& Item)
 		{
 			return !FixturePatches.Contains(Item->GetFixturePatch());
 		});
 	
+	bool bUpdatedGeneralSceneDescription = false;
 	for (UDMXEntityFixturePatch* FixturePatch : FixturePatches)
 	{
 		bool bExisting = ListSource.ContainsByPredicate([FixturePatch](const TSharedPtr<FDMXMVRFixtureListItem>& Item)
@@ -733,6 +734,13 @@ void SDMXMVRFixtureList::RefreshList()
 
 		if (!bExisting)
 		{
+			if (!bUpdatedGeneralSceneDescription)
+			{	
+				// Ony update the General Scene Description when Nodes were never acquired, or new patches were added.
+				DMXLibrary->UpdateGeneralSceneDescription();
+				bUpdatedGeneralSceneDescription = true;
+			}
+
 			UDMXMVRFixtureNode* MVRFixtureNode = FindMVRFixtureNode(GeneralSceneDescription, FixturePatch);
 			if (!MVRFixtureNode)
 			{
diff --git a/Engine/Plugins/VirtualProduction/DMX/DMXEngine/Source/DMXRuntime/Private/Game/DMXComponent.cpp b/Engine/Plugins/VirtualProduction/DMX/DMXEngine/Source/DMXRuntime/Private/Game/DMXComponent.cpp
--- a/Engine/Plugins/VirtualProduction/DMX/DMXEngine/Source/DMXRuntime/Private/Game/DMXComponent.cpp
+++ b/Engine/Plugins/VirtualProduction/DMX/DMXEngine/Source/DMXRuntime/Private/Game/DMXComponent.cpp
@@ -223,6 +223,14 @@ void UDMXComponent::OnAllFixturePatchesReceiveDMXInEditorEnabled(bool bEnabled)
 #if WITH_EDITOR
 void UDMXComponent::OnFixturePatchPropertiesChanged(const UDMXEntityFixturePatch* FixturePatch)
 {
+	// Look up the global setting first, which is on in most cases.
+	// This is faster as GetFixturePatch can be slow.
+	const UDMXProtocolSettings* ProtocolSettings = GetDefault<UDMXProtocolSettings>();
+	if (ProtocolSettings->ShouldAllFixturePatchesReceiveDMXInEditor() && IsComponentTickEnabled())
+	{
+		return;
+	}
+
 	if (FixturePatch == GetFixturePatch())
 	{
 		UpdateTickEnabled();
diff --git a/Engine/Plugins/VirtualProduction/DMX/DMXEngine/Source/DMXRuntime/Private/Library/DMXLibrary.cpp b/Engine/Plugins/VirtualProduction/DMX/DMXEngine/Source/DMXRuntime/Private/Library/DMXLibrary.cpp
--- a/Engine/Plugins/VirtualProduction/DMX/DMXEngine/Source/DMXRuntime/Private/Library/DMXLibrary.cpp
+++ b/Engine/Plugins/VirtualProduction/DMX/DMXEngine/Source/DMXRuntime/Private/Library/DMXLibrary.cpp
@@ -32,13 +32,16 @@ UDMXLibrary::UDMXLibrary()
 
 void UDMXLibrary::Serialize(FArchive& Ar)
 {
-#if WITH_EDITORONLY_DATA
-	if (Ar.IsSaving())
-	{
-		// Update the General Scene Description before saving it
-		UpdateGeneralSceneDescription();
-	}
-#endif
+// More performant without. Instead Update the General Scene Description on demand when it's actually needed.
+// Serialize is only here because the 5.1.1 hotfix doens't allow a header change and can be removed in 5.2.
+// 
+// #if WITH_EDITORONLY_DATA
+//	if (Ar.IsSaving())
+//	{
+//		// Update the General Scene Description before saving it
+//		UpdateGeneralSceneDescription();
+//	}
+//#endif
 
 	Super::Serialize(Ar);
 }
diff --git a/Engine/Plugins/VirtualProduction/DMX/DMXPixelMapping/Source/DMXPixelMappingEditor/Private/ViewModels/DMXPixelMappingOutputComponentModel.cpp b/Engine/Plugins/VirtualProduction/DMX/DMXPixelMapping/Source/DMXPixelMappingEditor/Private/ViewModels/DMXPixelMappingOutputComponentModel.cpp
--- a/Engine/Plugins/VirtualProduction/DMX/DMXPixelMapping/Source/DMXPixelMappingEditor/Private/ViewModels/DMXPixelMappingOutputComponentModel.cpp
+++ b/Engine/Plugins/VirtualProduction/DMX/DMXPixelMapping/Source/DMXPixelMappingEditor/Private/ViewModels/DMXPixelMappingOutputComponentModel.cpp
@@ -221,11 +221,9 @@ void FDMXPixelMappingOutputComponentModel::UpdateFixtureNode()
 		return;
 	}
 
-	GeneralSceneDescription->WriteDMXLibraryToGeneralSceneDescription(*DMXLibrary);
 	WeakFixtureNode = GeneralSceneDescription->FindFixtureNode(MVRUUID);
 }
 
-
 FDMXPixelMappingScreenComponentModel::FDMXPixelMappingScreenComponentModel(const TSharedRef<FDMXPixelMappingToolkit>& InToolkit, TWeakObjectPtr<UDMXPixelMappingScreenComponent> InScreenComponent)
 	: WeakScreenComponent(InScreenComponent)
 	, WeakToolkit(InToolkit)
diff --git a/Engine/Plugins/VirtualProduction/DMX/DMXPixelMapping/Source/DMXPixelMappingRuntime/Private/LayoutScripts/DMXPixelMappingLayoutScript.cpp b/Engine/Plugins/VirtualProduction/DMX/DMXPixelMapping/Source/DMXPixelMappingRuntime/Private/LayoutScripts/DMXPixelMappingLayoutScript.cpp
--- a/Engine/Plugins/VirtualProduction/DMX/DMXPixelMapping/Source/DMXPixelMappingRuntime/Private/LayoutScripts/DMXPixelMappingLayoutScript.cpp
+++ b/Engine/Plugins/VirtualProduction/DMX/DMXPixelMapping/Source/DMXPixelMappingRuntime/Private/LayoutScripts/DMXPixelMappingLayoutScript.cpp
@@ -63,9 +63,6 @@ void FDMXPixelMappingLayoutToken::InitializeFixtureID()
 	{
 		return;
 	}
-#if WITH_EDITOR
-	GeneralSceneDescription->WriteDMXLibraryToGeneralSceneDescription(*DMXLibrary);
-#endif
 	
 	UDMXMVRFixtureNode* FixtureNode = GeneralSceneDescription->FindFixtureNode(MVRUUID);
 	if (!FixtureNode)
