DMX - Fix Localization warnings in SDMXControlConsoleEditorFixturePatchRowWidget.cpp

diff --git a/Engine/Plugins/VirtualProduction/DMX/DMXEngine/Source/DMXControlConsoleEditor/Private/Widgets/SDMXControlConsoleEditorFixturePatchRowWidget.cpp b/Engine/Plugins/VirtualProduction/DMX/DMXEngine/Source/DMXControlConsoleEditor/Private/Widgets/SDMXControlConsoleEditorFixturePatchRowWidget.cpp
--- a/Engine/Plugins/VirtualProduction/DMX/DMXEngine/Source/DMXControlConsoleEditor/Private/Widgets/SDMXControlConsoleEditorFixturePatchRowWidget.cpp
+++ b/Engine/Plugins/VirtualProduction/DMX/DMXEngine/Source/DMXControlConsoleEditor/Private/Widgets/SDMXControlConsoleEditorFixturePatchRowWidget.cpp
@@ -71,12 +71,12 @@ void SDMXControlConsoleEditorFixturePatchRowWidget::Construct(const FArguments&
 							SNew(SButton)
 							.ButtonStyle(FAppStyle::Get(), "FlatButton.Success")
 							.OnClicked(this, &SDMXControlConsoleEditorFixturePatchRowWidget::OnAddNextClicked)
-							.ToolTipText(LOCTEXT("AddNextTooltip", "Add a Fader Group generated from this Fixture Patch next to selection."))
+							.ToolTipText(LOCTEXT("AddNextButtonTooltip", "Add a Fader Group generated from this Fixture Patch next to selection."))
 							.Visibility(TAttribute<EVisibility>(this, &SDMXControlConsoleEditorFixturePatchRowWidget::GetAddNextButtonVisibility))
 							[
 								SNew(STextBlock)
 								.Font(FAppStyle::GetFontStyle(TEXT("PropertyWindow.NormalFont")))
-								.Text(LOCTEXT("AddNextButton", "+ Add"))
+								.Text(LOCTEXT("AddNextButtonCaption", "+ Add"))
 							]
 						]
 
@@ -88,12 +88,12 @@ void SDMXControlConsoleEditorFixturePatchRowWidget::Construct(const FArguments&
 							SNew(SButton)
 							.ButtonStyle(FAppStyle::Get(), "FlatButton.Success")
 							.OnClicked(this, &SDMXControlConsoleEditorFixturePatchRowWidget::OnAddRowClicked)
-							.ToolTipText(LOCTEXT("AddRowTooltip", "Add a Fader Group generated from this Fixture Patch on a new row."))
+							.ToolTipText(LOCTEXT("AddRowButtonTooltip", "Add a Fader Group generated from this Fixture Patch on a new row."))
 							.Visibility(TAttribute<EVisibility>(this, &SDMXControlConsoleEditorFixturePatchRowWidget::GetAddRowButtonVisibility))
 							[
 								SNew(STextBlock)
 								.Font(FAppStyle::GetFontStyle(TEXT("PropertyWindow.NormalFont")))
-								.Text(LOCTEXT("AddRowButton", "+ Row"))
+								.Text(LOCTEXT("AddRowButtonCaption", "+ Row"))
 							]
 						]
 					]
@@ -106,12 +106,12 @@ void SDMXControlConsoleEditorFixturePatchRowWidget::Construct(const FArguments&
 						SNew(SButton)
 						.ButtonStyle(FAppStyle::Get(), "FlatButton.Success")
 						.OnClicked(this, &SDMXControlConsoleEditorFixturePatchRowWidget::OnGenerateClicked)
-						.ToolTipText(LOCTEXT("AddNextTooltip", "Edit selection generating a Fader Group from this Fixture Patch."))
+						.ToolTipText(LOCTEXT("GenerateButtonTooltip", "Edit selection generating a Fader Group from this Fixture Patch."))
 						.Visibility(TAttribute<EVisibility>(this, &SDMXControlConsoleEditorFixturePatchRowWidget::GetGenerateButtonVisibility))
 						[
 							SNew(STextBlock)
 							.Font(FAppStyle::GetFontStyle(TEXT("PropertyWindow.NormalFont")))
-							.Text(LOCTEXT("GenerateButton", "Generate"))
+							.Text(LOCTEXT("GenerateButtonCaption", "Generate"))
 						]
 					]
 				]
