DMX - Duplicating fixtures in the level causes shared MIDs and prevents from rebinding to another fixture patch, resulting in broken source and duplicate

diff --git a/Engine/Plugins/VirtualProduction/DMX/DMXEngine/Source/DMXRuntime/Private/Game/DMXComponent.cpp b/Engine/Plugins/VirtualProduction/DMX/DMXEngine/Source/DMXRuntime/Private/Game/DMXComponent.cpp
--- a/Engine/Plugins/VirtualProduction/DMX/DMXEngine/Source/DMXRuntime/Private/Game/DMXComponent.cpp
+++ b/Engine/Plugins/VirtualProduction/DMX/DMXEngine/Source/DMXRuntime/Private/Game/DMXComponent.cpp
@@ -60,6 +60,13 @@ void UDMXComponent::SetReceiveDMXFromPatch(bool bReceive)
 
 void UDMXComponent::OnFixturePatchReceivedDMX(UDMXEntityFixturePatch* FixturePatch, const FDMXNormalizedAttributeValueMap& NormalizedValuePerAttribute)
 {
+	// Hotfix 5.1 - Old bindings are not properly removed when duplicating Actors that use this component @todo.
+	if (FixturePatch != GetFixturePatch())
+	{
+		FixturePatch->OnFixturePatchReceivedDMX.RemoveAll(this);
+		return;
+	}
+
 	if (OnFixturePatchReceived.IsBound())
 	{
 		FEditorScriptExecutionGuard ScriptGuard;
@@ -222,7 +229,7 @@ void UDMXComponent::OnAllFixturePatchesReceiveDMXInEditorEnabled(bool bEnabled)
 
 #if WITH_EDITOR
 void UDMXComponent::OnFixturePatchPropertiesChanged(const UDMXEntityFixturePatch* FixturePatch)
-{
+{	
 	// Look up the global setting first, which is on in most cases.
 	// This is faster as GetFixturePatch can be slow.
 	const UDMXProtocolSettings* ProtocolSettings = GetDefault<UDMXProtocolSettings>();
@@ -233,6 +240,7 @@ void UDMXComponent::OnFixturePatchPropertiesChanged(const UDMXEntityFixturePatch
 
 	if (FixturePatch == GetFixturePatch())
 	{
+		SetupReceiveDMXBinding();
 		UpdateTickEnabled();
 	}
 }
diff --git a/Engine/Plugins/VirtualProduction/DMX/DMXFixtures/Source/DMXFixtures/Private/DMXFixtureActor.cpp b/Engine/Plugins/VirtualProduction/DMX/DMXFixtures/Source/DMXFixtures/Private/DMXFixtureActor.cpp
--- a/Engine/Plugins/VirtualProduction/DMX/DMXFixtures/Source/DMXFixtures/Private/DMXFixtureActor.cpp
+++ b/Engine/Plugins/VirtualProduction/DMX/DMXFixtures/Source/DMXFixtures/Private/DMXFixtureActor.cpp
@@ -87,14 +87,11 @@ void ADMXFixtureActor::InitializeFixture(UStaticMeshComponent* StaticMeshLens, U
 {
 	GetComponents(StaticMeshComponents);
 
-	// Create dynamic materials if needed
-	if (!DynamicMaterialLens || !DynamicMaterialBeam || !DynamicMaterialSpotLight || !DynamicMaterialPointLight)
-	{
-		DynamicMaterialLens = UMaterialInstanceDynamic::Create(LensMaterialInstance, nullptr);
-		DynamicMaterialBeam = UMaterialInstanceDynamic::Create(BeamMaterialInstance, nullptr);
-		DynamicMaterialSpotLight = UMaterialInstanceDynamic::Create(SpotLightMaterialInstance, nullptr);
-		DynamicMaterialPointLight = UMaterialInstanceDynamic::Create(PointLightMaterialInstance, nullptr);
-	}
+	// Create dynamic materials
+	DynamicMaterialLens = UMaterialInstanceDynamic::Create(LensMaterialInstance, nullptr);
+	DynamicMaterialBeam = UMaterialInstanceDynamic::Create(BeamMaterialInstance, nullptr);
+	DynamicMaterialSpotLight = UMaterialInstanceDynamic::Create(SpotLightMaterialInstance, nullptr);
+	DynamicMaterialPointLight = UMaterialInstanceDynamic::Create(PointLightMaterialInstance, nullptr);
 
 	// Get lens width (support scaling)
 	if (StaticMeshLens)
