DMX - Add sections to the DMX Menu, better Labels, better Tooltips

diff --git a/Engine/Plugins/VirtualProduction/DMX/DMXEngine/Source/DMXControlConsoleEditor/Private/DMXControlConsoleEditorModule.cpp b/Engine/Plugins/VirtualProduction/DMX/DMXEngine/Source/DMXControlConsoleEditor/Private/DMXControlConsoleEditorModule.cpp
--- a/Engine/Plugins/VirtualProduction/DMX/DMXEngine/Source/DMXControlConsoleEditor/Private/DMXControlConsoleEditorModule.cpp
+++ b/Engine/Plugins/VirtualProduction/DMX/DMXEngine/Source/DMXControlConsoleEditor/Private/DMXControlConsoleEditorModule.cpp
@@ -94,8 +94,8 @@ void FDMXControlConsoleEditorModule::ExtendDMXMenu(FMenuBuilder& MenuBuilder)
 {
 	MenuBuilder.AddMenuEntry(FDMXControlConsoleEditorCommands::Get().OpenControlConsole,
 			NAME_None,
-			LOCTEXT("DMXControlConsoleMenuLabel", "DMX Control Console"),
-			LOCTEXT("DMXControlConsoleMenuTooltip", "Opens a console that can send DMX to the network"),
+			LOCTEXT("DMXControlConsoleMenuLabel", "Control Console"),
+			LOCTEXT("DMXControlConsoleMenuTooltip", "Opens a small console that can send DMX locally or over the network"),
 			FSlateIcon(FDMXControlConsoleEditorStyle::Get().GetStyleSetName(), "DMXControlConsole.TabIcon")
 		);
 }
diff --git a/Engine/Plugins/VirtualProduction/DMX/DMXEngine/Source/DMXEditor/Private/DMXEditorModule.cpp b/Engine/Plugins/VirtualProduction/DMX/DMXEngine/Source/DMXEditor/Private/DMXEditorModule.cpp
--- a/Engine/Plugins/VirtualProduction/DMX/DMXEngine/Source/DMXEditor/Private/DMXEditorModule.cpp
+++ b/Engine/Plugins/VirtualProduction/DMX/DMXEngine/Source/DMXEditor/Private/DMXEditorModule.cpp
@@ -200,37 +200,46 @@ TSharedRef<SWidget> FDMXEditorModule::GenerateDMXLevelEditorToolbarMenu()
 	FMenuBuilder MenuBuilder(true, CommandBindings, LevelEditorToolbarDMXMenuExtender);
 
 	static const FName NoExtensionHook = NAME_None;
-	MenuBuilder.AddMenuEntry(FDMXEditorCommands::Get().OpenChannelsMonitor,
-		NoExtensionHook,
-		LOCTEXT("ChannelsMonitorLabel", "Open Channel Monitor"),
-		LOCTEXT("ChannelsMonitorTooltip", "Opens the Monitor for all DMX Channels in a Universe"),
-		FSlateIcon(FDMXEditorStyle::Get().GetStyleSetName(), "Icons.ChannelsMonitor")
-	);
-	MenuBuilder.AddMenuEntry(FDMXEditorCommands::Get().OpenActivityMonitor,
-		"OpenActivityMonitor",
-		LOCTEXT("ActivityMonitorLabel", "Open Activity Monitor"),
-		LOCTEXT("ActivityMonitorTooltip", "Open the Monitor for all DMX activity in a range of Universes"),
-		FSlateIcon(FDMXEditorStyle::Get().GetStyleSetName(), "Icons.ActivityMonitor")
-	);
-	MenuBuilder.AddMenuEntry(FDMXEditorCommands::Get().OpenPatchTool,
-		NoExtensionHook,
-		LOCTEXT("PatchToolLabel", "Open Patch Tool"),
-		LOCTEXT("PatchToolTooltip", "Open the patch tool - Useful to patch many fixtures at once."),
-		FSlateIcon(FDMXEditorStyle::Get().GetStyleSetName(), "Icons.PatchTool")
-	);
-	MenuBuilder.AddMenuEntry(FDMXEditorCommands::Get().ToggleReceiveDMX,
-		NoExtensionHook,
-		LOCTEXT("ReceiveDMXLabel", "Receive DMX"),
-		LOCTEXT("ReceiveDMXTooltip", "Sets whether DMX is received in from the network"),
-		FSlateIcon(FDMXEditorStyle::Get().GetStyleSetName(), "Icons.ReceiveDMX")
-	);
-	MenuBuilder.AddMenuEntry(FDMXEditorCommands::Get().ToggleSendDMX,
-		NoExtensionHook,
-		LOCTEXT("SendDMXLabel", "Send DMX"),
-		LOCTEXT("SendDMXTooltip", "Sets whether DMX is sent to the network"),
-		FSlateIcon(FDMXEditorStyle::Get().GetStyleSetName(), "Icons.SendDMX")
-	);
-	
+
+	MenuBuilder.BeginSection(NoExtensionHook, LOCTEXT("MenuSectionTools", "TOOLS"));
+	{
+		MenuBuilder.AddMenuEntry(FDMXEditorCommands::Get().OpenChannelsMonitor,
+			NoExtensionHook,
+			LOCTEXT("ChannelsMonitorLabel", "Channel Monitor"),
+			LOCTEXT("ChannelsMonitorTooltip", "Opens the Monitor displaying all DMX Channels in a Universe"),
+			FSlateIcon(FDMXEditorStyle::Get().GetStyleSetName(), "Icons.ChannelsMonitor")
+		);
+		MenuBuilder.AddMenuEntry(FDMXEditorCommands::Get().OpenActivityMonitor,
+			"OpenActivityMonitor",
+			LOCTEXT("ActivityMonitorLabel", "Activity Monitor"),
+			LOCTEXT("ActivityMonitorTooltip", "Opens the Activity Monitor that monitors DMX activity in a range of Universes"),
+			FSlateIcon(FDMXEditorStyle::Get().GetStyleSetName(), "Icons.ActivityMonitor")
+		);
+		MenuBuilder.AddMenuEntry(FDMXEditorCommands::Get().OpenPatchTool,
+			NoExtensionHook,
+			LOCTEXT("PatchToolLabel", "Patch Tool"),
+			LOCTEXT("PatchToolTooltip", "Opens the quick patch tool for helping patch large quantity of DMX fixtures"),
+			FSlateIcon(FDMXEditorStyle::Get().GetStyleSetName(), "Icons.PatchTool")
+		);
+	}
+
+	MenuBuilder.BeginSection(NoExtensionHook, LOCTEXT("MenuSectionDMXIO", "DMX IO"));
+	{
+		MenuBuilder.AddMenuEntry(FDMXEditorCommands::Get().ToggleReceiveDMX,
+			NoExtensionHook,
+			LOCTEXT("ReceiveDMXLabel", "Receive DMX"),
+			LOCTEXT("ReceiveDMXTooltip", "Sets whether DMX is received in from the network"),
+			FSlateIcon(FDMXEditorStyle::Get().GetStyleSetName(), "Icons.ReceiveDMX")
+		);
+		MenuBuilder.AddMenuEntry(FDMXEditorCommands::Get().ToggleSendDMX,
+			NoExtensionHook,
+			LOCTEXT("SendDMXLabel", "Send DMX"),
+			LOCTEXT("SendDMXTooltip", "Sets whether DMX is sent to the network"),
+			FSlateIcon(FDMXEditorStyle::Get().GetStyleSetName(), "Icons.SendDMX")
+		);
+	}
+	MenuBuilder.EndSection();
+
 	return MenuBuilder.MakeWidget();
 }
 
