DMX - Fix fixture patch is keeping the color from the first instance of a fixture when duplicating, order of the list is not respected when duplicating

diff --git a/Engine/Plugins/VirtualProduction/DMX/DMXEngine/Source/DMXEditor/Private/Widgets/FixturePatch/SDMXMVRFixtureList.cpp b/Engine/Plugins/VirtualProduction/DMX/DMXEngine/Source/DMXEditor/Private/Widgets/FixturePatch/SDMXMVRFixtureList.cpp
--- a/Engine/Plugins/VirtualProduction/DMX/DMXEngine/Source/DMXEditor/Private/Widgets/FixturePatch/SDMXMVRFixtureList.cpp
+++ b/Engine/Plugins/VirtualProduction/DMX/DMXEngine/Source/DMXEditor/Private/Widgets/FixturePatch/SDMXMVRFixtureList.cpp
@@ -158,6 +158,9 @@ namespace UE::DMX::SDMXMVRFixtureList::Private
 		constexpr bool bMarkLibraryDirty = false;
 		UDMXEntityFixturePatch* NewFixturePatch = UDMXEntityFixturePatch::CreateFixturePatchInLibrary(ConstructionParams, FixturePatchToDuplicate->Name, bMarkLibraryDirty);
 
+		// Use the same color as the duplicated patch
+		NewFixturePatch->EditorColor = FixturePatchToDuplicate->EditorColor;
+
 		Address += ChannelSpan;
 
 		DMXLibrary->PostEditChange();
@@ -1379,7 +1382,17 @@ void SDMXMVRFixtureList::OnDuplicateItems()
 	const FText TransactionText = LOCTEXT("DuplicateFixturePatchesTransaction", "Duplicate Fixture Patches");
 	const FScopedTransaction PasteTransaction(TransactionText);
 
-	const TArray<TWeakObjectPtr<UDMXEntityFixturePatch>> SelectedFixturePatches = FixturePatchSharedData->GetSelectedFixturePatches();
+	TArray<TWeakObjectPtr<UDMXEntityFixturePatch>> SelectedFixturePatches = FixturePatchSharedData->GetSelectedFixturePatches();
+	
+	// Sort in order of the list so duplicate happens in order of the list
+	Algo::SortBy(SelectedFixturePatches, [this](TWeakObjectPtr<UDMXEntityFixturePatch> FixturePatch)
+		{
+			return ListSource.IndexOfByPredicate([FixturePatch](const TSharedPtr<FDMXMVRFixtureListItem>& Item)
+				{
+					return Item->GetFixturePatch() == FixturePatch;
+				});
+		});
+
 	TArray<TWeakObjectPtr<UDMXEntityFixturePatch>> NewFixturePatches;
 	for (const TWeakObjectPtr<UDMXEntityFixturePatch> FixturePatch : SelectedFixturePatches)
 	{
