DMX - Right side fixture patcher is now a DMX Monitor

diff --git a/Engine/Plugins/VirtualProduction/DMX/DMXEngine/Source/DMXEditor/Private/Widgets/FixturePatch/SDMXChannelConnector.cpp b/Engine/Plugins/VirtualProduction/DMX/DMXEngine/Source/DMXEditor/Private/Widgets/FixturePatch/SDMXChannelConnector.cpp
--- a/Engine/Plugins/VirtualProduction/DMX/DMXEngine/Source/DMXEditor/Private/Widgets/FixturePatch/SDMXChannelConnector.cpp
+++ b/Engine/Plugins/VirtualProduction/DMX/DMXEngine/Source/DMXEditor/Private/Widgets/FixturePatch/SDMXChannelConnector.cpp
@@ -40,6 +40,11 @@ void SDMXChannelConnector::Construct(const FArguments& InArgs)
 		];
 }
 
+void SDMXChannelConnector::SetValue(uint8 Value)
+{
+	ChannelValueWidget->SetValue(Value);
+}
+
 FReply SDMXChannelConnector::OnMouseButtonDown(const FGeometry& MyGeometry, const FPointerEvent& MouseEvent)
 {
 	return OnMouseButtonDownOnChannel.Execute(ChannelValueWidget->GetChannelID(), MouseEvent);
diff --git a/Engine/Plugins/VirtualProduction/DMX/DMXEngine/Source/DMXEditor/Private/Widgets/FixturePatch/SDMXChannelConnector.h b/Engine/Plugins/VirtualProduction/DMX/DMXEngine/Source/DMXEditor/Private/Widgets/FixturePatch/SDMXChannelConnector.h
--- a/Engine/Plugins/VirtualProduction/DMX/DMXEngine/Source/DMXEditor/Private/Widgets/FixturePatch/SDMXChannelConnector.h
+++ b/Engine/Plugins/VirtualProduction/DMX/DMXEngine/Source/DMXEditor/Private/Widgets/FixturePatch/SDMXChannelConnector.h
@@ -74,6 +74,9 @@ public:
 	/** Sets the UniverseID of this connector */
 	void SetUniverseID(int32 InUniverseID) { UniverseID = InUniverseID; }
 
+	/** Sets the displayed value */
+	void SetValue(uint8 Value);
+
 protected:
 	// Begin SWidget interface	
 	virtual FReply OnMouseButtonDown(const FGeometry& MyGeometry, const FPointerEvent& MouseEvent) override;
diff --git a/Engine/Plugins/VirtualProduction/DMX/DMXEngine/Source/DMXEditor/Private/Widgets/FixturePatch/SDMXFixturePatcher.cpp b/Engine/Plugins/VirtualProduction/DMX/DMXEngine/Source/DMXEditor/Private/Widgets/FixturePatch/SDMXFixturePatcher.cpp
--- a/Engine/Plugins/VirtualProduction/DMX/DMXEngine/Source/DMXEditor/Private/Widgets/FixturePatch/SDMXFixturePatcher.cpp
+++ b/Engine/Plugins/VirtualProduction/DMX/DMXEngine/Source/DMXEditor/Private/Widgets/FixturePatch/SDMXFixturePatcher.cpp
@@ -3,6 +3,7 @@
 #include "SDMXFixturePatcher.h"
 
 #include "DMXEditor.h"
+#include "DMXEditorSettings.h"
 #include "DMXEditorTabNames.h"
 #include "DMXFixturePatchSharedData.h"
 #include "DMXFixturePatchNode.h"
@@ -112,7 +113,7 @@ void SDMXFixturePatcher::Construct(const FArguments& InArgs)
 						[
 							SNew(STextBlock)
 							.Font(FAppStyle::GetFontStyle("PropertyWindow.NormalFont"))
-							.Text(LOCTEXT("UniverseDisplayAllText", "Show all patched Universes"))
+							.Text(LOCTEXT("ShowAllPatchedUniversesLabel", "Show all patched Universes"))
 						]
 
 						+ SHorizontalBox::Slot()
@@ -124,6 +125,27 @@ void SDMXFixturePatcher::Construct(const FArguments& InArgs)
 							.IsChecked(false)
 							.OnCheckStateChanged(this, &SDMXFixturePatcher::OnToggleDisplayAllUniverses)
 						]
+
+						+ SHorizontalBox::Slot()
+						.AutoWidth()
+						.VAlign(VAlign_Center)
+						.Padding(FMargin(4.0f, 4.0f, 15.0f, 4.0f))
+						[
+							SNew(STextBlock)
+							.Font(FAppStyle::GetFontStyle("PropertyWindow.NormalFont"))
+							.Text(LOCTEXT("EnableInputMonitorLabel", "Monitor DMX Inputs"))
+							.ToolTipText(LOCTEXT("EnableInputMonitorTooltip", "If checked, monitors DMX Input Ports used in this DMX Library"))
+						]
+
+						+ SHorizontalBox::Slot()
+						.AutoWidth()
+						.VAlign(VAlign_Center)
+						.Padding(FMargin(4.0f, 4.0f, 15.0f, 4.0f))
+						[
+							SAssignNew(EnableDMXMonitorCheckBox, SCheckBox)
+							.IsChecked(false)
+							.OnCheckStateChanged(this, &SDMXFixturePatcher::OnToggleDMXMonitorEnabled)
+						]
 					]
 				]
 
@@ -165,6 +187,8 @@ void SDMXFixturePatcher::Construct(const FArguments& InArgs)
 
 	ShowSelectedUniverse();
 
+	const UDMXEditorSettings* DMXEditorSettings = GetDefault<UDMXEditorSettings>();
+	SetDMXMonitorEnabled(DMXEditorSettings->bFixturePatcherDMXMonitorEnabled);
 }
 END_SLATE_FUNCTION_BUILD_OPTIMIZATION
 
@@ -637,6 +661,36 @@ void SDMXFixturePatcher::OnToggleDisplayAllUniverses(ECheckBoxState CheckboxStat
 	}
 }
 
+void SDMXFixturePatcher::SetDMXMonitorEnabled(bool bEnabled)
+{
+	UDMXEditorSettings* DMXEditorSettings = GetMutableDefault<UDMXEditorSettings>();
+	DMXEditorSettings->bFixturePatcherDMXMonitorEnabled = bEnabled;
+	DMXEditorSettings->SaveConfig();
+
+	for (const TTuple<int32, TSharedPtr<SDMXPatchedUniverse>>& UniverseToWidgetPair : PatchedUniversesByID)
+	{
+		UniverseToWidgetPair.Value->SetMonitorInputsEnabled(bEnabled);
+	}
+}
+
+void SDMXFixturePatcher::OnToggleDMXMonitorEnabled(ECheckBoxState CheckboxState)
+{
+	switch (EnableDMXMonitorCheckBox->GetCheckedState())
+	{
+	case ECheckBoxState::Checked:
+		SetDMXMonitorEnabled(true);
+		return;
+
+	case ECheckBoxState::Unchecked:
+		SetDMXMonitorEnabled(false);
+		return;
+
+	case ECheckBoxState::Undetermined:
+	default:
+		checkNoEntry();
+	}
+}
+
 bool SDMXFixturePatcher::IsUniverseSelectionEnabled() const
 {
 	check(ShowAllUniversesCheckBox.IsValid());
diff --git a/Engine/Plugins/VirtualProduction/DMX/DMXEngine/Source/DMXEditor/Private/Widgets/FixturePatch/SDMXFixturePatcher.h b/Engine/Plugins/VirtualProduction/DMX/DMXEngine/Source/DMXEditor/Private/Widgets/FixturePatch/SDMXFixturePatcher.h
--- a/Engine/Plugins/VirtualProduction/DMX/DMXEngine/Source/DMXEditor/Private/Widgets/FixturePatch/SDMXFixturePatcher.h
+++ b/Engine/Plugins/VirtualProduction/DMX/DMXEngine/Source/DMXEditor/Private/Widgets/FixturePatch/SDMXFixturePatcher.h
@@ -104,9 +104,15 @@ protected:
 	/** Adds a universe widget */
 	void AddUniverse(int32 UniverseID);
 
-	/** Updates the grid depending on entities and selection */
+	/** Called when the display all universe option is turned on or off */
 	void OnToggleDisplayAllUniverses(ECheckBoxState CheckboxState);
 
+	/** Sets if the DMX Monitor feature is enabled */
+	void SetDMXMonitorEnabled(bool bEnabled);
+
+	/** Called when the monitor feature is turned on or off */
+	void OnToggleDMXMonitorEnabled(ECheckBoxState CheckboxState);
+
 	/** Returns true if the patcher allows for the selection of a single universe and show only that */
 	bool IsUniverseSelectionEnabled() const;
 
@@ -138,9 +144,12 @@ protected:
 	/** Returns the DMXLibrary or nullptr if not available */
 	UDMXLibrary* GetDMXLibrary() const;
 
-	/** Checkbox to determine if all check boxes shoudl be displayed */
+	/** Checkbox to determine if all patched Universes should be displayed */
 	TSharedPtr<SCheckBox> ShowAllUniversesCheckBox;
 
+	/** Checkbox to determine if the DMX Monitor functionality should be enabled */
+	TSharedPtr<SCheckBox> EnableDMXMonitorCheckBox;
+
 	/** Scrollbox containing all patch universes */
 	TSharedPtr<SScrollBox> PatchedUniverseScrollBox;
 
diff --git a/Engine/Plugins/VirtualProduction/DMX/DMXEngine/Source/DMXEditor/Private/Widgets/FixturePatch/SDMXPatchedUniverse.cpp b/Engine/Plugins/VirtualProduction/DMX/DMXEngine/Source/DMXEditor/Private/Widgets/FixturePatch/SDMXPatchedUniverse.cpp
--- a/Engine/Plugins/VirtualProduction/DMX/DMXEngine/Source/DMXEditor/Private/Widgets/FixturePatch/SDMXPatchedUniverse.cpp
+++ b/Engine/Plugins/VirtualProduction/DMX/DMXEngine/Source/DMXEditor/Private/Widgets/FixturePatch/SDMXPatchedUniverse.cpp
@@ -109,6 +109,9 @@ void SDMXPatchedUniverse::Construct(const FArguments& InArgs)
 	UDMXEntityFixturePatch::GetOnFixturePatchChanged().AddSP(this, &SDMXPatchedUniverse::OnFixturePatchChanged);
 	SharedData->OnFixturePatchSelectionChanged.AddSP(this, &SDMXPatchedUniverse::OnSelectionChanged);
 	FDMXPortManager::Get().OnPortsChanged.AddSP(this, &SDMXPatchedUniverse::UpdatePatchedUniverseReachability);
+
+	UDMXProtocolSettings* ProtocolSettings = GetMutableDefault<UDMXProtocolSettings>();
+	ProtocolSettings->GetOnSetReceiveDMXEnabled().AddSP(this, &SDMXPatchedUniverse::OnReceiveDMXEnabledChanged);
 }
 END_SLATE_FUNCTION_BUILD_OPTIMIZATION
 
@@ -120,6 +123,16 @@ void SDMXPatchedUniverse::RequestRefresh()
 	}
 }
 
+void SDMXPatchedUniverse::SetMonitorInputsEnabled(bool bEnabled)
+{
+	bMonitorInputs = bEnabled;
+
+	if (!bEnabled)
+	{
+		ResetMonitor();
+	}
+}
+
 void SDMXPatchedUniverse::CreateChannelConnectors()
 {
 	check(Grid.IsValid());
@@ -293,6 +306,11 @@ void SDMXPatchedUniverse::Tick(const FGeometry& AllottedGeometry, const double I
 	{
 		Unpatch(RemovedNode);
 	}
+
+	if (bMonitorInputs)
+	{
+		UpdateMonitor();
+	}
 }
 
 void SDMXPatchedUniverse::Unpatch(const TSharedPtr<FDMXFixturePatchNode>& Node)
@@ -336,6 +354,17 @@ void SDMXPatchedUniverse::OnFixturePatchChanged(const UDMXEntityFixturePatch* Fi
 			SetUniverseIDInternal(UniverseID);
 		}
 	}
+
+	FDMXEditorUtils::ClearAllDMXPortBuffers();
+	FDMXEditorUtils::ClearFixturePatchCachedData();
+}
+
+void SDMXPatchedUniverse::OnReceiveDMXEnabledChanged(bool bNewEnabled)
+{
+	if (!bNewEnabled)
+	{
+		ResetMonitor();
+	}
 }
 
 void SDMXPatchedUniverse::SetUniverseIDInternal(int32 NewUniverseID)
@@ -446,6 +475,42 @@ void SDMXPatchedUniverse::RefreshInternal()
 	}
 }
 
+void SDMXPatchedUniverse::UpdateMonitor()
+{
+	UDMXLibrary* DMXLibrary = GetDMXLibrary();
+	const UDMXProtocolSettings* ProtocolSettings = GetDefault<UDMXProtocolSettings>();
+	if (!DMXLibrary || !ProtocolSettings || !ProtocolSettings->IsReceiveDMXEnabled())
+	{
+		return;
+	}
+
+
+	const TSet<FDMXInputPortSharedRef>& InputPorts = DMXLibrary->GetInputPorts();
+	for (const FDMXInputPortSharedRef& InputPort : InputPorts)
+	{
+		FDMXSignalSharedPtr Signal;
+		InputPort->GameThreadGetDMXSignal(UniverseID, Signal);
+		if (Signal.IsValid())
+		{
+			for (int32 DataIndex = 0; DataIndex < Signal->ChannelData.Num(); DataIndex++)
+			{
+				if (ensureMsgf(ChannelConnectors.IsValidIndex(DataIndex), TEXT("Missing Channel to display value of Channel %i"), DataIndex))
+				{
+					ChannelConnectors[DataIndex]->SetValue(Signal->ChannelData[DataIndex]);
+				}
+			}
+		}
+	}
+}
+
+void SDMXPatchedUniverse::ResetMonitor()
+{
+	for (const TSharedPtr<SDMXChannelConnector>& Channel : ChannelConnectors)
+	{
+		Channel->SetValue(0);
+	}
+}
+
 FReply SDMXPatchedUniverse::HandleOnMouseButtonDownOnChannel(uint32 Channel, const FPointerEvent& PointerEvent)
 {
 	if (PointerEvent.GetEffectingButton() == EKeys::RightMouseButton)
diff --git a/Engine/Plugins/VirtualProduction/DMX/DMXEngine/Source/DMXEditor/Private/Widgets/FixturePatch/SDMXPatchedUniverse.h b/Engine/Plugins/VirtualProduction/DMX/DMXEngine/Source/DMXEditor/Private/Widgets/FixturePatch/SDMXPatchedUniverse.h
--- a/Engine/Plugins/VirtualProduction/DMX/DMXEngine/Source/DMXEditor/Private/Widgets/FixturePatch/SDMXPatchedUniverse.h
+++ b/Engine/Plugins/VirtualProduction/DMX/DMXEngine/Source/DMXEditor/Private/Widgets/FixturePatch/SDMXPatchedUniverse.h
@@ -77,6 +77,9 @@ public:
 	/** Requests to refresh the Widget on the next tick */
 	void RequestRefresh();
 
+	/** Sets if this universe should monitor inputs */
+	void SetMonitorInputsEnabled(bool bEnabled);
+
 	/** 
 	 * Patches the node.
 	 * Patches that have bAutoAssignAddress use their auto assigned address.
@@ -121,12 +124,21 @@ private:
 	/** Called when a Fixture Patch changed */
 	void OnFixturePatchChanged(const UDMXEntityFixturePatch* FixturePatch);
 
+	/** Called when receive DMX was enabled or disabled */
+	void OnReceiveDMXEnabledChanged(bool bNewEnabled);
+
 	/** Shows the universe directly */
 	void SetUniverseIDInternal(int32 NewUniverseID);
 
 	/** Refreshes the widget directly */
 	void RefreshInternal();
 
+	/** Updates the monitor */
+	void UpdateMonitor();
+
+	/** Resets the Monitor */
+	void ResetMonitor();
+
 	/** Creates a a new grid of channels */
 	void CreateChannelConnectors();
 
@@ -178,6 +190,9 @@ private:
 	/** The universe being displayed */
 	int32 UniverseID;
 
+	/** If true, monitors inputs */
+	bool bMonitorInputs = false;
+
 	/** If true the universe ID is out of controllers' ranges */
 	EDMXPatchedUniverseReachability PatchedUniverseReachability;
 	
diff --git a/Engine/Plugins/VirtualProduction/DMX/DMXEngine/Source/DMXEditor/Public/DMXEditorSettings.h b/Engine/Plugins/VirtualProduction/DMX/DMXEngine/Source/DMXEditor/Public/DMXEditorSettings.h
--- a/Engine/Plugins/VirtualProduction/DMX/DMXEngine/Source/DMXEditor/Public/DMXEditorSettings.h
+++ b/Engine/Plugins/VirtualProduction/DMX/DMXEngine/Source/DMXEditor/Public/DMXEditorSettings.h
@@ -138,8 +138,11 @@ public:
 	UPROPERTY(Config)
 	FDMXMVRFixtureListSettings MVRFixtureListSettings;
 
-	// Fixture Type Functions Editor
+	// DMX Library
 public:
+	UPROPERTY(Config)
+	bool bFixturePatcherDMXMonitorEnabled = false;
+
 	UPROPERTY(Config)
 	FDMXFixtureTypeFunctionsEditorSettings FixtureTypeFunctionsEditorSettings;
 
