DMX - Allow multi drag-drop Fixture Patches in the right side Fixture Patcher

diff --git a/Engine/Plugins/VirtualProduction/DMX/DMXEngine/Source/DMXEditor/Private/DragDrop/DMXEntityFixturePatchDragDropOp.cpp b/Engine/Plugins/VirtualProduction/DMX/DMXEngine/Source/DMXEditor/Private/DragDrop/DMXEntityFixturePatchDragDropOp.cpp
--- a/Engine/Plugins/VirtualProduction/DMX/DMXEngine/Source/DMXEditor/Private/DragDrop/DMXEntityFixturePatchDragDropOp.cpp
+++ b/Engine/Plugins/VirtualProduction/DMX/DMXEngine/Source/DMXEditor/Private/DragDrop/DMXEntityFixturePatchDragDropOp.cpp
@@ -2,8 +2,38 @@
 
 #include "DMXEntityFixturePatchDragDropOp.h"
 
+#include "Library/DMXEntityFixturePatch.h"
 
-FDMXEntityFixturePatchDragDropOperation::FDMXEntityFixturePatchDragDropOperation(UDMXLibrary* InLibrary, const TArray<TWeakObjectPtr<UDMXEntity>>& InFixturePatches, int32 InChannelOffset)
-	: FDMXEntityDragDropOperation(InLibrary, InFixturePatches)
-	, ChannelOffset(InChannelOffset)
-{}
+#include "Editor.h"
+
+
+#define LOCTEXT_NAMESPACE "FDMXEntityFixturePatchDragDropOperation"
+
+FDMXEntityFixturePatchDragDropOperation::FDMXEntityFixturePatchDragDropOperation(const TMap<TWeakObjectPtr<UDMXEntityFixturePatch>, int64>& InFixturePatchToAbsoluteChannelOffsetMap)
+	: FixturePatchToAbsoluteChannelOffsetMap(InFixturePatchToAbsoluteChannelOffsetMap)
+{
+	GEditor->BeginTransaction(FText::Format(LOCTEXT("DragDropTransaction", "Drag Fixture {0}|plural(one=Patch, other=Patches)"), FixturePatchToAbsoluteChannelOffsetMap.Num()));
+
+	for (const TTuple<TWeakObjectPtr<UDMXEntityFixturePatch>, int64>& FixturePatchToChannelOffsetPair : FixturePatchToAbsoluteChannelOffsetMap)
+	{
+		if (FixturePatchToChannelOffsetPair.Key.IsValid())
+		{
+			FixturePatchToChannelOffsetPair.Key->PreEditChange(UDMXEntityFixturePatch::StaticClass()->FindPropertyByName(UDMXEntityFixturePatch::GetUniverseIDPropertyNameChecked()));
+			FixturePatchToChannelOffsetPair.Key->PreEditChange(UDMXEntityFixturePatch::StaticClass()->FindPropertyByName(UDMXEntityFixturePatch::GetStartingChannelPropertyNameChecked()));
+		}
+	}
+}
+
+FDMXEntityFixturePatchDragDropOperation::~FDMXEntityFixturePatchDragDropOperation()
+{
+	for (const TTuple<TWeakObjectPtr<UDMXEntityFixturePatch>, int64>& FixturePatchToChannelOffsetPair : FixturePatchToAbsoluteChannelOffsetMap)
+	{
+		if (FixturePatchToChannelOffsetPair.Key.IsValid())
+		{
+			FixturePatchToChannelOffsetPair.Key->PostEditChange();
+		}
+	}
+	GEditor->EndTransaction();
+}
+
+#undef LOCTEXT_NAMESPACE
diff --git a/Engine/Plugins/VirtualProduction/DMX/DMXEngine/Source/DMXEditor/Private/DragDrop/DMXEntityFixturePatchDragDropOp.h b/Engine/Plugins/VirtualProduction/DMX/DMXEngine/Source/DMXEditor/Private/DragDrop/DMXEntityFixturePatchDragDropOp.h
--- a/Engine/Plugins/VirtualProduction/DMX/DMXEngine/Source/DMXEditor/Private/DragDrop/DMXEntityFixturePatchDragDropOp.h
+++ b/Engine/Plugins/VirtualProduction/DMX/DMXEngine/Source/DMXEditor/Private/DragDrop/DMXEntityFixturePatchDragDropOp.h
@@ -4,24 +4,30 @@
 
 #include "DMXEntityDragDropOp.h"
 
+#include "Input/DragAndDrop.h"
+
 
 class FDMXEntityFixturePatchDragDropOperation
-	: public FDMXEntityDragDropOperation
+	: public FDragDropOperation
 {
 public:
+	DRAG_DROP_OPERATOR_TYPE(FDMXEntityDragDropOperation, FDragDropOperation)
+
 	/** 
 	 * Constructs the fixture patch drag drop operation 
 	 * 
-	 * @param InLibrary					The library that contains the entities. Dragging from different libraries at once is not supported.
-	 * @param InFixturePatches			The dragged fixture patches.
-	 * @param InChannelOffset			The channel offset from the patch starting channel. Usually ChannelOffset = DraggedChannel - StartingChannel.
- 	 */
-	FDMXEntityFixturePatchDragDropOperation(UDMXLibrary* InLibrary, const TArray<TWeakObjectPtr<UDMXEntity>>& InFixturePatches, int32 InChannelOffset);
+	 * @param InLibrary								The library that contains the entities. Dragging from different libraries at once is not supported.
+	 * @param InFixturePatchToChannelOffsetMap		Each Fixture Patch being dragged, along with it's absolute channel offset from the anchor position.
+	 */
+	FDMXEntityFixturePatchDragDropOperation(const TMap<TWeakObjectPtr<UDMXEntityFixturePatch>, int64>& FixturePatchToAbsoluteChannelOffsetMap);
 
-	int32 GetChannelOffset() const { return ChannelOffset; }
+	/** Destructor */
+	virtual ~FDMXEntityFixturePatchDragDropOperation();
 
-	void SetChannelOffset(int32 InChannelOffset) { ChannelOffset = InChannelOffset; }
+	/** Returns Each Fixture Patch being dragged, along with it's channel offset from the anchor position. */
+	const TMap<TWeakObjectPtr<UDMXEntityFixturePatch>, int64>& GetFixturePatchToAbsoluteChannelOffsetMap() const { return FixturePatchToAbsoluteChannelOffsetMap; }
 
 protected:
-	int32 ChannelOffset;
+	/*	Each Fixture Patch being dragged, along with it's channel offset from the anchor position. */
+	TMap<TWeakObjectPtr<UDMXEntityFixturePatch>, int64> FixturePatchToAbsoluteChannelOffsetMap;
 };
diff --git a/Engine/Plugins/VirtualProduction/DMX/DMXEngine/Source/DMXEditor/Private/Widgets/FixturePatch/DMXFixturePatchNode.cpp b/Engine/Plugins/VirtualProduction/DMX/DMXEngine/Source/DMXEditor/Private/Widgets/FixturePatch/DMXFixturePatchNode.cpp
--- a/Engine/Plugins/VirtualProduction/DMX/DMXEngine/Source/DMXEditor/Private/Widgets/FixturePatch/DMXFixturePatchNode.cpp
+++ b/Engine/Plugins/VirtualProduction/DMX/DMXEngine/Source/DMXEditor/Private/Widgets/FixturePatch/DMXFixturePatchNode.cpp
@@ -16,7 +16,6 @@
 
 #define LOCTEXT_NAMESPACE "DMXFixturePatchNode"
 
-
 /** A fixture patch fragment in a grid. CreateFragments creates DMXFixturePatchFragment from a DMXFixturePatchNode. */
 class FDMXFixturePatchFragment
 	: public TSharedFromThis<FDMXFixturePatchFragment>
@@ -318,9 +317,6 @@ TSharedPtr<FDMXFixturePatchNode> FDMXFixturePatchNode::Create(TWeakPtr<FDMXEdito
 		return nullptr;
 	}
 
-	NewNode->LastTransactedUniverseID = InFixturePatch->GetUniverseID();
-	NewNode->LastTransactedChannelID = InFixturePatch->GetStartingChannel();
-
 	TSharedPtr<FDMXFixturePatchSharedData> SharedData = DMXEditor->GetFixturePatchSharedData();
 	if (!SharedData.IsValid())
 	{
@@ -339,7 +335,7 @@ TSharedPtr<FDMXFixturePatchNode> FDMXFixturePatchNode::Create(TWeakPtr<FDMXEdito
 	return NewNode;
 }
 
-void FDMXFixturePatchNode::SetAddresses(int32 NewUniverseID, int32 NewStartingChannel, int32 NewChannelSpan, bool bTransacted)
+void FDMXFixturePatchNode::SetAddresses(int32 NewUniverseID, int32 NewStartingChannel, int32 NewChannelSpan)
 {
 	if (!FixturePatch.IsValid())
 	{
@@ -350,29 +346,8 @@ void FDMXFixturePatchNode::SetAddresses(int32 NewUniverseID, int32 NewStartingCh
 	StartingChannel = NewStartingChannel;
 	ChannelSpan = NewChannelSpan;
 
-	if (bTransacted)
-	{
-		// As we allow the asset to update even when not transacted,
-		// We have to set the last known values to enable proper undo
-		FixturePatch->SetUniverseID(LastTransactedUniverseID);
-		FixturePatch->SetStartingChannel(LastTransactedChannelID);
-
-		LastTransactedUniverseID = UniverseID;
-		LastTransactedChannelID = StartingChannel;
-
-		const FScopedTransaction Transaction = FScopedTransaction(LOCTEXT("FixturePatchAssigned", "Patched Fixture"));
-		FixturePatch->PreEditChange(nullptr);
-
-		FixturePatch->SetUniverseID(NewUniverseID);
-		FixturePatch->SetStartingChannel(StartingChannel);
-
-		FixturePatch->PostEditChange();
-	}
-	else
-	{
-		FixturePatch->SetUniverseID(NewUniverseID);
-		FixturePatch->SetStartingChannel(StartingChannel);
-	}
+	FixturePatch->SetUniverseID(NewUniverseID);
+	FixturePatch->SetStartingChannel(StartingChannel);
 }
 
 TArray<TSharedPtr<SDMXFixturePatchFragment>> FDMXFixturePatchNode::GenerateWidgets(const TArray<TSharedPtr<FDMXFixturePatchNode>>& FixturePatchNodeGroup)
diff --git a/Engine/Plugins/VirtualProduction/DMX/DMXEngine/Source/DMXEditor/Private/Widgets/FixturePatch/DMXFixturePatchNode.h b/Engine/Plugins/VirtualProduction/DMX/DMXEngine/Source/DMXEditor/Private/Widgets/FixturePatch/DMXFixturePatchNode.h
--- a/Engine/Plugins/VirtualProduction/DMX/DMXEngine/Source/DMXEditor/Private/Widgets/FixturePatch/DMXFixturePatchNode.h
+++ b/Engine/Plugins/VirtualProduction/DMX/DMXEngine/Source/DMXEditor/Private/Widgets/FixturePatch/DMXFixturePatchNode.h
@@ -25,7 +25,7 @@ public:
 	static TSharedPtr<FDMXFixturePatchNode> Create(TWeakPtr<FDMXEditor> InDMXEditor, const TWeakObjectPtr<UDMXEntityFixturePatch>& InFixturePatch);
 
 	/** Sets the Addresses of the Node */
-	void SetAddresses(int32 UniverseID, int32 NewStartingChannel, int32 NewChannelSpan, bool bTransacted);
+	void SetAddresses(int32 UniverseID, int32 NewStartingChannel, int32 NewChannelSpan);
 
 	/** Generates widgets. Requires Addresses to be set via SetAddresses */
 	TArray<TSharedPtr<SDMXFixturePatchFragment>> GenerateWidgets(const TArray<TSharedPtr<FDMXFixturePatchNode>>& FixturePatchNodeGroup);
@@ -82,12 +82,6 @@ private:
 	/** Cached channel span of the patch */
 	int32 ChannelSpan = 0;
 
-	/** Last transacted Universe ID, required for propert undo/redo */
-	int32 LastTransactedUniverseID;
-
-	/** Last transacted Channel ID, required for propert undo/redo */
-	int32 LastTransactedChannelID;
-
 	/** If true the node is selected */
 	bool bSelected = false;
 
diff --git a/Engine/Plugins/VirtualProduction/DMX/DMXEngine/Source/DMXEditor/Private/Widgets/FixturePatch/SDMXFixturePatcher.cpp b/Engine/Plugins/VirtualProduction/DMX/DMXEngine/Source/DMXEditor/Private/Widgets/FixturePatch/SDMXFixturePatcher.cpp
--- a/Engine/Plugins/VirtualProduction/DMX/DMXEngine/Source/DMXEditor/Private/Widgets/FixturePatch/SDMXFixturePatcher.cpp
+++ b/Engine/Plugins/VirtualProduction/DMX/DMXEngine/Source/DMXEditor/Private/Widgets/FixturePatch/SDMXFixturePatcher.cpp
@@ -16,6 +16,7 @@
 #include "Widgets/Layout/SBox.h"
 #include "Widgets/SNullWidget.h"
 
+#include "Algo/RemoveIf.h"
 #include "Styling/AppStyle.h"
 #include "ScopedTransaction.h"
 #include "SlateOptMacros.h"
@@ -245,183 +246,133 @@ FReply SDMXFixturePatcher::OnDrop(const FGeometry& MyGeometry, const FDragDropEv
 
 void SDMXFixturePatcher::OnDragEnterChannel(int32 UniverseID, int32 ChannelID, const FDragDropEvent& DragDropEvent)
 {
-	TSharedPtr<FDragDropOperation> Operation = DragDropEvent.GetOperation();
-	if (!Operation.IsValid())
+	// Same as dropping onto a channel, but without FReply
+	OnDropOntoChannel(UniverseID, ChannelID, DragDropEvent);
+}
+
+FReply SDMXFixturePatcher::OnDropOntoChannel(int32 UniverseID, int32 ChannelID, const FDragDropEvent& DragDropEvent)
+{
+	const TMap<TSharedRef<FDMXFixturePatchNode>, int64> DraggedNodesToAbsoluteChannelOffsetMap = GetDraggedNodesToAbsoluteChannelOffsetMap(DragDropEvent);
+	if (DraggedNodesToAbsoluteChannelOffsetMap.IsEmpty())
 	{
-		return;
+		return FReply::Handled().EndDragDrop();
 	}
 
-	if (TSharedPtr<FDMXEntityFixturePatchDragDropOperation> FixturePatchDragDropOp = DragDropEvent.GetOperationAs<FDMXEntityFixturePatchDragDropOperation>())
+	int32 PadOffset = 0;
+	for (const TTuple<TSharedRef<FDMXFixturePatchNode>, int64>& DraggedNodesToChannelOffsetPair : DraggedNodesToAbsoluteChannelOffsetMap)
 	{
-		const TArray<TWeakObjectPtr<UDMXEntity>>& DraggedEntities = FixturePatchDragDropOp->GetDraggedEntities();
+		const TSharedRef<FDMXFixturePatchNode>& PatchedNode = DraggedNodesToChannelOffsetPair.Key;
 
-		if (DraggedEntities.Num() > 1)
+		if (UDMXEntityFixturePatch* FixturePatch = DraggedNodesToChannelOffsetPair.Key->GetFixturePatch().Get())
 		{
-			FixturePatchDragDropOp->SetFeedbackMessageError(LOCTEXT("CannotDragDropMoreThanOnePatch", "Multi asset drag drop is not supported."));
-		}
-
-		TSharedPtr<FDMXFixturePatchNode> DraggedNode = GetDraggedNode(DraggedEntities);
-
-		if (DraggedNode.IsValid())
-		{
-			if (UDMXEntityFixturePatch* FixturePatch = Cast<UDMXEntityFixturePatch>(DraggedEntities[0]))
-			{
-				const int32 ChannelSpan = FixturePatch->GetChannelSpan();
-				const int32 NewStartingChannel = [this, ChannelID, &FixturePatchDragDropOp, ChannelSpan]()
+			const int32 ChannelSpan = FixturePatch->GetChannelSpan();
+			const int32 AbsoluteOffset = DraggedNodesToChannelOffsetPair.Value + PadOffset;
+
+			int32 NewUniverseID = UniverseID;
+			int32 DesiredStartingChannel = ChannelID;
+			if (IsUniverseSelectionEnabled())
+			{			
+				// Stay within the current Universe if only one is displayed
+				if (FixturePatch->GetUniverseID() != UniverseID)
 				{
-					int32 StartingChannel = ChannelID - FixturePatchDragDropOp->GetChannelOffset();
-					return ClampStartingChannel(StartingChannel, ChannelSpan);
-				}();
-
+					continue;
+				}
 
-				// Patch the node but do not transact it (transact on drop instead)
-				const TSharedPtr<SDMXPatchedUniverse>& Universe = PatchedUniversesByID.FindChecked(UniverseID);
+				DesiredStartingChannel = ChannelID - AbsoluteOffset % DMX_MAX_ADDRESS;
+			}
+			else
+			{
+				int64 AbsoluteDesiredStartingChannel = UniverseID * DMX_MAX_ADDRESS + ChannelID - AbsoluteOffset;
 
-				const bool bCreateTransaction = false;
-				bool bPatchSuccess = Universe->Patch(DraggedNode, NewStartingChannel, bCreateTransaction);
-			
-				// If patching wasn't successful, try to move as close to the hovered node as possible
-				if (!bPatchSuccess)
+				// Allow drag to another universe if multiple universes are displayed
+				int32 DesiredUniverse = AbsoluteDesiredStartingChannel / DMX_MAX_ADDRESS;
+				if (DesiredUniverse < 1)
 				{
-					const bool bIsHoveredChannelInFrontOfPatch = [UniverseID, NewStartingChannel, FixturePatch]()
-					{
-						bool bUniverseInFront = UniverseID < FixturePatch->GetUniverseID();
-						bool ChannelInFront = UniverseID == FixturePatch->GetUniverseID() && NewStartingChannel < FixturePatch->GetStartingChannel();
-
-						return bUniverseInFront || ChannelInFront;
-					}();
-
-					if (bIsHoveredChannelInFrontOfPatch)
-					{
-						for (int32 Channel = FixturePatch->GetStartingChannel() - 1; Channel >= NewStartingChannel; Channel--)
-						{
-							// Try to patch as close as possible, but consider any approximation a success
-							bool bCloserApproximation = Universe->Patch(DraggedNode, Channel, bCreateTransaction);
-							bPatchSuccess = bPatchSuccess ? true : bCloserApproximation;
-						}
-					}
-					else
-					{
-						for (int32 Channel = FixturePatch->GetStartingChannel() + 1; Channel <= NewStartingChannel; Channel++)
-						{
-							// Try to patch as close as possible, but consider any approximation a success
-							bool bCloserApproximation = Universe->Patch(DraggedNode, Channel, bCreateTransaction);
-							bPatchSuccess = bPatchSuccess ? true : bCloserApproximation;
-						}
-					}
+					DesiredUniverse = 1;
+					AbsoluteDesiredStartingChannel = ChannelID - AbsoluteOffset % DMX_MAX_ADDRESS;
 				}
-
-				if (bPatchSuccess)
+				else if (DesiredUniverse > DMX_MAX_UNIVERSE)
 				{
-					TSharedRef<SWidget> DragDropDecorator = CreateDragDropDecorator(FixturePatch, NewStartingChannel);
-					FixturePatchDragDropOp->SetCustomFeedbackWidget(DragDropDecorator);
+					DesiredUniverse = DMX_MAX_UNIVERSE;
+					AbsoluteDesiredStartingChannel = DMX_MAX_UNIVERSE * DMX_MAX_ADDRESS - ChannelSpan + 1;
 				}
-				else if (!DraggedNode->IsPatched())
+				else
 				{
-					if (NewStartingChannel + ChannelSpan > DMX_UNIVERSE_SIZE)
+					NewUniverseID = DesiredUniverse;
+				}
+
+				// Remove from previous universe if needed
+				if (NewUniverseID != PatchedNode->GetUniverseID())
+				{
+					const TSharedPtr<SDMXPatchedUniverse>* OldUniverseWidgetPtr = PatchedUniversesByID.Find(PatchedNode->GetUniverseID());
+					if (OldUniverseWidgetPtr)
 					{
-						FixturePatchDragDropOp->SetFeedbackMessageError(LOCTEXT("CannotDragDropOnOccupiedChannels", "Channels range overflows max channels address (512)"));
+						(*OldUniverseWidgetPtr)->Remove(PatchedNode);
+						(*OldUniverseWidgetPtr)->RequestRefresh();
 					}
 				}
+
+				DesiredStartingChannel = AbsoluteDesiredStartingChannel % DMX_MAX_ADDRESS;
+			}
+
+			const int32 NewStartingChannel = ClampStartingChannel(DesiredStartingChannel, ChannelSpan);
+			DraggedNodesToChannelOffsetPair.Key->SetAddresses(NewUniverseID, NewStartingChannel, ChannelSpan);
+
+			const TSharedPtr<SDMXPatchedUniverse>* UniverseWidgetPtr = PatchedUniversesByID.Find(NewUniverseID);
+			if (!UniverseWidgetPtr)
+			{
+				AddUniverse(NewUniverseID);
+				UniverseWidgetPtr = &PatchedUniversesByID.FindChecked(NewUniverseID);
 			}
+			check(UniverseWidgetPtr);
+			(*UniverseWidgetPtr)->FindOrAdd(DraggedNodesToChannelOffsetPair.Key);
+
+			// Pad further patches so they can't be offset more than the currently dropped patch
+			PadOffset = FMath::Min(PadOffset, DesiredStartingChannel - NewStartingChannel);
 		}
 	}
+
+	return FReply::Handled().EndDragDrop();
 }
 
-FReply SDMXFixturePatcher::OnDropOntoChannel(int32 UniverseID, int32 ChannelID, const FDragDropEvent& DragDropEvent)
+TMap<TSharedRef<FDMXFixturePatchNode>, int64> SDMXFixturePatcher::GetDraggedNodesToAbsoluteChannelOffsetMap(const FDragDropEvent& DragDropEvent) const
 {
+	TMap<TSharedRef<FDMXFixturePatchNode>, int64> DraggedNodesToAbsoluteChannelOffsetMap;
+
 	TSharedPtr<FDragDropOperation> Operation = DragDropEvent.GetOperation();
 	if (!Operation.IsValid())
 	{
-		return FReply::Unhandled();
+		return DraggedNodesToAbsoluteChannelOffsetMap;
 	}
 
-	if (TSharedPtr<FDMXEntityFixturePatchDragDropOperation> FixturePatchDragDropOp = DragDropEvent.GetOperationAs<FDMXEntityFixturePatchDragDropOperation>())
+	const TSharedPtr<FDMXEntityFixturePatchDragDropOperation> FixturePatchDragDropOp = DragDropEvent.GetOperationAs<FDMXEntityFixturePatchDragDropOperation>();
+	if (!FixturePatchDragDropOp.IsValid())
 	{
-		const TArray<TWeakObjectPtr<UDMXEntity>>& DraggedEntities = FixturePatchDragDropOp->GetDraggedEntities();
-
-		TSharedPtr<FDMXFixturePatchNode> DraggedNode = GetDraggedNode(DraggedEntities);
-
-		if (DraggedNode.IsValid())
-		{
-			if (UDMXEntityFixturePatch* FixturePatch = Cast<UDMXEntityFixturePatch>(DraggedEntities[0]))
-			{
-				const int32 ChannelSpan = FixturePatch->GetChannelSpan();
-				const int32 NewStartingChannel = [this, ChannelID, &FixturePatchDragDropOp, ChannelSpan]()
-				{
-					int32 StartingChannel = ChannelID - FixturePatchDragDropOp->GetChannelOffset();
-					return ClampStartingChannel(StartingChannel, ChannelSpan);
-				}();
-
-				const TSharedPtr<SDMXPatchedUniverse>& Universe = PatchedUniversesByID.FindChecked(UniverseID);
-
-				bool bCreateTransaction = true;
-				if (Universe->Patch(DraggedNode, NewStartingChannel, bCreateTransaction))
-				{
-					return FReply::Handled().EndDragDrop();
-				}
-			}
-		}
+		return DraggedNodesToAbsoluteChannelOffsetMap;
 	}
 
-	return FReply::Unhandled();
-}
-
-TSharedPtr<FDMXFixturePatchNode> SDMXFixturePatcher::GetDraggedNode(const TArray<TWeakObjectPtr<UDMXEntity>>& DraggedEntities)
-{
-	if (DraggedEntities.Num() == 1)
+	const TMap<TWeakObjectPtr<UDMXEntityFixturePatch>, int64>& FixturePatchToAbsoluteChannelOffsetMap = FixturePatchDragDropOp->GetFixturePatchToAbsoluteChannelOffsetMap();
+	for (const TTuple<TWeakObjectPtr<UDMXEntityFixturePatch>, int64>& FixturePatchToAbsoluteChannelOffsetPair : FixturePatchToAbsoluteChannelOffsetMap)
 	{
-		if (UDMXEntityFixturePatch* FixturePatch = Cast<UDMXEntityFixturePatch>(DraggedEntities[0]))
+		UDMXEntityFixturePatch* FixturePatch = FixturePatchToAbsoluteChannelOffsetPair.Key.Get();
+		if (!FixturePatch)
 		{
-			const TArray<TWeakObjectPtr<UDMXEntityFixturePatch>> FixturePatchArray = TArray<TWeakObjectPtr<UDMXEntityFixturePatch>>({ FixturePatch });
-			TSharedPtr<FDMXFixturePatchNode> DraggedNode = FindPatchNode(FixturePatch);
-
-			if (!DraggedNode.IsValid())
-			{
-				DraggedNode = FDMXFixturePatchNode::Create(DMXEditorPtr, FixturePatch);
-			}
-
-			return DraggedNode;
+			continue;
 		}
-	}
 
-	return nullptr;
-}
-
-TSharedRef<SWidget> SDMXFixturePatcher::CreateDragDropDecorator(TWeakObjectPtr<UDMXEntityFixturePatch> FixturePatch, int32 ChannelID) const
-{
-	if (FixturePatch.IsValid())
-	{
-		int32 StartingChannel = ChannelID;
-		int32 ChannelSpan = FixturePatch->GetChannelSpan();
-		int32 EndingChannel = StartingChannel + FixturePatch->GetChannelSpan() - 1;
-
-		FText PatchName = FText::Format(LOCTEXT("PatchName", "{0}"), FText::FromString(FixturePatch->GetDisplayName()));;								
-		FText ChannelRangeName = FText::Format(LOCTEXT("ChannelRangeName", "Channel {0} - {1}"), StartingChannel, EndingChannel);
+		TSharedPtr<FDMXFixturePatchNode> DraggedNode = FindPatchNode(FixturePatch);
+		if (!DraggedNode.IsValid())
+		{
+			DraggedNode = FDMXFixturePatchNode::Create(DMXEditorPtr, FixturePatch);
+		}
 
-		return SNew(SBorder)
-			.BorderImage(FAppStyle::GetBrush("Graph.ConnectorFeedback.Border"))
-			[
-				SNew(SVerticalBox)				
-				+ SVerticalBox::Slot()
-				.VAlign(VAlign_Fill)
-				[
-					SNew(STextBlock)
-					.Font(FAppStyle::GetFontStyle("PropertyWindow.NormalFont"))
-					.Text(ChannelRangeName)
-				]
-				+ SVerticalBox::Slot()
-				.VAlign(VAlign_Bottom)
-				[
-					SNew(STextBlock)
-					.Text(PatchName)
-					.Font(FAppStyle::GetFontStyle("PropertyWindow.NormalFont"))
-					.ColorAndOpacity(FLinearColor(0.9f, 0.9f, 0.9f, 1.0f))
-				]
-			];
+		if (DraggedNode.IsValid())
+		{
+			DraggedNodesToAbsoluteChannelOffsetMap.Add(DraggedNode.ToSharedRef(), FixturePatchToAbsoluteChannelOffsetPair.Value);
+		}
 	}
 
-	return SNullWidget::NullWidget;
+	return DraggedNodesToAbsoluteChannelOffsetMap;
 }
 
 void SDMXFixturePatcher::PostUndo(bool bSuccess)
@@ -432,12 +383,12 @@ void SDMXFixturePatcher::PostUndo(bool bSuccess)
 		DMXLibrary->Modify();
 	}
 
-	RefreshFromProperties();
+	RefreshFromLibrary();
 }
 
 void SDMXFixturePatcher::PostRedo(bool bSuccess)
 {
-	RefreshFromProperties();
+	RefreshFromLibrary();
 }
 
 TSharedPtr<FDMXFixturePatchNode> SDMXFixturePatcher::FindPatchNode(const TWeakObjectPtr<UDMXEntityFixturePatch> FixturePatch) const
diff --git a/Engine/Plugins/VirtualProduction/DMX/DMXEngine/Source/DMXEditor/Private/Widgets/FixturePatch/SDMXFixturePatcher.h b/Engine/Plugins/VirtualProduction/DMX/DMXEngine/Source/DMXEditor/Private/Widgets/FixturePatch/SDMXFixturePatcher.h
--- a/Engine/Plugins/VirtualProduction/DMX/DMXEngine/Source/DMXEditor/Private/Widgets/FixturePatch/SDMXFixturePatcher.h
+++ b/Engine/Plugins/VirtualProduction/DMX/DMXEngine/Source/DMXEditor/Private/Widgets/FixturePatch/SDMXFixturePatcher.h
@@ -10,6 +10,7 @@
 #include "Widgets/SCompoundWidget.h"
 
 class FDMXEditor;
+class FDMXEntityFixturePatchDragDropOperation;
 class FDMXFixturePatchNode;
 class FDMXFixturePatchSharedData;
 class FDMXTreeNodeBase;
@@ -60,12 +61,9 @@ protected:
 	/** Called when drag dropped onto a channel */
 	FReply OnDropOntoChannel(int32 UniverseID, int32 ChannelID, const FDragDropEvent& DragDropEvent);
 
-	/** Initilizes the widget for an incoming dragged patch, and the patch so it can be dragged here */
-	TSharedPtr<FDMXFixturePatchNode> GetDraggedNode(const TArray<TWeakObjectPtr<UDMXEntity>>& DraggedEntities);
+	/** Returns a map of nodes being dragged along with their absolute channel offset from the Drag Drop Event */
+	TMap<TSharedRef<FDMXFixturePatchNode>, int64> GetDraggedNodesToAbsoluteChannelOffsetMap(const FDragDropEvent& DragDropEvent) const;
 	
-	/** Creates an info widget for drag dropping */
-	TSharedRef<SWidget> CreateDragDropDecorator(TWeakObjectPtr<UDMXEntityFixturePatch> FixturePatch, int32 ChannelID) const;
-
 	//~ Begin FEditorUndoClient Interface
 	virtual void PostUndo(bool bSuccess) override;
 	virtual void PostRedo(bool bSuccess) override;
@@ -138,7 +136,7 @@ protected:
 	FText GetTooltipText() const;
 
 protected:
-	/** Clamps the starting channel to remain within a valid channel range */
+	/** Creates a valid address given Universe, Starting Channel and Channel Span. Returns false if no valid Address can be found. */
 	int32 ClampStartingChannel(int32 StartingChannel, int32 ChannelSpan) const;
 
 	/** Returns the DMXLibrary or nullptr if not available */
diff --git a/Engine/Plugins/VirtualProduction/DMX/DMXEngine/Source/DMXEditor/Private/Widgets/FixturePatch/SDMXPatchedUniverse.cpp b/Engine/Plugins/VirtualProduction/DMX/DMXEngine/Source/DMXEditor/Private/Widgets/FixturePatch/SDMXPatchedUniverse.cpp
--- a/Engine/Plugins/VirtualProduction/DMX/DMXEngine/Source/DMXEditor/Private/Widgets/FixturePatch/SDMXPatchedUniverse.cpp
+++ b/Engine/Plugins/VirtualProduction/DMX/DMXEngine/Source/DMXEditor/Private/Widgets/FixturePatch/SDMXPatchedUniverse.cpp
@@ -17,6 +17,7 @@
 #include "Library/DMXEntityFixturePatch.h"
 #include "Library/DMXLibrary.h"
 
+#include "Algo/RemoveIf.h"
 #include "Framework/Commands/GenericCommands.h"
 #include "Framework/Commands/UICommandList.h"
 #include "Framework/MultiBox/MultiBoxBuilder.h"
@@ -33,6 +34,8 @@ void SDMXPatchedUniverse::Construct(const FArguments& InArgs)
 {
 	check(InArgs._UniverseID != INDEX_NONE);
 
+	SetCanTick(true);
+
 	UniverseID = InArgs._UniverseID;
 
 	DMXEditorPtr = InArgs._DMXEditor;
@@ -166,9 +169,12 @@ void SDMXPatchedUniverse::CreateChannelConnectors()
 	}
 }
 
-bool SDMXPatchedUniverse::Patch(const TSharedPtr<FDMXFixturePatchNode>& Node, int32 NewStartingChannel, bool bCreateTransaction)
+bool SDMXPatchedUniverse::FindOrAdd(const TSharedPtr<FDMXFixturePatchNode>& Node)
 {
-	check(Node.IsValid());
+	if (!Node.IsValid())
+	{
+		return false;
+	}
 
 	UDMXEntityFixturePatch* FixturePatch = Node->GetFixturePatch().Get();
 	if (!FixturePatch)
@@ -178,13 +184,28 @@ bool SDMXPatchedUniverse::Patch(const TSharedPtr<FDMXFixturePatchNode>& Node, in
 
 	PatchedNodes.AddUnique(Node);
 
-	const int32 NewChannelSpan = FixturePatch->GetChannelSpan();
-	Node->SetAddresses(UniverseID, NewStartingChannel, NewChannelSpan, bCreateTransaction);
 	RequestRefresh();
 
 	return true;
 }
 
+void SDMXPatchedUniverse::Remove(const TSharedPtr<FDMXFixturePatchNode>& Node)
+{
+	if (PatchedNodes.Contains(Node))
+	{
+		PatchedNodes.RemoveSingle(Node);
+
+		for (const TSharedPtr<SDMXFixturePatchFragment>& Widget : Node->GetGeneratedWidgets())
+		{
+			if (Widget.IsValid())
+			{
+				Grid->RemoveSlot(Widget.ToSharedRef());
+			}
+		}
+		RequestRefresh();
+	}
+}
+
 bool SDMXPatchedUniverse::CanAssignFixturePatch(TWeakObjectPtr<UDMXEntityFixturePatch> FixturePatch) const
 {
 	if (FixturePatch.IsValid())
@@ -288,51 +309,12 @@ TSharedPtr<FDMXFixturePatchNode> SDMXPatchedUniverse::FindPatchNodeOfType(UDMXEn
 
 void SDMXPatchedUniverse::Tick(const FGeometry& AllottedGeometry, const double InCurrentTime, const float InDeltaTime) 
 {
-	TArray<TSharedPtr<FDMXFixturePatchNode>> RemovedNodes;
-	for (const TSharedPtr<FDMXFixturePatchNode>& Node : PatchedNodes)
-	{
-		if (Node->GetUniverseID() != UniverseID)
-		{
-			RemovedNodes.Add(Node);
-		}
-		else if (Node->NeedsUpdateGrid())
-		{
-			RequestRefresh();
-			break;
-		}
-	}
-
-	for (const TSharedPtr<FDMXFixturePatchNode>& RemovedNode : RemovedNodes)
-	{
-		Unpatch(RemovedNode);
-	}
-
 	if (bMonitorInputs)
 	{
 		UpdateMonitor();
 	}
 }
 
-void SDMXPatchedUniverse::Unpatch(const TSharedPtr<FDMXFixturePatchNode>& Node)
-{
-	check(Node.IsValid());
-	check(Grid.IsValid());
-
-	if (PatchedNodes.Contains(Node))
-	{
-		PatchedNodes.RemoveSingle(Node);
-
-		for (const TSharedPtr<SDMXFixturePatchFragment>& Widget : Node->GetGeneratedWidgets())
-		{
-			if (Widget.IsValid())
-			{
-				Grid->RemoveSlot(Widget.ToSharedRef());
-			}
-		}
-		RequestRefresh();
-	}
-}
-
 void SDMXPatchedUniverse::OnFixturePatchChanged(const UDMXEntityFixturePatch* FixturePatch)
 {
 	if (!FixturePatch || FixturePatch->GetParentLibrary() != GetDMXLibrary())
@@ -340,23 +322,20 @@ void SDMXPatchedUniverse::OnFixturePatchChanged(const UDMXEntityFixturePatch* Fi
 		return;
 	}
 
-	const int32 UniverseIDOfFixturePatch = FixturePatch->GetUniverseID();
-	if (UniverseIDOfFixturePatch == UniverseID)
-	{
-		const TSharedPtr<FDMXFixturePatchNode>* NodePtr = PatchedNodes.FindByPredicate([FixturePatch](const TSharedPtr<FDMXFixturePatchNode>& Node)
-			{
-				return Node->GetFixturePatch() == FixturePatch;
-			});
-
-		// Redraw all if the patch cannot be found.
-		if (!NodePtr)
+	const TSharedPtr<FDMXFixturePatchNode>* NodePtr = PatchedNodes.FindByPredicate([FixturePatch](const TSharedPtr<FDMXFixturePatchNode>& Node)
 		{
-			SetUniverseIDInternal(UniverseID);
-		}
-	}
+			return Node->GetFixturePatch() == FixturePatch;
+		});
 
-	FDMXEditorUtils::ClearAllDMXPortBuffers();
-	FDMXEditorUtils::ClearFixturePatchCachedData();
+	if (NodePtr && (*NodePtr)->GetUniverseID() != UniverseID)
+	{
+		Remove(*NodePtr);
+	}
+	else if (FixturePatch->GetUniverseID() == UniverseID)
+	{
+		// Redraw all if the patch is in this Universe but not found
+		SetUniverseIDInternal(UniverseID);
+	}
 }
 
 void SDMXPatchedUniverse::OnReceiveDMXEnabledChanged(bool bNewEnabled)
@@ -369,7 +348,10 @@ void SDMXPatchedUniverse::OnReceiveDMXEnabledChanged(bool bNewEnabled)
 
 void SDMXPatchedUniverse::SetUniverseIDInternal(int32 NewUniverseID)
 {
-	check(NewUniverseID > -1);
+	if (NewUniverseID < 0)
+	{
+		return;
+	}
 
 	// Find patches in new universe
 	UDMXLibrary* Library = GetDMXLibrary();
@@ -379,7 +361,7 @@ void SDMXPatchedUniverse::SetUniverseIDInternal(int32 NewUniverseID)
 		TArray<TSharedPtr<FDMXFixturePatchNode>> CachedPatchedNodes = PatchedNodes;
 		for (const TSharedPtr<FDMXFixturePatchNode>& Node : CachedPatchedNodes)
 		{
-			Unpatch(Node);
+			Remove(Node);
 		}
 		check(PatchedNodes.Num() == 0);
 
@@ -414,7 +396,7 @@ void SDMXPatchedUniverse::SetUniverseIDInternal(int32 NewUniverseID)
 			}
 			check(Node.IsValid());
 
-			Patch(Node, FixturePatch->GetStartingChannel(), false);
+			FindOrAdd(Node);
 		}
 
 		// Update the channel connectors' Universe ID
@@ -513,47 +495,92 @@ void SDMXPatchedUniverse::ResetMonitor()
 
 FReply SDMXPatchedUniverse::HandleOnMouseButtonDownOnChannel(uint32 Channel, const FPointerEvent& PointerEvent)
 {
-	if (PointerEvent.GetEffectingButton() == EKeys::RightMouseButton)
+	if (PointerEvent.GetEffectingButton() == EKeys::LeftMouseButton)
 	{
-		const TArray<TWeakObjectPtr<UDMXEntityFixturePatch>> OldSelection = SharedData->GetSelectedFixturePatches();
-		const TArray<UDMXEntityFixturePatch*> FixturePatchesOnChannel = GetFixturePatchesOnChannel(Channel);
-		const bool bOldSelectionContainsPatchesOnChannel = FixturePatchesOnChannel.ContainsByPredicate([&OldSelection](const UDMXEntityFixturePatch* FixturePatch)
+		// Get valid selection, sorted by Universe/Address
+		TArray<UDMXEntityFixturePatch*> Selection;
+		Algo::TransformIf(SharedData->GetSelectedFixturePatches(), Selection,
+			[](TWeakObjectPtr<UDMXEntityFixturePatch> WeakFixturePatch) { return WeakFixturePatch.IsValid(); },
+			[](TWeakObjectPtr<UDMXEntityFixturePatch> WeakFixturePatch) { return WeakFixturePatch.Get(); }
+		);
+		Algo::SortBy(Selection, [](UDMXEntityFixturePatch* FixturePatch)
 			{
-				return OldSelection.Contains(FixturePatch);
+				return FixturePatch->GetUniverseID() * DMX_MAX_ADDRESS + FixturePatch->GetStartingChannel();
 			});
-		if (!bOldSelectionContainsPatchesOnChannel)
+
+		const TArray<UDMXEntityFixturePatch*> FixturePatchesOnChannel = GetFixturePatchesOnChannel(Channel);
+		UDMXEntityFixturePatch* ClickedPatch = GetTopmostFixturePatchOnChannel(Channel);
+
+		// Shift Select
+		if (FSlateApplication::Get().GetModifierKeys().IsShiftDown() &&
+			!Selection.IsEmpty())
 		{
-			if (UDMXEntityFixturePatch* Patch = GetTopmostFixturePatchOnChannel(Channel))
+			ShiftSelectAnchorPatch = ShiftSelectAnchorPatch.IsValid() ? ShiftSelectAnchorPatch : Selection[0];
+			
+			UDMXLibrary* DMXLibrary = ShiftSelectAnchorPatch->GetParentLibrary();
+			if (!DMXLibrary)
 			{
-				if (FSlateApplication::Get().GetModifierKeys().IsControlDown())
-				{
-					SharedData->AddFixturePatchToSelection(Patch);
-				}
-				else
+				return FReply::Handled();
+			}
+			TArray<UDMXEntityFixturePatch*> FixturePatches = DMXLibrary->GetEntitiesTypeCast<UDMXEntityFixturePatch>();
+			const int64 AbsoluteClickedChannel = UniverseID * DMX_MAX_ADDRESS + Channel;
+			const int64 ShiftSelectAbsoluteStartingChannel = ShiftSelectAnchorPatch->GetUniverseID() * DMX_MAX_ADDRESS + ShiftSelectAnchorPatch->GetStartingChannel();
+			const int64 ShiftSelectAbsoluteEndingChannel = ShiftSelectAnchorPatch->GetUniverseID() * DMX_MAX_ADDRESS + ShiftSelectAnchorPatch->GetEndingChannel();
+
+			const TRange<int64> SelectedRange(FMath::Min(ShiftSelectAbsoluteStartingChannel, AbsoluteClickedChannel), FMath::Max(ShiftSelectAbsoluteStartingChannel, AbsoluteClickedChannel + 1));
+
+			TArray<TWeakObjectPtr<UDMXEntityFixturePatch>> ShiftSelection;
+			Algo::TransformIf(FixturePatches, ShiftSelection,
+				[&SelectedRange, Channel, this](UDMXEntityFixturePatch* FixturePatch)
+				{			
+					const int64 PatchAbsoluteStartingChannel = FixturePatch->GetUniverseID() * DMX_MAX_ADDRESS + FixturePatch->GetStartingChannel();
+					const int64 PatchAbsoluteEndingChannel = FixturePatch->GetUniverseID() * DMX_MAX_ADDRESS + FixturePatch->GetEndingChannel();
+					const TRange<int64> PatchRange(PatchAbsoluteStartingChannel, PatchAbsoluteEndingChannel + 1);
+					return 
+						SelectedRange.Overlaps(PatchRange);
+				},
+				[](UDMXEntityFixturePatch* FixturePatch)
 				{
-					SharedData->SelectFixturePatch(Patch);
+					return FixturePatch;
 				}
-			}
+			);
+			SharedData->SelectFixturePatches(ShiftSelection);
+
+			return FReply::Handled();
 		}
-	}
+		ShiftSelectAnchorPatch.Reset();
 
-	if (PointerEvent.GetEffectingButton() == EKeys::LeftMouseButton)
-	{
-		if (UDMXEntityFixturePatch* Patch = GetTopmostFixturePatchOnChannel(Channel))
-		{
-			if (FSlateApplication::Get().GetModifierKeys().IsControlDown())
+		// Ctrl-Select
+		if (FSlateApplication::Get().GetModifierKeys().IsControlDown() && 
+			ClickedPatch)
+		{		
+			if (const bool bUnselect = Selection.Num() > 1 && Selection.Contains(ClickedPatch))
 			{
-				SharedData->AddFixturePatchToSelection(Patch);
+				TArray<TWeakObjectPtr<UDMXEntityFixturePatch>> WeakSelection = SharedData->GetSelectedFixturePatches();
+				WeakSelection.Remove(ClickedPatch);
+				SharedData->SelectFixturePatches(WeakSelection);
 			}
 			else
 			{
-				SharedData->SelectFixturePatch(Patch);
+				SharedData->AddFixturePatchToSelection(ClickedPatch);
 			}
 
-			if (ensureMsgf(ChannelConnectors.IsValidIndex(Channel - 1), TEXT("Trying to drag Fixture Patch, but the dragged channel is not valid.")))
+			return FReply::Handled();
+		}
+
+		// Normal Select and Detect Drag
+		if (ClickedPatch &&
+			ensureMsgf(ChannelConnectors.IsValidIndex(Channel - 1), TEXT("Trying to drag Fixture Patch, but the dragged channel is not valid.")))
+		{
+			ShiftSelectAnchorPatch.Reset();
+
+			// Replace selection with clicked patch if it wasn't selected
+			if (!Selection.Contains(ClickedPatch))
 			{
-				return FReply::Handled().DetectDrag(ChannelConnectors[Channel - 1].ToSharedRef(), EKeys::LeftMouseButton);
+				SharedData->SelectFixturePatch(ClickedPatch);
 			}
+
+			return FReply::Handled().DetectDrag(ChannelConnectors[Channel - 1].ToSharedRef(), EKeys::LeftMouseButton);
 		}
 	}
 
@@ -562,15 +589,32 @@ FReply SDMXPatchedUniverse::HandleOnMouseButtonDownOnChannel(uint32 Channel, con
 
 FReply SDMXPatchedUniverse::HandleOnMouseButtonUpOnChannel(uint32 Channel, const FPointerEvent& PointerEvent)
 {
+	if (PointerEvent.GetEffectingButton() == EKeys::LeftMouseButton)
+	{
+		UDMXEntityFixturePatch* Patch = GetTopmostFixturePatchOnChannel(Channel);
+		if (!Patch)
+		{
+			return FReply::Unhandled();
+		}
+
+		// Normal Select
+		if (!FSlateApplication::Get().GetModifierKeys().IsControlDown() && !FSlateApplication::Get().GetModifierKeys().IsShiftDown())
+		{
+			UDMXEntityFixturePatch* ClickedPatch = GetTopmostFixturePatchOnChannel(Channel);
+			if (ClickedPatch)
+			{
+				SharedData->SelectFixturePatch(ClickedPatch);
+			}
+		}
+	}
+
 	if (PointerEvent.GetEffectingButton() == EKeys::RightMouseButton)
 	{
-		const FVector2D& SummonLocation = PointerEvent.GetScreenSpacePosition();
+		const FVector2D& ContextMenuSummonLocation = PointerEvent.GetScreenSpacePosition();
 
 		TSharedRef<SWidget> MenuContent = CreateContextMenu(Channel);
 		FWidgetPath WidgetPath = PointerEvent.GetEventPath() ? *PointerEvent.GetEventPath() : FWidgetPath();
-		FSlateApplication::Get().PushMenu(AsShared(), WidgetPath, MenuContent, SummonLocation, FPopupTransitionEffect(FPopupTransitionEffect::ContextMenu));
-
-		return FReply::Handled();
+		FSlateApplication::Get().PushMenu(AsShared(), WidgetPath, MenuContent, ContextMenuSummonLocation, FPopupTransitionEffect(FPopupTransitionEffect::ContextMenu));
 	}
 
 	return FReply::Unhandled();
@@ -578,16 +622,41 @@ FReply SDMXPatchedUniverse::HandleOnMouseButtonUpOnChannel(uint32 Channel, const
 
 FReply SDMXPatchedUniverse::HandleOnDragDetectedOnChannel(uint32 Channel, const FPointerEvent& PointerEvent)
 {
-	if (UDMXEntityFixturePatch* Patch = GetTopmostFixturePatchOnChannel(Channel))
+	TArray<TWeakObjectPtr<UDMXEntityFixturePatch>> SelectedFixturePatches = SharedData->GetSelectedFixturePatches();
+
+	// Only valid patches, sorted by Universe/Address
+	const int32 NumValidPatches = Algo::RemoveIf(SelectedFixturePatches, [this](TWeakObjectPtr<UDMXEntityFixturePatch> FixturePatch) 
+		{ 
+			return !FixturePatch.IsValid();
+		});
+	SelectedFixturePatches.SetNum(NumValidPatches);
+	Algo::SortBy(SelectedFixturePatches, [](TWeakObjectPtr<UDMXEntityFixturePatch> FixturePatch)
+		{
+			return (int64)FixturePatch->GetUniverseID() * DMX_MAX_ADDRESS + FixturePatch->GetStartingChannel();
+		}
+	);
+
+	if (SelectedFixturePatches.IsEmpty())
 	{
-		int32 ChannelOffset = Channel - Patch->GetStartingChannel();
+		return FReply::Unhandled();
+	}
 
-		TSharedRef<FDMXEntityFixturePatchDragDropOperation> DragDropOp = MakeShared<FDMXEntityFixturePatchDragDropOperation>(Patch->GetParentLibrary(), TArray<TWeakObjectPtr<UDMXEntity>>{ Patch }, ChannelOffset);
+	TMap<TWeakObjectPtr<UDMXEntityFixturePatch>, int64> FixturePatchToAbsoluteChannelOffsetMap;
+	for (TWeakObjectPtr<UDMXEntityFixturePatch> FixturePatch : SelectedFixturePatches)
+	{
+		if (!FixturePatch.IsValid())
+		{
+			continue;
+		}
+
+		FixturePatch->Modify();
 
-		return FReply::Handled().BeginDragDrop(DragDropOp);
+		const int64 AbsoluteOffset = UniverseID * DMX_MAX_ADDRESS + Channel - (FixturePatch->GetUniverseID() * DMX_MAX_ADDRESS + FixturePatch->GetStartingChannel());
+		FixturePatchToAbsoluteChannelOffsetMap.Add(FixturePatch, AbsoluteOffset);
 	}
 
-	return FReply::Unhandled();
+	const TSharedRef<FDMXEntityFixturePatchDragDropOperation> DragDropOp = MakeShared<FDMXEntityFixturePatchDragDropOperation>(FixturePatchToAbsoluteChannelOffsetMap);
+	return FReply::Handled().BeginDragDrop(DragDropOp);
 }
 
 void SDMXPatchedUniverse::HandleDragEnterChannel(uint32 Channel, const FDragDropEvent& DragDropEvent)
diff --git a/Engine/Plugins/VirtualProduction/DMX/DMXEngine/Source/DMXEditor/Private/Widgets/FixturePatch/SDMXPatchedUniverse.h b/Engine/Plugins/VirtualProduction/DMX/DMXEngine/Source/DMXEditor/Private/Widgets/FixturePatch/SDMXPatchedUniverse.h
--- a/Engine/Plugins/VirtualProduction/DMX/DMXEngine/Source/DMXEditor/Private/Widgets/FixturePatch/SDMXPatchedUniverse.h
+++ b/Engine/Plugins/VirtualProduction/DMX/DMXEngine/Source/DMXEditor/Private/Widgets/FixturePatch/SDMXPatchedUniverse.h
@@ -80,16 +80,11 @@ public:
 	/** Sets if this universe should monitor inputs */
 	void SetMonitorInputsEnabled(bool bEnabled);
 
-	/** 
-	 * Patches the node.
-	 * Patches that have bAutoAssignAddress use their auto assigned address.
-	 * Others are assigned to the specified new starting channel
-	 * Returns false if the patch cannot be patched.
-	 */
-	bool Patch(const TSharedPtr<FDMXFixturePatchNode>& Node, int32 NewStartingChannel, bool bCreateTransaction);
+	/** Finds or adds a Node to this Universe */
+	bool FindOrAdd(const TSharedPtr<FDMXFixturePatchNode>& Node);
 
-	/** If set to true, shows a universe name above the patcher universe */
-	void SetShowUniverseName(bool bShow);
+	/** Removes a Node from this Universe */
+	void Remove(const TSharedPtr<FDMXFixturePatchNode>& Node);
 
 	/** Returns wether the patch can be patched to its current channels */
 	bool CanAssignFixturePatch(TWeakObjectPtr<UDMXEntityFixturePatch> TestedPatch) const;
@@ -118,9 +113,6 @@ protected:
 	//~ End SWidget Interface
 
 private:
-	/** Removes the node. Should be called when the node is Patched in another instance */
-	void Unpatch(const TSharedPtr<FDMXFixturePatchNode>& Node);
-
 	/** Called when a Fixture Patch changed */
 	void OnFixturePatchChanged(const UDMXEntityFixturePatch* FixturePatch);
 
@@ -217,6 +209,9 @@ private:
 	/** Delegate executed when a Drag Drop event was dropped onto a Channel */
 	FOnDropOntoChannel OnDropOntoChannel;
 
+	/** Anchor of Shift-Select while shift select is ongoing */
+	TWeakObjectPtr<UDMXEntityFixturePatch> ShiftSelectAnchorPatch;
+
 	/** The Fixture Patch Widgets that are currently being displalyed */
 	TArray<TSharedPtr<SDMXFixturePatchFragment>> FixturePatchWidgets;
 
diff --git a/Engine/Plugins/VirtualProduction/DMX/DMXEngine/Source/DMXRuntime/Private/Library/DMXEntityFixturePatch.cpp b/Engine/Plugins/VirtualProduction/DMX/DMXEngine/Source/DMXRuntime/Private/Library/DMXEntityFixturePatch.cpp
--- a/Engine/Plugins/VirtualProduction/DMX/DMXEngine/Source/DMXRuntime/Private/Library/DMXEntityFixturePatch.cpp
+++ b/Engine/Plugins/VirtualProduction/DMX/DMXEngine/Source/DMXRuntime/Private/Library/DMXEntityFixturePatch.cpp
@@ -18,6 +18,7 @@
 #include "Library/DMXLibrary.h"
 #include "Modulators/DMXModulator.h"
 
+#include "Algo/Find.h"
 #include "UObject/UObjectGlobals.h"
 
 DECLARE_LOG_CATEGORY_CLASS(DMXEntityFixturePatchLog, Log, All);
@@ -94,14 +95,16 @@ UDMXEntityFixturePatch* UDMXEntityFixturePatch::CreateFixturePatchInLibrary(FDMX
 			}
 			NewFixturePatch->MVRFixtureUUID = ConstructionParams.MVRFixtureUUID;
 
-#if WITH_EDITORONLY_DATA
+#if WITH_EDITOR
 			// Make a nice Editor Color
-			const FLinearColor EditorColor = [DMXLibrary, &ConstructionParams]()
+			const FLinearColor EditorColor = [DMXLibrary, NewFixturePatch]()
 			{
 				const TArray<UDMXEntityFixturePatch*> FixturePatches = DMXLibrary->GetEntitiesTypeCast<UDMXEntityFixturePatch>();
-				const UDMXEntityFixturePatch* const* FixturePatchOfSameTypePtr = FixturePatches.FindByPredicate([&ConstructionParams](const UDMXEntityFixturePatch* FixturePatch)
+				const UDMXEntityFixturePatch* const* FixturePatchOfSameTypePtr = Algo::FindByPredicate(FixturePatches, [NewFixturePatch](const UDMXEntityFixturePatch* FixturePatch)
 					{
-						return FixturePatch->GetFixtureType() == ConstructionParams.FixtureTypeRef.GetFixtureType();
+						return
+							FixturePatch != NewFixturePatch &&
+							FixturePatch->GetFixtureType() == NewFixturePatch->GetFixtureType();
 					});
 				return FixturePatchOfSameTypePtr ? (*FixturePatchOfSameTypePtr)->EditorColor : FLinearColor::MakeRandomColor();
 			}();
