DMX - Fix CIS Static Analysis Win64 - Variable is initialized through itself

diff --git a/Engine/Plugins/VirtualProduction/DMX/DMXEngine/Source/DMXControlConsoleEditor/Private/Views/SDMXControlConsoleEditorFaderGroupView.cpp b/Engine/Plugins/VirtualProduction/DMX/DMXEngine/Source/DMXControlConsoleEditor/Private/Views/SDMXControlConsoleEditorFaderGroupView.cpp
--- a/Engine/Plugins/VirtualProduction/DMX/DMXEngine/Source/DMXControlConsoleEditor/Private/Views/SDMXControlConsoleEditorFaderGroupView.cpp
+++ b/Engine/Plugins/VirtualProduction/DMX/DMXEngine/Source/DMXControlConsoleEditor/Private/Views/SDMXControlConsoleEditorFaderGroupView.cpp
@@ -269,7 +269,7 @@ void SDMXControlConsoleEditorFaderGroupView::OnElementRemoved()
 		const TSharedPtr<SDMXControlConsoleEditorFader> FaderWidget = StaticCastSharedPtr<SDMXControlConsoleEditorFader>(Widget.Pin());
 		if (FaderWidget.IsValid())
 		{
-			UDMXControlConsoleFaderBase* Fader = Fader = FaderWidget->GetFader();
+			UDMXControlConsoleFaderBase* Fader = FaderWidget->GetFader();
 			if (Fader && Elements.Contains(Fader))
 			{
 				continue;
