DMX - Fix Fixture Matrix attribute array collapses when editing elements

diff --git a/Engine/Plugins/VirtualProduction/DMX/DMXEngine/Source/DMXEditor/Private/Widgets/FixtureType/SDMXFixtureModeEditor.cpp b/Engine/Plugins/VirtualProduction/DMX/DMXEngine/Source/DMXEditor/Private/Widgets/FixtureType/SDMXFixtureModeEditor.cpp
--- a/Engine/Plugins/VirtualProduction/DMX/DMXEngine/Source/DMXEditor/Private/Widgets/FixtureType/SDMXFixtureModeEditor.cpp
+++ b/Engine/Plugins/VirtualProduction/DMX/DMXEngine/Source/DMXEditor/Private/Widgets/FixtureType/SDMXFixtureModeEditor.cpp
@@ -24,9 +24,9 @@ void SDMXFixtureModeEditor::Construct(const FArguments& InArgs, const TSharedRef
 	FixtureTypeSharedData = InDMXEditor->GetFixtureTypeSharedData();
 
 	UDMXEntityFixtureType::GetOnFixtureTypeChanged().AddSP(this, &SDMXFixtureModeEditor::OnFixtureTypePropertiesChanged);
-	FixtureTypeSharedData->OnFixtureTypesSelected.AddSP(this, &SDMXFixtureModeEditor::Refresh);
-	FixtureTypeSharedData->OnFixtureTypesSelected.AddSP(this, &SDMXFixtureModeEditor::Refresh);
-	FixtureTypeSharedData->OnModesSelected.AddSP(this, &SDMXFixtureModeEditor::Refresh);
+	FixtureTypeSharedData->OnFixtureTypesSelected.AddSP(this, &SDMXFixtureModeEditor::RequestRefresh);
+	FixtureTypeSharedData->OnFixtureTypesSelected.AddSP(this, &SDMXFixtureModeEditor::RequestRefresh);
+	FixtureTypeSharedData->OnModesSelected.AddSP(this, &SDMXFixtureModeEditor::RequestRefresh);
 
 	// Create a Struct Details View. This is not the most convenient type to work with as a property type customization for the FDMXFixtureMode struct cannot be used.
 	// Reason is soley significant performance gains, it's much faster than the easier approach with a UDMXEntityFixtureType customization as it was used up to 4.27.
@@ -71,7 +71,15 @@ void SDMXFixtureModeEditor::Construct(const FArguments& InArgs, const TSharedRef
 		]
 	];
 
-	Refresh();
+	RefreshInternal();
+}
+
+void SDMXFixtureModeEditor::RequestRefresh()
+{
+	if (!RefreshHandle.IsValid())
+	{
+		RefreshHandle = GEditor->GetTimerManager()->SetTimerForNextTick(FTimerDelegate::CreateSP(this, &SDMXFixtureModeEditor::RefreshInternal));
+	}
 }
 
 void SDMXFixtureModeEditor::NotifyPreChange(FProperty* PropertyAboutToChange)
@@ -144,7 +152,7 @@ void SDMXFixtureModeEditor::NotifyPostChange(const FPropertyChangedEvent& Proper
 		else if (PropertyName == GET_MEMBER_NAME_CHECKED(FDMXFixtureMode, bFixtureMatrixEnabled))
 		{
 			FixtureTypeSharedData->SetFunctionAndMatrixSelection(TArray<int32>(), ModeBeingEditedPtr->bFixtureMatrixEnabled);
-			Refresh();
+			RequestRefresh();
 		}
 		else if (PropertyName == GET_MEMBER_NAME_CHECKED(FDMXFixtureMode, bAutoChannelSpan))
 		{
@@ -152,7 +160,7 @@ void SDMXFixtureModeEditor::NotifyPostChange(const FPropertyChangedEvent& Proper
 			{
 				FixtureType->UpdateChannelSpan(ModeIndex);
 			}
-			Refresh();
+			RequestRefresh();
 		}
 		else if (PropertyName == GET_MEMBER_NAME_CHECKED(FDMXFixtureCellAttribute, DataType))
 		{
@@ -161,7 +169,7 @@ void SDMXFixtureModeEditor::NotifyPostChange(const FPropertyChangedEvent& Proper
 			FixtureType->SetFixtureMatrixEnabled(ModeIndex, true);
 
 			FixtureType->UpdateChannelSpan(ModeIndex);
-			Refresh();
+			RequestRefresh();
 		}
 		FPropertyChangedEvent ObjectPropertyChangedEvent(PropertyChangedEvent);
 		FixtureType->PostEditChangeProperty(ObjectPropertyChangedEvent);
@@ -170,8 +178,10 @@ void SDMXFixtureModeEditor::NotifyPostChange(const FPropertyChangedEvent& Proper
 	Transaction.Reset();
 }
 
-void SDMXFixtureModeEditor::Refresh()
+void SDMXFixtureModeEditor::RefreshInternal()
 {
+	RefreshHandle.Invalidate();
+
 	const TArray<TWeakObjectPtr<UDMXEntityFixtureType>>& SelectedFixtureTypes = FixtureTypeSharedData->GetSelectedFixtureTypes();
 	const TArray<int32>& SelectedModeIndices = FixtureTypeSharedData->GetSelectedModeIndices();
 
@@ -232,7 +242,7 @@ void SDMXFixtureModeEditor::OnFixtureTypePropertiesChanged(const UDMXEntityFixtu
 {
 	if (!Transaction.IsValid() && FixtureType == WeakFixtureType.Get())
 	{
-		Refresh();
+		RequestRefresh();
 	}
 }
 
diff --git a/Engine/Plugins/VirtualProduction/DMX/DMXEngine/Source/DMXEditor/Private/Widgets/FixtureType/SDMXFixtureModeEditor.h b/Engine/Plugins/VirtualProduction/DMX/DMXEngine/Source/DMXEditor/Private/Widgets/FixtureType/SDMXFixtureModeEditor.h
--- a/Engine/Plugins/VirtualProduction/DMX/DMXEngine/Source/DMXEditor/Private/Widgets/FixtureType/SDMXFixtureModeEditor.h
+++ b/Engine/Plugins/VirtualProduction/DMX/DMXEngine/Source/DMXEditor/Private/Widgets/FixtureType/SDMXFixtureModeEditor.h
@@ -33,6 +33,9 @@ public:
 	/** Constructs this widget */
 	void Construct(const FArguments& InArgs, const TSharedRef<FDMXEditor>& InDMXEditor);
 
+	/** Requests to refresh the widget on the next tick */
+	void RequestRefresh();
+
 protected:
 	//~ Begin FNotifyHook Interface
 	virtual void NotifyPreChange(FProperty* PropertyAboutToChange) override;
@@ -40,8 +43,8 @@ protected:
 	//~ End FNotifyHook Interface
 
 private:
-	/** Refereshes the Mode displayed */
-	void Refresh();
+	/** Refereshes the Mode displayed instantly */
+	void RefreshInternal();
 
 	/** Sets the Mode that is being edited */
 	void SetMode(UDMXEntityFixtureType* InFixtureType, int32 InModeIndex);
@@ -73,6 +76,9 @@ private:
 	/** Text block that displays info in case the struct details view cannot be shown */
 	TSharedPtr<STextBlock> InfoTextBlock;
 
+	/** Timer handle set while refresh is requested but not carried out yet */
+	FTimerHandle RefreshHandle;
+
 	/** Shared Data for */
 	TSharedPtr<FDMXFixtureTypeSharedData> FixtureTypeSharedData;
 
