DMX - Fix changing names in DMX UIs always appends a number

diff --git a/Engine/Plugins/VirtualProduction/DMX/DMXEngine/Source/DMXRuntime/Private/DMXRuntimeUtils.cpp b/Engine/Plugins/VirtualProduction/DMX/DMXEngine/Source/DMXRuntime/Private/DMXRuntimeUtils.cpp
--- a/Engine/Plugins/VirtualProduction/DMX/DMXEngine/Source/DMXRuntime/Private/DMXRuntimeUtils.cpp
+++ b/Engine/Plugins/VirtualProduction/DMX/DMXEngine/Source/DMXRuntime/Private/DMXRuntimeUtils.cpp
@@ -144,8 +144,18 @@ FString FDMXRuntimeUtils::ConvertTransformToGDTF4x3MatrixString(FTransform Trans
 	return Result;
 }
 
-FString FDMXRuntimeUtils::GenerateUniqueNameFromExisting(const TSet<FString>& InExistingNames, const FString& InDesiredName)
+FString FDMXRuntimeUtils::GenerateUniqueNameFromExisting(const TSet<FString>& InExistingNames, FString InDesiredName, uint8 MinNumDigitsInPostfix, bool bAlwaysGeneratePostfix)
 {
+	if (!bAlwaysGeneratePostfix && !InDesiredName.IsEmpty() && !InExistingNames.Contains(InDesiredName))
+	{
+		return InDesiredName;
+	}
+
+	if (!ensureAlwaysMsgf(!InDesiredName.IsEmpty(), TEXT("Empty name provided when generating unique name from exsiting. Using 'New'.")))
+	{
+		InDesiredName = TEXT("New");
+	}
+
 	auto GetBaseNameLambda([](const FString InName) -> FString
 		{
 			const FRegexPattern BaseNamePattern(TEXT("(\\w*)_\\d*"));
@@ -209,7 +219,7 @@ FString FDMXRuntimeUtils::GenerateUniqueNameFromExisting(const TSet<FString>& In
 		// Increment the trailing index string
 		NewIndex++;
 		FString NewIndexString = FString::FromInt(NewIndex);
-		while(NewIndexString.Len() < 3)
+		while(NewIndexString.Len() < MinNumDigitsInPostfix)
 		{
 			NewIndexString = FString(TEXT("0")) + NewIndexString;
 		}
@@ -231,7 +241,13 @@ FString FDMXRuntimeUtils::GenerateUniqueNameFromExisting(const TSet<FString>& In
 		}
 		else
 		{
-			return DesiredName.EndsWith(TEXT("_")) ? DesiredName + FString(TEXT("000")) : DesiredName + FString(TEXT("_000"));
+			FString NewPostfix;
+			while (NewPostfix.Len() < MinNumDigitsInPostfix)
+			{
+				NewPostfix = NewPostfix + FString(TEXT("0"));
+			}
+
+			return DesiredName.EndsWith(TEXT("_")) ? DesiredName + NewPostfix : DesiredName + TEXT("_") + NewPostfix;
 		}
 	}
 }
@@ -255,7 +271,9 @@ FString FDMXRuntimeUtils::FindUniqueEntityName(const UDMXLibrary* InLibrary, TSu
 		BaseName = *InEntityClass->GetName();
 	}
 
-	return FDMXRuntimeUtils::GenerateUniqueNameFromExisting(EntityNames, BaseName);
+	constexpr uint8 MinNumDigitsInPostfix = 3;
+	constexpr bool bAlwaysGeneratePostfix = true;
+	return FDMXRuntimeUtils::GenerateUniqueNameFromExisting(EntityNames, BaseName, MinNumDigitsInPostfix, bAlwaysGeneratePostfix);
 }
 
 bool FDMXRuntimeUtils::GetNameAndIndexFromString(const FString& InString, FString& OutName, int32& OutIndex)
diff --git a/Engine/Plugins/VirtualProduction/DMX/DMXEngine/Source/DMXRuntime/Public/DMXRuntimeUtils.h b/Engine/Plugins/VirtualProduction/DMX/DMXEngine/Source/DMXRuntime/Public/DMXRuntimeUtils.h
--- a/Engine/Plugins/VirtualProduction/DMX/DMXEngine/Source/DMXRuntime/Public/DMXRuntimeUtils.h
+++ b/Engine/Plugins/VirtualProduction/DMX/DMXEngine/Source/DMXRuntime/Public/DMXRuntimeUtils.h
@@ -23,10 +23,15 @@ public:
 	static FString ConvertTransformToGDTF4x3MatrixString(FTransform Transform);
 
 	/**
-	 * Generates a unique name given a base one and a list of existing ones, by appending an index to
-	 * existing names. If InBaseName is an empty String, it returns "Default name".
+	 * Generates a unique name given a base one and a list of existing ones, by appending an index.
+	 * 
+	 * @param InExistingNames					Set of exiting names the new name should not be using.
+	 * @param InDesiredName						The desired Name. Note, desired name should not be empty (ensured). Generates 'New' if desired name is empty.
+	 * @param NumDigitsInPostfix (optional)		The min num digits to generate when appending a post fix.
+	 * @param bAlwaysGeneratePostfix (optional)	If true, always generates digits as postfix.
+	 * @return									Returns the unique name.
 	 */
-	static FString GenerateUniqueNameFromExisting(const TSet<FString>& InExistingNames, const FString& InBaseName);
+	static FString GenerateUniqueNameFromExisting(const TSet<FString>& InExistingNames, FString InDesiredName, uint8 NumDigitsInPostfix = 2, bool bAlwaysGeneratePostfix = false);
 
 	/**
 	 * Creates an unique name for an Entity from a specific type, using the type name as base.
