DMX - Fix Data Type Customization doesn't support multi editing

diff --git a/Engine/Plugins/VirtualProduction/DMX/DMXEngine/Source/DMXEditor/Private/Customizations/DMXFixtureSignalFormatCustomization.cpp b/Engine/Plugins/VirtualProduction/DMX/DMXEngine/Source/DMXEditor/Private/Customizations/DMXFixtureSignalFormatCustomization.cpp
--- a/Engine/Plugins/VirtualProduction/DMX/DMXEngine/Source/DMXEditor/Private/Customizations/DMXFixtureSignalFormatCustomization.cpp
+++ b/Engine/Plugins/VirtualProduction/DMX/DMXEngine/Source/DMXEditor/Private/Customizations/DMXFixtureSignalFormatCustomization.cpp
@@ -42,7 +42,7 @@ void FDMXFixtureSignalFormatCustomization::CustomizeHeader(TSharedRef<IPropertyH
 			SAssignNew(SignalFormatSelector, SDMXSignalFormatSelector)
 			.HasMultipleValues_Lambda([this]()
 				{
-					return PropertyHandle->GetNumOuterObjects() > 1;
+					return GetSignalFormats().Num() > 1;
 				})
 			.InitialSelection(InitialSelection)
 			.OnSignalFormatSelected(this, &FDMXFixtureSignalFormatCustomization::SetSignalFormats)
@@ -60,7 +60,7 @@ TArray<EDMXFixtureSignalFormat> FDMXFixtureSignalFormatCustomization::GetSignalF
 		const EDMXFixtureSignalFormat* SignalFormatPtr = reinterpret_cast<EDMXFixtureSignalFormat*>(RawData);
 		if (SignalFormatPtr)
 		{
-			SignalFormats.Add(*SignalFormatPtr);
+			SignalFormats.AddUnique(*SignalFormatPtr);
 		}
 	}
 
diff --git a/Engine/Plugins/VirtualProduction/DMX/DMXEngine/Source/DMXEditor/Private/Widgets/SDMXSignalFormatSelector.cpp b/Engine/Plugins/VirtualProduction/DMX/DMXEngine/Source/DMXEditor/Private/Widgets/SDMXSignalFormatSelector.cpp
--- a/Engine/Plugins/VirtualProduction/DMX/DMXEngine/Source/DMXEditor/Private/Widgets/SDMXSignalFormatSelector.cpp
+++ b/Engine/Plugins/VirtualProduction/DMX/DMXEngine/Source/DMXEditor/Private/Widgets/SDMXSignalFormatSelector.cpp
@@ -8,12 +8,14 @@
 #include "Widgets/Input/SComboBox.h"
 
 
+#define LOCTEXT_NAMESPACE "SDMXSignalFormatSelector"
+
 const TArray<TSharedPtr<FString>> SDMXSignalFormatSelector::AvailableSignalFormats =
 {
 	MakeShared<FString>(TEXT("8bit")),
 	MakeShared<FString>(TEXT("16bit")),
 	MakeShared<FString>(TEXT("24bit")),
-	MakeShared<FString>(TEXT("32bit (not fully supported)"))
+	MakeShared<FString>(FText(LOCTEXT("32BitValueNotFullySupportedComboBoxEntry", "32bit (not fully supported)")).ToString())
 };
 
 void SDMXSignalFormatSelector::Construct(const FArguments& InArgs)
@@ -29,22 +31,22 @@ void SDMXSignalFormatSelector::Construct(const FArguments& InArgs)
 		SignalFormatsSource.RemoveAt(3);
 	}
 
-	const TSharedPtr<FString>& InitialSelection = [InArgs, this]()
+	const TSharedRef<FString>& InitialSelection = [InArgs, this]()
 	{
 		switch (InArgs._InitialSelection)
 		{
 			case EDMXFixtureSignalFormat::E8Bit:
-				return SignalFormatsSource[0];
+				return SignalFormatsSource[0].ToSharedRef();
 			case EDMXFixtureSignalFormat::E16Bit:
-				return SignalFormatsSource[1];
+				return SignalFormatsSource[1].ToSharedRef();
 			case EDMXFixtureSignalFormat::E24Bit:
-				return SignalFormatsSource[2];
+				return SignalFormatsSource[2].ToSharedRef();
 			case EDMXFixtureSignalFormat::E32Bit:
-				return SignalFormatsSource[3];
+				return SignalFormatsSource[3].ToSharedRef();
 			default:
-				checkNoEntry(); // Unhandled enum value
+				checkNoEntry(); // Unhahndled enum value
+				return SignalFormatsSource[0].ToSharedRef();
 		}
-		return SignalFormatsSource[0];
 	}();
 
 	ChildSlot
@@ -56,7 +58,28 @@ void SDMXSignalFormatSelector::Construct(const FArguments& InArgs)
 		.OnGenerateWidget(this, &SDMXSignalFormatSelector::GenerateComboBoxWidget)
 		[
 			SAssignNew(SelectionTextBlock, STextBlock)
-			.Text(FText::FromString(*InitialSelection))
+			.Text_Lambda([this, InArgs]()
+				{
+					if (HasMultipleValues.IsBound() && HasMultipleValues.Get())
+					{
+						return LOCTEXT("HasMultipleValues", "Multiple Values");
+					}
+
+					switch (GetSelectedSignalFormat())
+					{
+					case EDMXFixtureSignalFormat::E8Bit:
+						return FText::FromString(*SignalFormatsSource[0]);
+					case EDMXFixtureSignalFormat::E16Bit:
+						return FText::FromString(*SignalFormatsSource[1]);
+					case EDMXFixtureSignalFormat::E24Bit:
+						return FText::FromString(*SignalFormatsSource[2]);
+					case EDMXFixtureSignalFormat::E32Bit:
+						return FText::FromString(*SignalFormatsSource[3]);
+					default:
+						return LOCTEXT("InvalidSelection", "Invalid Selection");
+					}
+
+				})
 			.Font(FAppStyle::GetFontStyle(TEXT("PropertyWindow.NormalFont")))
 		]
 	];
@@ -94,23 +117,16 @@ TSharedRef<SWidget> SDMXSignalFormatSelector::GenerateComboBoxWidget(TSharedPtr<
 		SNew(STextBlock)
 		.Text_Lambda([this, Item]()
 			{
-				if (HasMultipleValues.Get())
-				{
-					return NSLOCTEXT("SDMXSignalFormatSelector", "MultipleValues", "Multiple Values");
-				}
-				else
-				{
-					return FText::FromString(*Item);
-				}
+				return FText::FromString(*Item);
 			})
 		.Font(FAppStyle::GetFontStyle(TEXT("PropertyWindow.NormalFont")));
 }
 
 void SDMXSignalFormatSelector::OnSignalFormatSelected(TSharedPtr<FString> SelectedItem, ESelectInfo::Type SelectInfo)
 {
-	SelectionTextBlock->SetText(FText::FromString(**SelectedItem));
-
 	EDMXFixtureSignalFormat SelectedSignalFormat = GetSelectedSignalFormat();
 
 	OnSignalFormatSelectedDelegate.ExecuteIfBound(SelectedSignalFormat);
 }
+
+#undef LOCTEXT_NAMESPACE
