DMX - Fix some GDTFs import the Matrix at the wrong channel

diff --git a/Engine/Plugins/VirtualProduction/DMX/DMXEngine/Source/DMXEditor/Private/DMXInitializeFixtureTypeFromGDTFHelper.cpp b/Engine/Plugins/VirtualProduction/DMX/DMXEngine/Source/DMXEditor/Private/DMXInitializeFixtureTypeFromGDTFHelper.cpp
--- a/Engine/Plugins/VirtualProduction/DMX/DMXEngine/Source/DMXEditor/Private/DMXInitializeFixtureTypeFromGDTFHelper.cpp
+++ b/Engine/Plugins/VirtualProduction/DMX/DMXEngine/Source/DMXEditor/Private/DMXInitializeFixtureTypeFromGDTFHelper.cpp
@@ -100,22 +100,17 @@ namespace UE::DMX::DMXInitializeFixtureTypeFromGDTFHelper::Private
 			GetDefaultValue(InDMXChannelNode, OutChannelProperties.DefaultValue);
 			GetDataType(InDMXChannelNode, OutChannelProperties.SignalFormat, OutChannelProperties.bLSBMode);
 
-			TArray<int32> Offsets;
-			if (IsMatrix(InDMXChannelNode, GeometryNodes))
+		
+			if (IsMatrix(InDMXChannelNode, GeometryNodes) && GetMatrixOffset(InDMXChannelNode, GeometryNodes, OutChannelProperties.FirstChannel))
 			{
-				Offsets = GetMatrixOffsetArray(InDMXChannelNode, GeometryNodes);
+				return true;
 			}
-			else
-			{
-				Offsets = GetOffsetArray(InDMXChannelNode);
-			}
-			if (Offsets.IsEmpty())
+			else if(GetMinOffset(InDMXChannelNode, OutChannelProperties.FirstChannel))
 			{
-				return false;
+				return true;
 			}
-			OutChannelProperties.FirstChannel = *Algo::MinElement(Offsets);
 
-			return true;
+			return false;
 		}
 
 	private:
@@ -194,13 +189,13 @@ namespace UE::DMX::DMXInitializeFixtureTypeFromGDTFHelper::Private
 			return;
 		}
 
-		/** Returns an array of offsets from an XmlNode */
-		static TArray<int32> GetOffsetArray(const FXmlNode* XmlNode)
+		/** Returns the Offset Attribute as an array of ints from an Xml Node */
+		static TArray<int32> GetOffsetArray(const FXmlNode* InXmlNode)
 		{
 			constexpr TCHAR OffsetTag[] = TEXT("DMXOffset");
 			TArray<int32> OffsetArray;
 			TArray<FString> OffsetStrArray;
-			FindAttributeEvenIfDMXSubstringIsMissing(*XmlNode, OffsetTag).ParseIntoArray(OffsetStrArray, TEXT(","));
+			FindAttributeEvenIfDMXSubstringIsMissing(*InXmlNode, OffsetTag).ParseIntoArray(OffsetStrArray, TEXT(","));
 
 			for (int32 OffsetIndex = 0; OffsetIndex < OffsetStrArray.Num(); ++OffsetIndex)
 			{
@@ -209,17 +204,49 @@ namespace UE::DMX::DMXInitializeFixtureTypeFromGDTFHelper::Private
 				OffsetArray.Add(OffsetValue);
 			}
 
+			if (OffsetArray.IsEmpty())
+			{
+				return OffsetArray;
+			}
+
 			return OffsetArray;
 		}
 
-		/** Returns an array of offsets from an XmlNode */
-		static TArray<int32> GetMatrixOffsetArray(const FXmlNode* InDMXChannelNode, const TArray<const FXmlNode*>& GeometryNodes)
+
+		/** Gets the min offset from an XmlNode's Offset Attriubute. Returns false if no min offset can be found. */
+		static bool GetMinOffset(const FXmlNode* InXmlNode, int32& OutOffset)
+		{
+			const TArray<int32> OffsetArray = GetOffsetArray(InXmlNode);
+			if (OffsetArray.IsEmpty())
+			{
+				return false;
+			}
+
+			OutOffset = *Algo::MinElement(OffsetArray);
+			return true;
+		}
+
+		/** Gets the max offset from an XmlNode's Offset Attriubute. Returns false if no max offset can be found. */
+		static bool GetMaxOffset(const FXmlNode* InXmlNode, int32& OutOffset)
+		{
+			const TArray<int32> OffsetArray = GetOffsetArray(InXmlNode);
+			if (OffsetArray.IsEmpty())
+			{
+				return false;
+			}
+
+			OutOffset = *Algo::MaxElement(OffsetArray);
+			return true;
+		}
+
+		/** Gets the matrix offset attribute value from a Channel Node. Returns false if no offset can be found. */
+		static bool GetMatrixOffset(const FXmlNode* InDMXChannelNode, const TArray<const FXmlNode*>& InGeometryNodes, int32& OutOffset)
 		{
 			constexpr TCHAR GeometryTag[] = TEXT("Geometry");
 			const FString GeometryName = InDMXChannelNode->GetAttribute(GeometryTag);
 
 			TArray<const FXmlNode*> ReferencedGeometryNodes;
-			Algo::CopyIf(GeometryNodes, ReferencedGeometryNodes, [GeometryTag, GeometryName](const FXmlNode* GeometryNode)
+			Algo::CopyIf(InGeometryNodes, ReferencedGeometryNodes, [GeometryTag, GeometryName](const FXmlNode* GeometryNode)
 				{
 					if (GeometryNode->GetAttribute(GeometryTag) == GeometryName)
 					{
@@ -228,39 +255,64 @@ namespace UE::DMX::DMXInitializeFixtureTypeFromGDTFHelper::Private
 
 					return false;
 				});
-
+			
 			if (ReferencedGeometryNodes.IsEmpty())
 			{
 				// Try to parse as normal channel instead
-				return GetOffsetArray(InDMXChannelNode);
+				return false;
+			}
+
+			int ChannelOffset;
+			if (!GetMinOffset(InDMXChannelNode, ChannelOffset))
+			{
+				return false;
 			}
 
-			if (const FXmlNode* XmlNode = ReferencedGeometryNodes[0])
+			TArray<int32> OffsetArray;
+			for(const FXmlNode* GeometryNode : InGeometryNodes)
 			{
-				constexpr TCHAR BreakTag[] = TEXT("Break");
-				if (const FXmlNode* BreakNode = XmlNode->FindChildNode(BreakTag))
+				if (GeometryNode)
 				{
-					return GetOffsetArray(BreakNode);
+					constexpr TCHAR BreakTag[] = TEXT("Break");
+					if (const FXmlNode* BreakNode = GeometryNode->FindChildNode(BreakTag))
+					{
+						int32 Offset;
+						if (!GetMinOffset(BreakNode, Offset))
+						{
+							return false;
+						}
+
+						OffsetArray.Add(Offset);
+					}
 				}
 			}
 
-			return TArray<int32>();
+			if (OffsetArray.IsEmpty())
+			{
+				return false;
+			}
+			
+			OutOffset = *Algo::MinElement(OffsetArray) + ChannelOffset - 1;
+			return true;
 		}
 
 		/** Returns the data type offset array implies */
 		static void GetDataType(const FXmlNode* InDMXChannelNode, EDMXFixtureSignalFormat& OutSignalFormat, bool& OutLSBOrder)
 		{
+			OutLSBOrder = false;
+
 			// For the data type we always need to refer to the offset of the channel, not the geometry
 			TArray<int32> ChannelOffestArray = GetOffsetArray(InDMXChannelNode);
 
 			// Compute number of used addresses
-			int32 AddressMin = DMX_MAX_ADDRESS;
-			int32 AddressMax = 0;
-			for (const int32& Address : ChannelOffestArray)
+			int32 AddressMin;
+			int32 AddressMax;
+			if (!GetMinOffset(InDMXChannelNode, AddressMin) ||
+				!GetMaxOffset(InDMXChannelNode, AddressMax))
 			{
-				AddressMin = FMath::Min(AddressMin, Address);
-				AddressMax = FMath::Max(AddressMax, Address);
+				return;
 			}
+
 			const int32 NumUsedAddresses = FMath::Clamp(AddressMax - AddressMin + 1, 1, DMX_MAX_FUNCTION_SIZE);
 
 			OutSignalFormat = static_cast<EDMXFixtureSignalFormat>(NumUsedAddresses - 1);
@@ -270,10 +322,6 @@ namespace UE::DMX::DMXInitializeFixtureTypeFromGDTFHelper::Private
 			{
 				OutLSBOrder = ChannelOffestArray[0] > ChannelOffestArray[1];
 			}
-			else
-			{
-				OutLSBOrder = false;
-			}
 		}
 	};
 };
