DMX - Fix CIS NonUnity Compile UnrealEditor Win64

diff --git a/Engine/Plugins/VirtualProduction/DMX/DMXEngine/Source/DMXEditor/Private/Widgets/FixtureType/SDMXFixtureModeEditor.h b/Engine/Plugins/VirtualProduction/DMX/DMXEngine/Source/DMXEditor/Private/Widgets/FixtureType/SDMXFixtureModeEditor.h
--- a/Engine/Plugins/VirtualProduction/DMX/DMXEngine/Source/DMXEditor/Private/Widgets/FixtureType/SDMXFixtureModeEditor.h
+++ b/Engine/Plugins/VirtualProduction/DMX/DMXEngine/Source/DMXEditor/Private/Widgets/FixtureType/SDMXFixtureModeEditor.h
@@ -5,6 +5,7 @@
 #include "Widgets/SCompoundWidget.h"
 
 #include "CoreMinimal.h"
+#include "Engine/EngineTypes.h"
 #include "Misc/NotifyHook.h"
 
 class FDMXEditor;
