DMX - Fix FixtureTypes cannot be renamed

diff --git a/Engine/Plugins/VirtualProduction/DMX/DMXEngine/Source/DMXEditor/Private/DMXEditorUtils.cpp b/Engine/Plugins/VirtualProduction/DMX/DMXEngine/Source/DMXEditor/Private/DMXEditorUtils.cpp
--- a/Engine/Plugins/VirtualProduction/DMX/DMXEngine/Source/DMXEditor/Private/DMXEditorUtils.cpp
+++ b/Engine/Plugins/VirtualProduction/DMX/DMXEngine/Source/DMXEditor/Private/DMXEditorUtils.cpp
@@ -27,8 +27,7 @@
 
 
 #define LOCTEXT_NAMESPACE "FDMXEditorUtils"
-// In blueprints name verification, it is said that '.' is known for causing problems
-#define DMX_INVALID_NAME_CHARACTERS TEXT(".")
+
 
 // Text object factory for pasting DMX Entities
 struct FDMXEntityObjectTextFactory : public FCustomizableTextObjectFactory
@@ -251,16 +250,6 @@ bool FDMXEditorUtils::ValidateEntityName(const FString& NewEntityName, const UDM
 		return false;
 	}
 
-	for (const TCHAR& Character : DMX_INVALID_NAME_CHARACTERS)
-	{
-		if (NewEntityName.Contains(&Character))
-		{
-			OutReason = FText::Format(LOCTEXT("NameWithInvalidCharacters", "Name can not contain: {0}"),
-				FText::FromString(&Character));
-			return false;
-		}
-	}
-
 	// Check against existing names for the current entity type
 	bool bNameIsUsed = false;
 	InLibrary->ForEachEntityOfTypeWithBreak(InEntityClass, [&bNameIsUsed, &NewEntityName](UDMXEntity* Entity)
