DMX - Fix Control Console multi-selection of channel when filter is applied will select channels in-between

diff --git a/Engine/Plugins/VirtualProduction/DMX/DMXEngine/Source/DMXControlConsoleEditor/Private/DMXControlConsoleEditorSelection.cpp b/Engine/Plugins/VirtualProduction/DMX/DMXEngine/Source/DMXControlConsoleEditor/Private/DMXControlConsoleEditorSelection.cpp
--- a/Engine/Plugins/VirtualProduction/DMX/DMXEngine/Source/DMXControlConsoleEditor/Private/DMXControlConsoleEditorSelection.cpp
+++ b/Engine/Plugins/VirtualProduction/DMX/DMXEngine/Source/DMXControlConsoleEditor/Private/DMXControlConsoleEditorSelection.cpp
@@ -20,261 +20,139 @@ FDMXControlConsoleEditorSelection::FDMXControlConsoleEditorSelection(const TShar
 
 void FDMXControlConsoleEditorSelection::AddToSelection(UDMXControlConsoleFaderGroup* FaderGroup)
 {
-	if (!FaderGroup || SelectedFaderGroups.Contains(FaderGroup))
+	if (FaderGroup)
 	{
-		return;
-	}
+		SelectedFaderGroups.AddUnique(FaderGroup);
 
-	SelectedFaderGroups.Add(FaderGroup);
-	OnSelectionChanged.Broadcast();
-	OnFaderGroupSelectionChanged.Broadcast(FaderGroup);
+		UpdateMultiSelectAnchor(UDMXControlConsoleFaderGroup::StaticClass());
+
+		OnSelectionChanged.Broadcast();
+	}
 }
 
 void FDMXControlConsoleEditorSelection::AddToSelection(UDMXControlConsoleFaderBase* Fader)
 {
-	if (!Fader || SelectedFaders.Contains(Fader))
+	if (Fader)
 	{
-		return;
-	}
+		SelectedFaders.AddUnique(Fader);
 
-	SelectedFaders.Add(Fader);
-	OnSelectionChanged.Broadcast();
-	OnFaderSelectionChanged.Broadcast(Fader);
+		UDMXControlConsoleFaderGroup& FaderGroup = Fader->GetOwnerFaderGroupChecked();
+		SelectedFaderGroups.AddUnique(&FaderGroup);
 
-	UDMXControlConsoleFaderGroup& FaderGroup = Fader->GetOwnerFaderGroupChecked();
-	AddToSelection(&FaderGroup);
+		UpdateMultiSelectAnchor(UDMXControlConsoleFaderBase::StaticClass());
+
+		OnSelectionChanged.Broadcast();
+	}
 }
 
 void FDMXControlConsoleEditorSelection::RemoveFromSelection(UDMXControlConsoleFaderGroup* FaderGroup)
 {
-	if (!FaderGroup || !SelectedFaderGroups.Contains(FaderGroup))
+	if (FaderGroup && SelectedFaderGroups.Contains(FaderGroup))
 	{
-		return;
-	}
+		ClearFadersSelection(FaderGroup);
+		SelectedFaderGroups.Remove(FaderGroup);
 
-	ClearFadersSelection(FaderGroup);
-	SelectedFaderGroups.Remove(FaderGroup);
-	OnSelectionChanged.Broadcast();
-	OnFaderGroupSelectionChanged.Broadcast(FaderGroup);
+		UpdateMultiSelectAnchor(UDMXControlConsoleFaderGroup::StaticClass());
 
-	if (!SelectedFaderGroups.IsEmpty())
-	{
-		return;
+		OnSelectionChanged.Broadcast();
 	}
-
-	OnClearFaderGroupSelection.Broadcast();
 }
 
 void FDMXControlConsoleEditorSelection::RemoveFromSelection(UDMXControlConsoleFaderBase* Fader)
 {
-	if (!Fader || !SelectedFaders.Contains(Fader))
+	if (Fader && SelectedFaders.Contains(Fader))
 	{
-		return;
-	}
+		SelectedFaders.Remove(Fader);
 
-	SelectedFaders.Remove(Fader);
-	OnSelectionChanged.Broadcast();
-	OnFaderSelectionChanged.Broadcast(Fader);
+		UpdateMultiSelectAnchor(UDMXControlConsoleFaderBase::StaticClass());
+
+		OnSelectionChanged.Broadcast();
+	}
 }
 
-void FDMXControlConsoleEditorSelection::Multiselect(UDMXControlConsoleFaderGroup* FaderGroup)
+void FDMXControlConsoleEditorSelection::Multiselect(UObject* FaderOrFaderGroupObject)
 {
-	if (!FaderGroup)
+	const UClass* MultiSelectClass = FaderOrFaderGroupObject->GetClass();
+	if (!ensureMsgf(MultiSelectClass == UDMXControlConsoleFaderGroup::StaticClass() || FaderOrFaderGroupObject->IsA(UDMXControlConsoleFaderBase::StaticClass()), TEXT("Invalid type when trying to multiselect")))
 	{
 		return;
 	}
 
-	UDMXControlConsole* ControlConsole = FDMXControlConsoleEditorManager::Get().GetDMXControlConsole();
-	if (!ControlConsole)
-	{
-		return;
-	}
+	SelectedFaderGroups.Remove(nullptr);
+	SelectedFaders.Remove(nullptr);
 
-	TArray<UDMXControlConsoleFaderGroup*> AllFaderGroups = ControlConsole->GetAllFaderGroups();
-	auto SortSelectedFaderGroupsLambda = [AllFaderGroups](TWeakObjectPtr<UObject> FaderGroupObjectA, TWeakObjectPtr<UObject> FaderGroupObjectB)
+	// Normal selection if nothing is selected or there's no valid anchor
+	if (!MultiSelectAnchor.IsValid() ||
+		(SelectedFaderGroups.IsEmpty() && SelectedFaders.IsEmpty()))
 	{
-		UDMXControlConsoleFaderGroup* FaderGroupA = Cast<UDMXControlConsoleFaderGroup>(FaderGroupObjectA);
-		UDMXControlConsoleFaderGroup* FaderGroupB = Cast<UDMXControlConsoleFaderGroup>(FaderGroupObjectB);
-
-		if (!FaderGroupA || !FaderGroupB)
+		if (UDMXControlConsoleFaderGroup* FaderGroup = Cast<UDMXControlConsoleFaderGroup>(FaderOrFaderGroupObject))
 		{
-			return false;
+			AddToSelection(FaderGroup);
 		}
-
-		const int32 IndexA = AllFaderGroups.IndexOfByKey(FaderGroupA);
-		const int32 IndexB = AllFaderGroups.IndexOfByKey(FaderGroupB);
-
-		bool bIsBefore = IndexA < IndexB;
-		return bIsBefore;
-	};
-
-	if (IsSelected(FaderGroup))
-	{
-		TArray<TWeakObjectPtr<UObject>> CurrentSelectedFaderGroups = GetSelectedFaderGroups();
-
-		Algo::Sort(CurrentSelectedFaderGroups, SortSelectedFaderGroupsLambda);
-
-		const int32 Index = CurrentSelectedFaderGroups.IndexOfByKey(FaderGroup);
-
-		// Remove from selection all fader groups after the given one
-		for (int32 FaderGroupIndex = Index + 1; FaderGroupIndex < CurrentSelectedFaderGroups.Num(); ++FaderGroupIndex)
+		else if (UDMXControlConsoleFaderBase* Fader = Cast<UDMXControlConsoleFaderBase>(FaderOrFaderGroupObject))
 		{
-			UDMXControlConsoleFaderGroup* FaderGroupToUnselect = Cast<UDMXControlConsoleFaderGroup>(CurrentSelectedFaderGroups[FaderGroupIndex]);
-			if (!FaderGroupToUnselect)
-			{
-				continue;
-			}
-
-			RemoveFromSelection(FaderGroupToUnselect);
+			AddToSelection(Fader);
 		}
+		return;
 	}
-	else
-	{
-		AddToSelection(FaderGroup);
 
-		TArray<TWeakObjectPtr<UObject>> CurrentSelectedFaderGroups = GetSelectedFaderGroups();
-		if (CurrentSelectedFaderGroups.Num() <= 1)
+	TArray<UObject*> FadersAndFaderGroups;
+	UDMXControlConsole* Console = FDMXControlConsoleEditorManager::Get().GetDMXControlConsole();
+	for (UDMXControlConsoleFaderGroup* AnyFaderGroup : Console->GetAllFaderGroups())
+	{
+		FadersAndFaderGroups.AddUnique(AnyFaderGroup);
+		for (UDMXControlConsoleFaderBase* AnyFader : AnyFaderGroup->GetAllFaders())
 		{
-			return;
+			FadersAndFaderGroups.AddUnique(AnyFader);
 		}
+	}
 
-		Algo::Sort(CurrentSelectedFaderGroups, SortSelectedFaderGroupsLambda);
-
-		ClearSelection();
-
-		// Check fader groups mutual position
-		int32 FirstIndex;
-		int32 LastIndex;
-		if (FaderGroup == CurrentSelectedFaderGroups[0])
+	const int32 IndexOfFaderGroupAnchor = FadersAndFaderGroups.IndexOfByPredicate([this](const UObject* Object)
 		{
-			FirstIndex = AllFaderGroups.IndexOfByKey(FaderGroup);
-			LastIndex = AllFaderGroups.IndexOfByKey(CurrentSelectedFaderGroups.Last());
-		}
-		else 
-		{
-			FirstIndex = AllFaderGroups.IndexOfByKey(CurrentSelectedFaderGroups[0]);
-			LastIndex = AllFaderGroups.IndexOfByKey(FaderGroup);
-		}
-
-		// Select all fader groups between first and last selected
-		for (int32 FaderGroupIndex = FirstIndex; FaderGroupIndex <= LastIndex; ++FaderGroupIndex)
+			return MultiSelectAnchor == Object;
+		});
+	const int32 IndexOfFaderAnchor = FadersAndFaderGroups.IndexOfByPredicate([this](const UObject* Object)
 		{
-			UDMXControlConsoleFaderGroup* FaderGroupToSelect = AllFaderGroups[FaderGroupIndex];
-			if (!FaderGroupToSelect)
-			{
-				continue;
-			}
-			
-			// Select only if it's not already selected
-			if (!IsSelected(FaderGroupToSelect))
-			{
-				AddToSelection(FaderGroupToSelect);
-			}
-
-			// Select all its faders anyway
-			TArray<UDMXControlConsoleFaderBase*> Faders = FaderGroupToSelect->GetAllFaders();
-			for (UDMXControlConsoleFaderBase* Fader : Faders)
-			{
-				if (!Fader || IsSelected(Fader))
-				{
-					continue;
-				}
-
-				AddToSelection(Fader);
-			}
-		}
-	}
-}
+			return MultiSelectAnchor == Object;
+		});
 
-void FDMXControlConsoleEditorSelection::Multiselect(UDMXControlConsoleFaderBase* Fader)
-{
-	if (!Fader)
+	const int32 IndexOfAnchor = FMath::Max(IndexOfFaderGroupAnchor, IndexOfFaderAnchor);
+	if (!ensureAlwaysMsgf(IndexOfAnchor != INDEX_NONE, TEXT("No previous selection when multi selecting, cannot multiselect.")))
 	{
 		return;
 	}
 
-	UDMXControlConsoleFaderGroup& FaderGroup = Fader->GetOwnerFaderGroupChecked();
-
-	auto SortSelectedFadersLambda = [&FaderGroup](TWeakObjectPtr<UObject> FaderObjectA, TWeakObjectPtr<UObject> FaderObjectB)
-		{
-			UDMXControlConsoleFaderBase* FaderA = Cast<UDMXControlConsoleFaderBase>(FaderObjectA);
-			UDMXControlConsoleFaderBase* FaderB = Cast<UDMXControlConsoleFaderBase>(FaderObjectB);
-
-			if (!FaderA || !FaderB)
-			{
-				return false;
-			}
-
-			const int32 IndexA = FaderGroup.GetAllFaders().IndexOfByKey(FaderA);
-			const int32 IndexB = FaderGroup.GetAllFaders().IndexOfByKey(FaderB);
+	const int32 IndexOfSelection = FadersAndFaderGroups.IndexOfByKey(FaderOrFaderGroupObject);
 
-			bool bIsBefore = IndexA < IndexB;
-			return bIsBefore;
-		};
+	const int32 StartIndex = FMath::Min(IndexOfAnchor, IndexOfSelection);
+	const int32 EndIndex = FMath::Max(IndexOfAnchor, IndexOfSelection);
 
-	if (IsSelected(Fader))
+	SelectedFaderGroups.Reset();
+	SelectedFaders.Reset();
+	for (int32 IndexToSelect = StartIndex; IndexToSelect <= EndIndex; IndexToSelect++)
 	{
-		TArray<UDMXControlConsoleFaderBase*> CurrentSelectedFaders = GetSelectedFadersFromFaderGroup(&FaderGroup);
-
-		Algo::Sort(CurrentSelectedFaders, SortSelectedFadersLambda);
-
-		const int32 Index = CurrentSelectedFaders.IndexOfByKey(Fader);
-
-		// Remove from selection all faders after the given one
-		for (int32 FaderIndex = Index + 1; FaderIndex < CurrentSelectedFaders.Num(); ++FaderIndex)
+		if (!ensureMsgf(FadersAndFaderGroups.IsValidIndex(IndexToSelect), TEXT("Invalid index when multiselecting")))
 		{
-			UDMXControlConsoleFaderBase* FaderToUnselect = CurrentSelectedFaders[FaderIndex];
-			if (!FaderToUnselect)
-			{
-				continue;
-			}
-
-			RemoveFromSelection(FaderToUnselect);
+			break;
 		}
-	}
-	else
-	{
-		AddToSelection(Fader);
-
-		TArray<UDMXControlConsoleFaderBase*> Faders = FaderGroup.GetAllFaders();
-		TArray<UDMXControlConsoleFaderBase*> CurrentSelectedFaders = GetSelectedFadersFromFaderGroup(&FaderGroup);
 
-		if (CurrentSelectedFaders.Num() <= 1)
+		if (UDMXControlConsoleFaderGroup* FaderGroupToSelect = Cast<UDMXControlConsoleFaderGroup>(FadersAndFaderGroups[IndexToSelect]))
 		{
-			return;
+			SelectedFaderGroups.AddUnique(FaderGroupToSelect);
 		}
-
-		Algo::Sort(CurrentSelectedFaders, SortSelectedFadersLambda);
-
-		ClearFadersSelection(&FaderGroup);
-
-		int32 FirstIndex;
-		int32 LastIndex;
-
-		// Check faders mutual position
-		if (Fader == CurrentSelectedFaders[0])
+		else if (UDMXControlConsoleFaderBase* FaderToSelect = Cast<UDMXControlConsoleFaderBase>(FadersAndFaderGroups[IndexToSelect]))
 		{
-			FirstIndex = Faders.IndexOfByKey(Fader);
-			LastIndex = Faders.IndexOfByKey(CurrentSelectedFaders.Last());
-		}
-		else
-		{
-			FirstIndex = Faders.IndexOfByKey(CurrentSelectedFaders[0]);
-			LastIndex = Faders.IndexOfByKey(Fader);
-		}
-
-		// Select all faders between first and last selected
-		for (int32 FaderIndex = FirstIndex; FaderIndex <= LastIndex; ++FaderIndex)
-		{
-			UDMXControlConsoleFaderBase* FaderToSelect = Faders[FaderIndex];
-			if (!FaderToSelect || IsSelected(FaderToSelect))
-			{
-				continue;
-			}
-
-			AddToSelection(FaderToSelect);
+			SelectedFaders.AddUnique(FaderToSelect);
 		}
 	}
+	if (!SelectedFaders.IsEmpty())
+	{
+		// Always select the fader group of the first selected fader
+		UDMXControlConsoleFaderBase* FirstSelectedFader = CastChecked<UDMXControlConsoleFaderBase>(SelectedFaders[0]);
+		SelectedFaderGroups.AddUnique(&FirstSelectedFader->GetOwnerFaderGroupChecked());
+	}
+
+	OnSelectionChanged.Broadcast();
 }
 
 void FDMXControlConsoleEditorSelection::ReplaceInSelection(UDMXControlConsoleFaderGroup* FaderGroup)
@@ -348,12 +226,6 @@ bool FDMXControlConsoleEditorSelection::IsSelected(UDMXControlConsoleFaderBase*
 	return SelectedFaders.Contains(Fader);
 }
 
-void FDMXControlConsoleEditorSelection::ClearFadersSelection()
-{
-	SelectedFaders.Empty();
-	OnClearFaderSelection.Broadcast();
-}
-
 void FDMXControlConsoleEditorSelection::ClearFadersSelection(UDMXControlConsoleFaderGroup* FaderGroup)
 {
 	if (!FaderGroup || !SelectedFaderGroups.Contains(FaderGroup))
@@ -381,14 +253,20 @@ void FDMXControlConsoleEditorSelection::ClearFadersSelection(UDMXControlConsoleF
 
 	SelectedFaders.RemoveAll(IsFaderGroupOwnerLambda);
 
-	OnClearFaderSelection.Broadcast();
+	if (!MultiSelectAnchor.IsValid() || MultiSelectAnchor->IsA(UDMXControlConsoleFaderBase::StaticClass()))
+	{
+		UpdateMultiSelectAnchor(UDMXControlConsoleFaderBase::StaticClass());
+	}
+
+	OnSelectionChanged.Broadcast();
 }
 
 void FDMXControlConsoleEditorSelection::ClearSelection()
 {
-	ClearFadersSelection();
-	SelectedFaderGroups.Empty();
-	OnClearFaderGroupSelection.Broadcast();
+	SelectedFaderGroups.Reset();
+	SelectedFaders.Reset();
+
+	OnSelectionChanged.Broadcast();
 }
 
 UDMXControlConsoleFaderGroup* FDMXControlConsoleEditorSelection::GetFirstSelectedFaderGroup() const
@@ -506,4 +384,29 @@ TArray<UDMXControlConsoleFaderBase*> FDMXControlConsoleEditorSelection::GetSelec
 	return CurrentSelectedFaders;
 }
 
+void FDMXControlConsoleEditorSelection::UpdateMultiSelectAnchor(UClass* PreferedClass)
+{
+	if (!ensureMsgf(PreferedClass == UDMXControlConsoleFaderGroup::StaticClass() || PreferedClass == UDMXControlConsoleFaderBase::StaticClass(), TEXT("Invalid class when trying to update multi select anchor")))
+	{
+		return;
+	}
+
+	if (PreferedClass == UDMXControlConsoleFaderGroup::StaticClass() && !SelectedFaderGroups.IsEmpty())
+	{
+		MultiSelectAnchor = SelectedFaderGroups.Last();
+	}
+	else if (!SelectedFaders.IsEmpty())
+	{
+		MultiSelectAnchor = SelectedFaders.Last();
+	}
+	else if (!SelectedFaderGroups.IsEmpty())
+	{
+		MultiSelectAnchor = SelectedFaderGroups.Last();
+	}
+	else
+	{
+		MultiSelectAnchor = nullptr;
+	}
+}
+
 #undef LOCTEXT_NAMESPACE
diff --git a/Engine/Plugins/VirtualProduction/DMX/DMXEngine/Source/DMXControlConsoleEditor/Private/DMXControlConsoleEditorSelection.h b/Engine/Plugins/VirtualProduction/DMX/DMXEngine/Source/DMXControlConsoleEditor/Private/DMXControlConsoleEditorSelection.h
--- a/Engine/Plugins/VirtualProduction/DMX/DMXEngine/Source/DMXControlConsoleEditor/Private/DMXControlConsoleEditorSelection.h
+++ b/Engine/Plugins/VirtualProduction/DMX/DMXEngine/Source/DMXControlConsoleEditor/Private/DMXControlConsoleEditorSelection.h
@@ -3,6 +3,7 @@
 #pragma once
 
 #include "CoreMinimal.h"
+#include "UObject/WeakObjectPtrTemplates.h"
 
 class FDMXControlConsoleEditorManager;
 class UDMXControlConsoleFaderBase;
@@ -13,8 +14,7 @@ class FDMXControlConsoleEditorSelection
 	: public TSharedFromThis<FDMXControlConsoleEditorSelection>
 {
 public:
-	DECLARE_EVENT_OneParam(FDMXControlConsoleEditorSelection, FDMXFaderSelectionEvent, UDMXControlConsoleFaderBase*)
-	DECLARE_EVENT_OneParam(FDMXControlConsoleEditorSelection, FDMXFaderGroupSelectionEvent, UDMXControlConsoleFaderGroup*)
+	DECLARE_EVENT(FDMXControlConsoleEditorSelection, FDMXControlConsoleSelectionEvent)
 
 	/** Constructor */
 	FDMXControlConsoleEditorSelection(const TSharedRef<FDMXControlConsoleEditorManager>& InControlConsoleManager);
@@ -31,11 +31,8 @@ public:
 	/** Removes the given Fader from selection */
 	void RemoveFromSelection(UDMXControlConsoleFaderBase* Fader);
 
-	/** Handles Fader Groups multiselection */
-	void Multiselect(UDMXControlConsoleFaderGroup* FaderGroup);
-
-	/** Handles Faders multiselection */
-	void Multiselect(UDMXControlConsoleFaderBase* Fader);
+	/** Multiselects the Fader or Fader Group and the current selection */
+	void Multiselect(UObject* FaderOrFaderGroupObject);
 
 	/** Replaces the given selected Fader Group with the next available one */
 	void ReplaceInSelection(UDMXControlConsoleFaderGroup* FaderGroup);
@@ -49,12 +46,6 @@ public:
 	/** Gets wheter the given Fader is selected or not */
 	bool IsSelected(UDMXControlConsoleFaderBase* Fader) const;
 
-	/** Sets multiselection state */
-	void SetAllowMultiselect(bool bAllow);
-
-	/** Clears from selection alla Faders */
-	void ClearFadersSelection();
-
 	/** Clears from selection alla Faders owned by the given FaderGroup */
 	void ClearFadersSelection(UDMXControlConsoleFaderGroup* FaderGroup);
 
@@ -76,22 +67,13 @@ public:
 	/** Gets all selected Faders from the fiven Fader Group */
 	TArray<UDMXControlConsoleFaderBase*> GetSelectedFadersFromFaderGroup(UDMXControlConsoleFaderGroup* FaderGroup) const;
 
-	/** Gets a reference to OnSelectionChanged delegate */
-	FSimpleMulticastDelegate& GetOnSelectionChanged() { return OnSelectionChanged; }
-
-	/** Gets a reference to OnFaderGroupSelectionChanged delegate */
-	FDMXFaderGroupSelectionEvent& GetOnFaderGroupSelectionChanged() { return OnFaderGroupSelectionChanged; }
-
-	/** Gets a reference to OnClearFaderGroupSelection delegate */
-	FSimpleMulticastDelegate& GetOnClearFaderGroupSelection() { return OnClearFaderGroupSelection; }
-
-	/** Gets a reference to OnFaderSelectionChanged delegate */
-	FDMXFaderSelectionEvent& GetOnFaderSelectionChanged() { return OnFaderSelectionChanged; }
-
-	/** Gets a reference to OnClearFaderSelection delegate */
-	FSimpleMulticastDelegate& GetOnClearFaderSelection() { return OnClearFaderSelection; }
+	/** Returns an event raised when the Selection changed */
+	FDMXControlConsoleSelectionEvent& GetOnSelectionChanged() { return OnSelectionChanged; }
 
 private:
+	/** Updates the multi select anchor */
+	void UpdateMultiSelectAnchor(UClass* PreferedClass);
+
 	/** Weak reference to DMX DMX Control Console */
 	TWeakPtr<FDMXControlConsoleEditorManager> WeakControlConsoleManager;
 
@@ -101,18 +83,9 @@ private:
 	/** Array of current selected Faders */
 	TArray<TWeakObjectPtr<UObject>> SelectedFaders;
 
-	/** Called whenever current selection changes */
-	FSimpleMulticastDelegate OnSelectionChanged;
-
-	/** Called when Fader Group selection changes */
-	FDMXFaderGroupSelectionEvent OnFaderGroupSelectionChanged;
+	/** Anchor while multi selecting */
+	TWeakObjectPtr<UObject> MultiSelectAnchor;
 
-	/** Called on Fader Group selection clearing */
-	FSimpleMulticastDelegate OnClearFaderGroupSelection;
-
-	/** Called when Fader selection changes */
-	FDMXFaderSelectionEvent OnFaderSelectionChanged;
-
-	/** Called on Fader selection clearing */
-	FSimpleMulticastDelegate OnClearFaderSelection;
+	/** Called whenever current selection changes */
+	FDMXControlConsoleSelectionEvent OnSelectionChanged;
 };
diff --git a/Engine/Plugins/VirtualProduction/DMX/DMXEngine/Source/DMXControlConsoleEditor/Private/Views/SDMXControlConsoleEditorFaderGroupView.cpp b/Engine/Plugins/VirtualProduction/DMX/DMXEngine/Source/DMXControlConsoleEditor/Private/Views/SDMXControlConsoleEditorFaderGroupView.cpp
--- a/Engine/Plugins/VirtualProduction/DMX/DMXEngine/Source/DMXControlConsoleEditor/Private/Views/SDMXControlConsoleEditorFaderGroupView.cpp
+++ b/Engine/Plugins/VirtualProduction/DMX/DMXEngine/Source/DMXControlConsoleEditor/Private/Views/SDMXControlConsoleEditorFaderGroupView.cpp
@@ -135,6 +135,16 @@ void SDMXControlConsoleEditorFaderGroupView::ApplyGlobalFilter(const FString& In
 
 	const EVisibility NewVisibility = bHasVisibleChildren ? EVisibility::Visible : EVisibility::Collapsed;
 	SetVisibility(NewVisibility);
+
+	// If not visible, remove from selection
+	if (NewVisibility == EVisibility::Collapsed)
+	{
+		const TSharedRef<FDMXControlConsoleEditorSelection> SelectionHandler = FDMXControlConsoleEditorManager::Get().GetSelectionHandler();
+		if (SelectionHandler->IsSelected(FaderGroup.Get()))
+		{
+			SelectionHandler->RemoveFromSelection(FaderGroup.Get());
+		}
+	}
 }
 
 void SDMXControlConsoleEditorFaderGroupView::ShowAllElements()
diff --git a/Engine/Plugins/VirtualProduction/DMX/DMXEngine/Source/DMXControlConsoleEditor/Private/Views/SDMXControlConsoleEditorView.cpp b/Engine/Plugins/VirtualProduction/DMX/DMXEngine/Source/DMXControlConsoleEditor/Private/Views/SDMXControlConsoleEditorView.cpp
--- a/Engine/Plugins/VirtualProduction/DMX/DMXEngine/Source/DMXControlConsoleEditor/Private/Views/SDMXControlConsoleEditorView.cpp
+++ b/Engine/Plugins/VirtualProduction/DMX/DMXEngine/Source/DMXControlConsoleEditor/Private/Views/SDMXControlConsoleEditorView.cpp
@@ -21,6 +21,7 @@
 #include "LevelEditor.h"
 #include "PropertyEditorModule.h"
 #include "ScopedTransaction.h"
+#include "TimerManager.h"
 #include "Application/ThrottleManager.h"
 #include "Framework/MultiBox/MultiBoxBuilder.h"
 #include "Modules/ModuleManager.h"
@@ -48,14 +49,12 @@ SDMXControlConsoleEditorView::~SDMXControlConsoleEditorView()
 void SDMXControlConsoleEditorView::Construct(const FArguments& InArgs)
 {
 	FDMXControlConsoleEditorManager& ControlConsoleManager = FDMXControlConsoleEditorManager::Get();
-	ControlConsoleManager.GetOnControlConsoleLoaded().AddSP(this, &SDMXControlConsoleEditorView::UpdateDetailsViews);
+	ControlConsoleManager.GetOnControlConsoleLoaded().AddSP(this, &SDMXControlConsoleEditorView::RequestUpdateDetailsViews);
 	ControlConsoleManager.GetOnControlConsoleLoaded().AddSP(this, &SDMXControlConsoleEditorView::OnFaderGroupRowRemoved);
 	ControlConsoleManager.GetOnControlConsoleLoaded().AddSP(this, &SDMXControlConsoleEditorView::OnFaderGroupRowAdded);
 
 	const TSharedRef<FDMXControlConsoleEditorSelection> SelectionHandler = ControlConsoleManager.GetSelectionHandler();
-	SelectionHandler->GetOnSelectionChanged().AddSP(this, &SDMXControlConsoleEditorView::UpdateDetailsViews);
-	SelectionHandler->GetOnClearFaderGroupSelection().AddSP(this, &SDMXControlConsoleEditorView::UpdateDetailsViews);
-	SelectionHandler->GetOnClearFaderSelection().AddSP(this, &SDMXControlConsoleEditorView::UpdateDetailsViews);
+	SelectionHandler->GetOnSelectionChanged().AddSP(this, &SDMXControlConsoleEditorView::RequestUpdateDetailsViews);
 
 	OnActiveTabChangedDelegateHandle = FGlobalTabmanager::Get()->OnActiveTabChanged_Subscribe(FOnActiveTabChanged::FDelegate::CreateSP(this, &SDMXControlConsoleEditorView::OnActiveTabChanged));
 
@@ -228,7 +227,7 @@ void SDMXControlConsoleEditorView::Construct(const FArguments& InArgs)
 			]
 		];
 	
-	UpdateDetailsViews();
+	ForceUpdateDetailsViews();
 }
 
 UDMXControlConsole* SDMXControlConsoleEditorView::GetControlConsole() const
@@ -314,8 +313,18 @@ TSharedRef<SWidget> SDMXControlConsoleEditorView::GenerateToolbar()
 	return ToolbarBuilder.MakeWidget();
 }
 
-void SDMXControlConsoleEditorView::UpdateDetailsViews()
+void SDMXControlConsoleEditorView::RequestUpdateDetailsViews()
 {
+	if(!UpdateDetailsViewTimerHandle.IsValid())
+	{
+		UpdateDetailsViewTimerHandle = GEditor->GetTimerManager()->SetTimerForNextTick(FTimerDelegate::CreateSP(this, &SDMXControlConsoleEditorView::ForceUpdateDetailsViews));
+	}
+}
+
+void SDMXControlConsoleEditorView::ForceUpdateDetailsViews()
+{
+	UpdateDetailsViewTimerHandle.Invalidate();
+
 	UDMXControlConsole* ControlConsole = GetControlConsole();
 	if (!ensureMsgf(ControlConsole, TEXT("Invalid DMX Control Console, can't update details view correctly.")))
 	{
diff --git a/Engine/Plugins/VirtualProduction/DMX/DMXEngine/Source/DMXControlConsoleEditor/Private/Views/SDMXControlConsoleEditorView.h b/Engine/Plugins/VirtualProduction/DMX/DMXEngine/Source/DMXControlConsoleEditor/Private/Views/SDMXControlConsoleEditorView.h
--- a/Engine/Plugins/VirtualProduction/DMX/DMXEngine/Source/DMXControlConsoleEditor/Private/Views/SDMXControlConsoleEditorView.h
+++ b/Engine/Plugins/VirtualProduction/DMX/DMXEngine/Source/DMXControlConsoleEditor/Private/Views/SDMXControlConsoleEditorView.h
@@ -4,6 +4,7 @@
 
 #include "CoreMinimal.h"
 
+#include "Engine/EngineTypes.h"
 #include "Widgets/SCompoundWidget.h"
 
 class SDMXControlConsoleEditorFaderGroupRowView;
@@ -45,8 +46,11 @@ private:
 	/** Generates the toolbar for this DMX Control Console View */
 	TSharedRef<SWidget> GenerateToolbar();
 
-	/** Updates Details View's object arrays */
-	void UpdateDetailsViews();
+	/** Requests to update the Details Views on the next tick */
+	void RequestUpdateDetailsViews();
+
+	/** Updates the Details Views */
+	void ForceUpdateDetailsViews();
 
 	/** Updates FixturePatchVerticalBox widget */
 	void UpdateFixturePatchRows();
@@ -101,4 +105,7 @@ private:
 
 	/** Delegate handle bound to the FGlobalTabmanager::OnActiveTabChanged delegate */
 	FDelegateHandle OnActiveTabChangedDelegateHandle;
+
+	/** Timer handle in use while updating details views is requested but not carried out yet */
+	FTimerHandle UpdateDetailsViewTimerHandle;
 };
diff --git a/Engine/Plugins/VirtualProduction/DMX/DMXEngine/Source/DMXControlConsoleEditor/Private/Widgets/SDMXControlConsoleEditorFader.cpp b/Engine/Plugins/VirtualProduction/DMX/DMXEngine/Source/DMXControlConsoleEditor/Private/Widgets/SDMXControlConsoleEditorFader.cpp
--- a/Engine/Plugins/VirtualProduction/DMX/DMXEngine/Source/DMXControlConsoleEditor/Private/Widgets/SDMXControlConsoleEditorFader.cpp
+++ b/Engine/Plugins/VirtualProduction/DMX/DMXEngine/Source/DMXControlConsoleEditor/Private/Widgets/SDMXControlConsoleEditorFader.cpp
@@ -160,7 +160,7 @@ void SDMXControlConsoleEditorFader::Construct(const FArguments& InArgs, const TO
 	];
 
 	const TSharedRef<FDMXControlConsoleEditorSelection> SelectionHandler = FDMXControlConsoleEditorManager::Get().GetSelectionHandler();
-	SelectionHandler->GetOnFaderSelectionChanged().AddSP(this, &SDMXControlConsoleEditorFader::OnSelectionChanged);
+	SelectionHandler->GetOnSelectionChanged().AddSP(this, &SDMXControlConsoleEditorFader::OnSelectionChanged);
 }
 
 void SDMXControlConsoleEditorFader::SetValueByPercentage(float InNewPercentage)
@@ -260,6 +260,13 @@ void SDMXControlConsoleEditorFader::ApplyGlobalFilter(const FString& InSearchStr
 	}
 
 	SetVisibility(EVisibility::Collapsed);
+
+	// If not visible, remove from selection
+	if (IsSelected())
+	{
+		const TSharedRef<FDMXControlConsoleEditorSelection> SelectionHandler = FDMXControlConsoleEditorManager::Get().GetSelectionHandler();
+		SelectionHandler->RemoveFromSelection(Fader.Get());
+	}
 }
 
 FReply SDMXControlConsoleEditorFader::OnKeyDown(const FGeometry& MyGeometry, const FKeyEvent& InKeyEvent)
@@ -321,27 +328,26 @@ FReply SDMXControlConsoleEditorFader::OnMouseButtonDown(const FGeometry& MyGeome
 		}
 
 		const TSharedRef<FDMXControlConsoleEditorSelection> SelectionHandler = FDMXControlConsoleEditorManager::Get().GetSelectionHandler();
-		const bool bWasInitallySelected = IsSelected();
 
-		if (!MouseEvent.IsLeftShiftDown())
+		if (MouseEvent.IsLeftShiftDown())
 		{
-			if (!MouseEvent.IsControlDown())
-			{
-				SelectionHandler->ClearSelection();
-			}
-
-			if (!bWasInitallySelected)
+			SelectionHandler->Multiselect(Fader.Get());
+		}
+		else if (MouseEvent.IsControlDown())
+		{
+			if (IsSelected())
 			{
-				SelectionHandler->AddToSelection(Fader.Get());
+				SelectionHandler->RemoveFromSelection(Fader.Get());
 			}
 			else
 			{
-				SelectionHandler->RemoveFromSelection(Fader.Get());
+				SelectionHandler->AddToSelection(Fader.Get());
 			}
 		}
 		else
 		{
-			SelectionHandler->Multiselect(Fader.Get());
+			SelectionHandler->ClearSelection();
+			SelectionHandler->AddToSelection(Fader.Get());
 		}
 
 		return FReply::Handled();
@@ -478,19 +484,23 @@ void SDMXControlConsoleEditorFader::HandleValueChanged(uint32 NewValue)
 	}
 }
 
-void SDMXControlConsoleEditorFader::OnSelectionChanged(UDMXControlConsoleFaderBase* InFader)
+void SDMXControlConsoleEditorFader::OnSelectionChanged()
 {
-	if (!InFader)
+	const TSharedRef<FDMXControlConsoleEditorSelection> SelectionHandler = FDMXControlConsoleEditorManager::Get().GetSelectionHandler();
+	
+	// Remove from selection if not visible. Avoids selecting this when filtering
+	const bool bMultiSelecting = SelectionHandler->GetSelectedFaders().Num() + SelectionHandler->GetSelectedFaderGroups().Num() > 1;
+	if (bMultiSelecting && GetVisibility() == EVisibility::Collapsed)
 	{
+		SelectionHandler->RemoveFromSelection(Fader.Get());
 		return;
 	}
 
-	if (InFader!= Fader || !IsSelected())
+	// Set keyboard focus
+	if (IsSelected())
 	{
-		return;
+		FSlateApplication::Get().SetKeyboardFocus(AsShared());
 	}
-
-	FSlateApplication::Get().SetKeyboardFocus(AsShared());
 }
 
 FReply SDMXControlConsoleEditorFader::OnDeleteClicked()
diff --git a/Engine/Plugins/VirtualProduction/DMX/DMXEngine/Source/DMXControlConsoleEditor/Private/Widgets/SDMXControlConsoleEditorFader.h b/Engine/Plugins/VirtualProduction/DMX/DMXEngine/Source/DMXControlConsoleEditor/Private/Widgets/SDMXControlConsoleEditorFader.h
--- a/Engine/Plugins/VirtualProduction/DMX/DMXEngine/Source/DMXControlConsoleEditor/Private/Widgets/SDMXControlConsoleEditorFader.h
+++ b/Engine/Plugins/VirtualProduction/DMX/DMXEngine/Source/DMXControlConsoleEditor/Private/Widgets/SDMXControlConsoleEditorFader.h
@@ -75,7 +75,7 @@ private:
 	void HandleValueChanged(uint32 NewValue);
 
 	/** Called when fader selection changes */
-	void OnSelectionChanged(UDMXControlConsoleFaderBase* InFader);
+	void OnSelectionChanged();
 
 	/** Called when the delete button was clicked */
 	FReply OnDeleteClicked();
diff --git a/Engine/Plugins/VirtualProduction/DMX/DMXEngine/Source/DMXControlConsoleEditor/Private/Widgets/SDMXControlConsoleEditorFaderGroup.cpp b/Engine/Plugins/VirtualProduction/DMX/DMXEngine/Source/DMXControlConsoleEditor/Private/Widgets/SDMXControlConsoleEditorFaderGroup.cpp
--- a/Engine/Plugins/VirtualProduction/DMX/DMXEngine/Source/DMXControlConsoleEditor/Private/Widgets/SDMXControlConsoleEditorFaderGroup.cpp
+++ b/Engine/Plugins/VirtualProduction/DMX/DMXEngine/Source/DMXControlConsoleEditor/Private/Widgets/SDMXControlConsoleEditorFaderGroup.cpp
@@ -119,7 +119,7 @@ void SDMXControlConsoleEditorFaderGroup::Construct(const FArguments& InArgs, con
 	];
 
 	const TSharedRef<FDMXControlConsoleEditorSelection> SelectionHandler = FDMXControlConsoleEditorManager::Get().GetSelectionHandler();
-	SelectionHandler->GetOnFaderGroupSelectionChanged().AddSP(this, &SDMXControlConsoleEditorFaderGroup::OnSelectionChanged);
+	SelectionHandler->GetOnSelectionChanged().AddSP(this, &SDMXControlConsoleEditorFaderGroup::OnSelectionChanged);
 }
 
 FReply SDMXControlConsoleEditorFaderGroup::OnKeyDown(const FGeometry& MyGeometry, const FKeyEvent& InKeyEvent)
@@ -147,7 +147,6 @@ FReply SDMXControlConsoleEditorFaderGroup::OnKeyDown(const FGeometry& MyGeometry
 
 				SelectedFaderGroup->PostEditChange();
 
-				SelectionHandler->ClearFadersSelection(SelectedFaderGroup);
 				SelectionHandler->RemoveFromSelection(SelectedFaderGroup);
 			}
 		}
@@ -179,33 +178,26 @@ FReply SDMXControlConsoleEditorFaderGroup::OnMouseButtonDown(const FGeometry& My
 		if (FaderGroup)
 		{
 			const TSharedRef<FDMXControlConsoleEditorSelection> SelectionHandler = FDMXControlConsoleEditorManager::Get().GetSelectionHandler();
-			const bool bWasInitiallySelected = IsSelected();
 
-			if (!MouseEvent.IsLeftShiftDown())
+			if (MouseEvent.IsLeftShiftDown())
 			{
-				if (!MouseEvent.IsControlDown())
-				{
-					SelectionHandler->ClearSelection();
-				}
-
-				if (!bWasInitiallySelected)
+				SelectionHandler->Multiselect(FaderGroup);
+			}
+			else if (MouseEvent.IsControlDown())
+			{
+				if (IsSelected())
 				{
-					SelectionHandler->AddToSelection(FaderGroup);
+					SelectionHandler->RemoveFromSelection(FaderGroup);
 				}
 				else
 				{
-					SelectionHandler->RemoveFromSelection(FaderGroup);
+					SelectionHandler->AddToSelection(FaderGroup);
 				}
 			}
-			else 
-			{
-				SelectionHandler->Multiselect(FaderGroup);
-			}
-
-			// Expand/Collapse ExpandArrowButton according to Fader Group selection state
-			if (ExpandArrowButton.IsValid())
+			else
 			{
-				ExpandArrowButton->SetExpandArrow(!bWasInitiallySelected);
+				SelectionHandler->ClearSelection();
+				SelectionHandler->AddToSelection(FaderGroup);
 			}
 		}
 	}
@@ -230,23 +222,25 @@ bool SDMXControlConsoleEditorFaderGroup::IsSelected() const
 	return SelectionHandler->IsSelected(FaderGroup);
 }
 
-void SDMXControlConsoleEditorFaderGroup::OnSelectionChanged(UDMXControlConsoleFaderGroup* InFaderGroup)
+void SDMXControlConsoleEditorFaderGroup::OnSelectionChanged()
 {
-	if (!InFaderGroup)
+	if (UDMXControlConsoleFaderGroup* FaderGroup = GetFaderGroup())
 	{
-		return;
-	}
+		const TSharedRef<FDMXControlConsoleEditorSelection> SelectionHandler = FDMXControlConsoleEditorManager::Get().GetSelectionHandler();
 
-	const UDMXControlConsoleFaderGroup* FaderGroup = GetFaderGroup();
-	if (InFaderGroup != FaderGroup || !IsSelected())
-	{
-		return;
-	}
+		// Remove from selection if not visible. Avoids selecting this when filtering
+		const bool bMultiSelecting = SelectionHandler->GetSelectedFaders().Num() + SelectionHandler->GetSelectedFaderGroups().Num() > 1;
+		if (bMultiSelecting && GetVisibility() == EVisibility::Collapsed)
+		{
+			SelectionHandler->RemoveFromSelection(FaderGroup);
+			return;
+		}
 
-	const TSharedRef<FDMXControlConsoleEditorSelection> SelectionHandler = FDMXControlConsoleEditorManager::Get().GetSelectionHandler();
-	if (SelectionHandler->GetSelectedFadersFromFaderGroup(InFaderGroup).IsEmpty())
-	{
-		FSlateApplication::Get().SetKeyboardFocus(AsShared());
+		// Set keyboard focus
+		if (IsSelected() && SelectionHandler->GetSelectedFadersFromFaderGroup(FaderGroup).IsEmpty())
+		{
+			FSlateApplication::Get().SetKeyboardFocus(AsShared());
+		}
 	}
 }
 
diff --git a/Engine/Plugins/VirtualProduction/DMX/DMXEngine/Source/DMXControlConsoleEditor/Private/Widgets/SDMXControlConsoleEditorFaderGroup.h b/Engine/Plugins/VirtualProduction/DMX/DMXEngine/Source/DMXControlConsoleEditor/Private/Widgets/SDMXControlConsoleEditorFaderGroup.h
--- a/Engine/Plugins/VirtualProduction/DMX/DMXEngine/Source/DMXControlConsoleEditor/Private/Widgets/SDMXControlConsoleEditorFaderGroup.h
+++ b/Engine/Plugins/VirtualProduction/DMX/DMXEngine/Source/DMXControlConsoleEditor/Private/Widgets/SDMXControlConsoleEditorFaderGroup.h
@@ -55,7 +55,7 @@ private:
 	bool IsSelected() const;
 
 	/** Called when fader group selection changes */
-	void OnSelectionChanged(UDMXControlConsoleFaderGroup* InFaderGroup);
+	void OnSelectionChanged();
 
 	/** Gets current FaderGroupName */
 	FText OnGetFaderGroupNameText() const;
