Fix non-unity
#rnx

diff --git a/Engine/Plugins/VirtualProduction/DMX/DMXEngine/Source/DMXControlConsoleEditor/Private/DMXControlConsoleEditorSelection.h b/Engine/Plugins/VirtualProduction/DMX/DMXEngine/Source/DMXControlConsoleEditor/Private/DMXControlConsoleEditorSelection.h
--- a/Engine/Plugins/VirtualProduction/DMX/DMXEngine/Source/DMXControlConsoleEditor/Private/DMXControlConsoleEditorSelection.h
+++ b/Engine/Plugins/VirtualProduction/DMX/DMXEngine/Source/DMXControlConsoleEditor/Private/DMXControlConsoleEditorSelection.h
@@ -3,7 +3,7 @@
 #pragma once
 
 #include "CoreMinimal.h"
-#include "UObject/WeakObjectPtrTemplates.h"
+#include "UObject/WeakObjectPtr.h"
 
 class FDMXControlConsoleEditorManager;
 class UDMXControlConsoleFaderBase;
