DMX - Fix fixture patches and Fixture Types are created without the RF_Public flag when using their Create methods

diff --git a/Engine/Plugins/VirtualProduction/DMX/DMXEngine/Source/DMXRuntime/Private/Library/DMXEntityFixturePatch.cpp b/Engine/Plugins/VirtualProduction/DMX/DMXEngine/Source/DMXRuntime/Private/Library/DMXEntityFixturePatch.cpp
--- a/Engine/Plugins/VirtualProduction/DMX/DMXEngine/Source/DMXRuntime/Private/Library/DMXEntityFixturePatch.cpp
+++ b/Engine/Plugins/VirtualProduction/DMX/DMXEngine/Source/DMXRuntime/Private/Library/DMXEntityFixturePatch.cpp
@@ -51,7 +51,7 @@ UDMXEntityFixturePatch* UDMXEntityFixturePatch::CreateFixturePatchInLibrary(FDMX
 		if (ensureMsgf(DMXLibrary, TEXT("Cannot create Fixture Patch when Fixture Type's DMX Library is invalid.")))
 		{
 			const FString EntityName = FDMXRuntimeUtils::FindUniqueEntityName(DMXLibrary, UDMXEntityFixturePatch::StaticClass(), DesiredName);
-			UDMXEntityFixturePatch* NewFixturePatch = NewObject<UDMXEntityFixturePatch>(DMXLibrary, UDMXEntityFixturePatch::StaticClass(), NAME_None, RF_Transactional);
+			UDMXEntityFixturePatch* NewFixturePatch = NewObject<UDMXEntityFixturePatch>(DMXLibrary, UDMXEntityFixturePatch::StaticClass(), NAME_None, RF_Public | RF_Transactional);
 
 #if WITH_EDITOR
 			if (bMarkDMXLibraryDirty)
diff --git a/Engine/Plugins/VirtualProduction/DMX/DMXEngine/Source/DMXRuntime/Private/Library/DMXEntityFixtureType.cpp b/Engine/Plugins/VirtualProduction/DMX/DMXEngine/Source/DMXRuntime/Private/Library/DMXEntityFixtureType.cpp
--- a/Engine/Plugins/VirtualProduction/DMX/DMXEngine/Source/DMXRuntime/Private/Library/DMXEntityFixtureType.cpp
+++ b/Engine/Plugins/VirtualProduction/DMX/DMXEngine/Source/DMXRuntime/Private/Library/DMXEntityFixtureType.cpp
@@ -173,7 +173,7 @@ UDMXEntityFixtureType* UDMXEntityFixtureType::CreateFixtureTypeInLibrary(FDMXEnt
 
 		FString EntityName = FDMXRuntimeUtils::FindUniqueEntityName(ParentDMXLibrary, UDMXEntityFixtureType::StaticClass(), DesiredName);
 
-		UDMXEntityFixtureType* NewFixtureType = NewObject<UDMXEntityFixtureType>(ParentDMXLibrary, UDMXEntityFixtureType::StaticClass(), NAME_None, RF_Transactional);
+		UDMXEntityFixtureType* NewFixtureType = NewObject<UDMXEntityFixtureType>(ParentDMXLibrary, UDMXEntityFixtureType::StaticClass(), NAME_None, RF_Public | RF_Transactional);
 		NewFixtureType->SetName(EntityName);
 		NewFixtureType->DMXCategory = ConstructionParams.DMXCategory;
 		NewFixtureType->Modes = ConstructionParams.Modes;
diff --git a/Engine/Plugins/VirtualProduction/DMX/DMXEngine/Source/DMXRuntime/Private/Library/DMXLibrary.cpp b/Engine/Plugins/VirtualProduction/DMX/DMXEngine/Source/DMXRuntime/Private/Library/DMXLibrary.cpp
--- a/Engine/Plugins/VirtualProduction/DMX/DMXEngine/Source/DMXRuntime/Private/Library/DMXLibrary.cpp
+++ b/Engine/Plugins/VirtualProduction/DMX/DMXEngine/Source/DMXRuntime/Private/Library/DMXLibrary.cpp
@@ -30,22 +30,6 @@ UDMXLibrary::UDMXLibrary()
 	GeneralSceneDescription = NewObject<UDMXMVRGeneralSceneDescription>(this, GeneralSceneDescriptionName);
 }
 
-void UDMXLibrary::Serialize(FArchive& Ar)
-{
-// More performant without. Instead Update the General Scene Description on demand when it's actually needed.
-// Serialize is only here because the 5.1.1 hotfix doens't allow a header change and can be removed in 5.2.
-// 
-// #if WITH_EDITORONLY_DATA
-//	if (Ar.IsSaving())
-//	{
-//		// Update the General Scene Description before saving it
-//		UpdateGeneralSceneDescription();
-//	}
-//#endif
-
-	Super::Serialize(Ar);
-}
-
 void UDMXLibrary::PostInitProperties()
 {
 	Super::PostInitProperties();
@@ -92,6 +76,12 @@ void UDMXLibrary::PostLoad()
 				UE_LOG(LogDMXRuntime, Warning, TEXT("Invalid Entity found in Library %s. Please resave the library."), *GetName());
 				Entities.Remove(Entity);
 			}
+
+			// UE-177836: Fix RF_Public flag is not set on Entities in 5.0, 5.1.
+			if (!Entity->HasAnyFlags(RF_Public))
+			{
+				Entity->SetFlags(Entity->GetFlags() | RF_Public);
+			}
 		}
 	}
 }
diff --git a/Engine/Plugins/VirtualProduction/DMX/DMXEngine/Source/DMXRuntime/Public/Library/DMXLibrary.h b/Engine/Plugins/VirtualProduction/DMX/DMXEngine/Source/DMXRuntime/Public/Library/DMXLibrary.h
--- a/Engine/Plugins/VirtualProduction/DMX/DMXEngine/Source/DMXRuntime/Public/Library/DMXLibrary.h
+++ b/Engine/Plugins/VirtualProduction/DMX/DMXEngine/Source/DMXRuntime/Public/Library/DMXLibrary.h
@@ -56,7 +56,6 @@ public:
 
 protected:
 	// ~Begin UObject Interface
-	virtual void Serialize(FArchive& Ar) override;
 	virtual void PostInitProperties() override;
 	virtual void PostLoad() override;
 	virtual void PostDuplicate(EDuplicateMode::Type DuplicateMode) override;
