DMX - Double click now correctly expands/collapses control console fader groups¨

#lodckdown Alejandro.Arango

diff --git a/Engine/Plugins/VirtualProduction/DMX/DMXEngine/Source/DMXControlConsoleEditor/Private/Widgets/SDMXControlConsoleEditorFaderGroup.cpp b/Engine/Plugins/VirtualProduction/DMX/DMXEngine/Source/DMXControlConsoleEditor/Private/Widgets/SDMXControlConsoleEditorFaderGroup.cpp
--- a/Engine/Plugins/VirtualProduction/DMX/DMXEngine/Source/DMXControlConsoleEditor/Private/Widgets/SDMXControlConsoleEditorFaderGroup.cpp
+++ b/Engine/Plugins/VirtualProduction/DMX/DMXEngine/Source/DMXControlConsoleEditor/Private/Widgets/SDMXControlConsoleEditorFaderGroup.cpp
@@ -205,6 +205,20 @@ FReply SDMXControlConsoleEditorFaderGroup::OnMouseButtonDown(const FGeometry& My
 	return FReply::Handled();
 }
 
+FReply SDMXControlConsoleEditorFaderGroup::OnMouseButtonDoubleClick(const FGeometry& MyGeometry, const FPointerEvent& MouseEvent)
+{
+	if (MouseEvent.GetEffectingButton() == EKeys::LeftMouseButton)
+	{
+		if (ExpandArrowButton.IsValid())
+		{
+			ExpandArrowButton->ToggleExpandArrow();
+			return FReply::Handled();
+		}
+	}
+
+	return FReply::Unhandled();
+}
+
 UDMXControlConsoleFaderGroup* SDMXControlConsoleEditorFaderGroup::GetFaderGroup() const
 {
 	return FaderGroupView.IsValid() ? FaderGroupView.Pin()->GetFaderGroup() : nullptr;
diff --git a/Engine/Plugins/VirtualProduction/DMX/DMXEngine/Source/DMXControlConsoleEditor/Private/Widgets/SDMXControlConsoleEditorFaderGroup.h b/Engine/Plugins/VirtualProduction/DMX/DMXEngine/Source/DMXControlConsoleEditor/Private/Widgets/SDMXControlConsoleEditorFaderGroup.h
--- a/Engine/Plugins/VirtualProduction/DMX/DMXEngine/Source/DMXControlConsoleEditor/Private/Widgets/SDMXControlConsoleEditorFaderGroup.h
+++ b/Engine/Plugins/VirtualProduction/DMX/DMXEngine/Source/DMXControlConsoleEditor/Private/Widgets/SDMXControlConsoleEditorFaderGroup.h
@@ -44,6 +44,7 @@ protected:
 	//~ Begin SWidget interface
 	virtual FReply OnKeyDown(const FGeometry& MyGeometry, const FKeyEvent& InKeyEvent) override;
 	virtual FReply OnMouseButtonDown(const FGeometry& MyGeometry, const FPointerEvent& MouseEvent) override;
+	virtual FReply OnMouseButtonDoubleClick(const FGeometry& MyGeometry, const FPointerEvent& MouseEvent) override;
 	virtual bool SupportsKeyboardFocus() const override { return true; }
 	//~ End of SWidget interface
 
