DMX - Fix crash when adding fixture attributes to project settings while DMX library is open

diff --git a/Engine/Plugins/VirtualProduction/DMX/DMXEngine/Source/DMXEditor/Private/Customizations/DMXFixtureCategoryCustomization.cpp b/Engine/Plugins/VirtualProduction/DMX/DMXEngine/Source/DMXEditor/Private/Customizations/DMXFixtureCategoryCustomization.cpp
--- a/Engine/Plugins/VirtualProduction/DMX/DMXEngine/Source/DMXEditor/Private/Customizations/DMXFixtureCategoryCustomization.cpp
+++ b/Engine/Plugins/VirtualProduction/DMX/DMXEngine/Source/DMXEditor/Private/Customizations/DMXFixtureCategoryCustomization.cpp
@@ -50,11 +50,6 @@ void FDMXFixtureCategoryCustomization::CustomizeHeader(TSharedRef<IPropertyHandl
 			.bDisplayWarningIcon(true)
 			.OnValueChanged(this, &FDMXFixtureCategoryCustomization::SetValue)
 		];
-
-	ProtocolSettings->GetOnDefaultAttributesChanged().AddLambda([PropertyUtilities]()
-		{
-			PropertyUtilities->ForceRefresh();
-		});
 }
 
 void FDMXFixtureCategoryCustomization::CustomizeChildren(TSharedRef<IPropertyHandle> InPropertyHandle, IDetailChildrenBuilder& InChildBuilder, IPropertyTypeCustomizationUtils& CustomizationUtils)
