DMX - Reseting Filters doesn't work for Matrix Cell Faders

diff --git a/Engine/Plugins/VirtualProduction/DMX/DMXEngine/Source/DMXControlConsoleEditor/Private/Views/SDMXControlConsoleEditorFaderGroupView.cpp b/Engine/Plugins/VirtualProduction/DMX/DMXEngine/Source/DMXControlConsoleEditor/Private/Views/SDMXControlConsoleEditorFaderGroupView.cpp
--- a/Engine/Plugins/VirtualProduction/DMX/DMXEngine/Source/DMXControlConsoleEditor/Private/Views/SDMXControlConsoleEditorFaderGroupView.cpp
+++ b/Engine/Plugins/VirtualProduction/DMX/DMXEngine/Source/DMXControlConsoleEditor/Private/Views/SDMXControlConsoleEditorFaderGroupView.cpp
@@ -160,20 +160,8 @@ void SDMXControlConsoleEditorFaderGroupView::ShowAllElements()
 	{
 		const TSharedRef<SWidget> Child = Children->GetChildAt(ChildIndex);
 		if (Child->GetType() == "SDMXControlConsoleEditorMatrixCell")
-		{
-			StaticCastSharedRef<SDMXControlConsoleEditorMatrixCell>(Child)->SetVisibility(EVisibility::Visible);
-			FChildren* CellChildren = ElementsHorizontalBox->GetChildren();
-			if (!CellChildren)
-			{
-				continue;
-			}
-			
-			int32 NumCellChildren = CellChildren->Num();
-			for (int32 CellChildIndex = 0; CellChildIndex < NumChildren; ++CellChildIndex)
-			{
-				const TSharedRef<SWidget> CellChild = Children->GetChildAt(ChildIndex);
-				CellChild->SetVisibility(EVisibility::Visible);
-			}
+		{			
+			StaticCastSharedRef<SDMXControlConsoleEditorMatrixCell>(Child)->ApplyGlobalFilter(TEXT(""));
 		}
 		else if (Child->GetType() == "SDMXControlConsoleEditorFader")
 		{
diff --git a/Engine/Plugins/VirtualProduction/DMX/DMXEngine/Source/DMXControlConsoleEditor/Private/Widgets/SDMXControlConsoleEditorExpandArrowButton.h b/Engine/Plugins/VirtualProduction/DMX/DMXEngine/Source/DMXControlConsoleEditor/Private/Widgets/SDMXControlConsoleEditorExpandArrowButton.h
--- a/Engine/Plugins/VirtualProduction/DMX/DMXEngine/Source/DMXControlConsoleEditor/Private/Widgets/SDMXControlConsoleEditorExpandArrowButton.h
+++ b/Engine/Plugins/VirtualProduction/DMX/DMXEngine/Source/DMXControlConsoleEditor/Private/Widgets/SDMXControlConsoleEditorExpandArrowButton.h
@@ -46,5 +46,5 @@ private:
 	ExpandArrowButtonDelegate OnExpandClicked;
 
 	/** True if the button is in expand state */
-	bool bIsExpanded = false;
+	bool bIsExpanded = true;
 };
