DMX - Undo changes submitted for UE-177836 to recover from various undesired side effects when deleting an entity from a DMX Library

diff --git a/Engine/Plugins/VirtualProduction/DMX/DMXEngine/Source/DMXRuntime/Private/Library/DMXEntityFixturePatch.cpp b/Engine/Plugins/VirtualProduction/DMX/DMXEngine/Source/DMXRuntime/Private/Library/DMXEntityFixturePatch.cpp
--- a/Engine/Plugins/VirtualProduction/DMX/DMXEngine/Source/DMXRuntime/Private/Library/DMXEntityFixturePatch.cpp
+++ b/Engine/Plugins/VirtualProduction/DMX/DMXEngine/Source/DMXRuntime/Private/Library/DMXEntityFixturePatch.cpp
@@ -51,7 +51,7 @@ UDMXEntityFixturePatch* UDMXEntityFixturePatch::CreateFixturePatchInLibrary(FDMX
 		if (ensureMsgf(DMXLibrary, TEXT("Cannot create Fixture Patch when Fixture Type's DMX Library is invalid.")))
 		{
 			const FString EntityName = FDMXRuntimeUtils::FindUniqueEntityName(DMXLibrary, UDMXEntityFixturePatch::StaticClass(), DesiredName);
-			UDMXEntityFixturePatch* NewFixturePatch = NewObject<UDMXEntityFixturePatch>(DMXLibrary, UDMXEntityFixturePatch::StaticClass(), NAME_None, RF_Public | RF_Transactional);
+			UDMXEntityFixturePatch* NewFixturePatch = NewObject<UDMXEntityFixturePatch>(DMXLibrary, UDMXEntityFixturePatch::StaticClass(), NAME_None, RF_Transactional);
 
 #if WITH_EDITOR
 			if (bMarkDMXLibraryDirty)
diff --git a/Engine/Plugins/VirtualProduction/DMX/DMXEngine/Source/DMXRuntime/Private/Library/DMXEntityFixtureType.cpp b/Engine/Plugins/VirtualProduction/DMX/DMXEngine/Source/DMXRuntime/Private/Library/DMXEntityFixtureType.cpp
--- a/Engine/Plugins/VirtualProduction/DMX/DMXEngine/Source/DMXRuntime/Private/Library/DMXEntityFixtureType.cpp
+++ b/Engine/Plugins/VirtualProduction/DMX/DMXEngine/Source/DMXRuntime/Private/Library/DMXEntityFixtureType.cpp
@@ -173,7 +173,7 @@ UDMXEntityFixtureType* UDMXEntityFixtureType::CreateFixtureTypeInLibrary(FDMXEnt
 
 		FString EntityName = FDMXRuntimeUtils::FindUniqueEntityName(ParentDMXLibrary, UDMXEntityFixtureType::StaticClass(), DesiredName);
 
-		UDMXEntityFixtureType* NewFixtureType = NewObject<UDMXEntityFixtureType>(ParentDMXLibrary, UDMXEntityFixtureType::StaticClass(), NAME_None, RF_Public | RF_Transactional);
+		UDMXEntityFixtureType* NewFixtureType = NewObject<UDMXEntityFixtureType>(ParentDMXLibrary, UDMXEntityFixtureType::StaticClass(), NAME_None, RF_Transactional);
 		NewFixtureType->SetName(EntityName);
 		NewFixtureType->DMXCategory = ConstructionParams.DMXCategory;
 		NewFixtureType->Modes = ConstructionParams.Modes;
diff --git a/Engine/Plugins/VirtualProduction/DMX/DMXEngine/Source/DMXRuntime/Private/Library/DMXLibrary.cpp b/Engine/Plugins/VirtualProduction/DMX/DMXEngine/Source/DMXRuntime/Private/Library/DMXLibrary.cpp
--- a/Engine/Plugins/VirtualProduction/DMX/DMXEngine/Source/DMXRuntime/Private/Library/DMXLibrary.cpp
+++ b/Engine/Plugins/VirtualProduction/DMX/DMXEngine/Source/DMXRuntime/Private/Library/DMXLibrary.cpp
@@ -76,12 +76,6 @@ void UDMXLibrary::PostLoad()
 				UE_LOG(LogDMXRuntime, Warning, TEXT("Invalid Entity found in Library %s. Please resave the library."), *GetName());
 				Entities.Remove(Entity);
 			}
-
-			// UE-177836: Fix RF_Public flag is not set on Entities in 5.0, 5.1.
-			if (!Entity->HasAnyFlags(RF_Public))
-			{
-				Entity->SetFlags(Entity->GetFlags() | RF_Public);
-			}
 		}
 	}
 }
