DMX - Fix modifying a sequencer over MU will result in a crash on distant editor

diff --git a/Engine/Plugins/VirtualProduction/DMX/DMXEngine/Source/DMXRuntime/Private/Sequencer/MovieSceneDMXLibrarySection.cpp b/Engine/Plugins/VirtualProduction/DMX/DMXEngine/Source/DMXRuntime/Private/Sequencer/MovieSceneDMXLibrarySection.cpp
--- a/Engine/Plugins/VirtualProduction/DMX/DMXEngine/Source/DMXRuntime/Private/Sequencer/MovieSceneDMXLibrarySection.cpp
+++ b/Engine/Plugins/VirtualProduction/DMX/DMXEngine/Source/DMXRuntime/Private/Sequencer/MovieSceneDMXLibrarySection.cpp
@@ -365,6 +365,15 @@ void UMovieSceneDMXLibrarySection::PostEditImport()
 	UpdateChannelProxy();
 }
 
+#if WITH_EDITOR
+void UMovieSceneDMXLibrarySection::PostEditChangeProperty(FPropertyChangedEvent& PropertyChangedEvent)
+{
+	Super::PostEditChangeProperty(PropertyChangedEvent);
+
+	UpdateChannelProxy();
+}
+#endif // WITH_EDITOR
+
 void UMovieSceneDMXLibrarySection::RefreshChannels()
 {
 	UpdateChannelProxy();
diff --git a/Engine/Plugins/VirtualProduction/DMX/DMXEngine/Source/DMXRuntime/Public/Sequencer/MovieSceneDMXLibrarySection.h b/Engine/Plugins/VirtualProduction/DMX/DMXEngine/Source/DMXRuntime/Public/Sequencer/MovieSceneDMXLibrarySection.h
--- a/Engine/Plugins/VirtualProduction/DMX/DMXEngine/Source/DMXRuntime/Public/Sequencer/MovieSceneDMXLibrarySection.h
+++ b/Engine/Plugins/VirtualProduction/DMX/DMXEngine/Source/DMXRuntime/Public/Sequencer/MovieSceneDMXLibrarySection.h
@@ -164,6 +164,9 @@ public:
 	// Begin UObject interface
 	virtual void Serialize(FArchive& Ar) override;
 	virtual void PostEditImport() override;
+#if WITH_EDITOR
+	void PostEditChangeProperty(FPropertyChangedEvent& PropertyChangedEvent) override;
+#endif // WITH_EDITOR
 	// ~End UObject interface
 
 	/** Refreshes the channels. Useful e.g. when underlying DMX Library changes */
