DMX - Fix Monitors do not recognize newly added ports when open

diff --git a/Engine/Plugins/VirtualProduction/DMX/DMXEngine/Source/DMXEditor/Private/Widgets/FixturePatch/SDMXFixturePatcher.cpp b/Engine/Plugins/VirtualProduction/DMX/DMXEngine/Source/DMXEditor/Private/Widgets/FixturePatch/SDMXFixturePatcher.cpp
--- a/Engine/Plugins/VirtualProduction/DMX/DMXEngine/Source/DMXEditor/Private/Widgets/FixturePatch/SDMXFixturePatcher.cpp
+++ b/Engine/Plugins/VirtualProduction/DMX/DMXEngine/Source/DMXEditor/Private/Widgets/FixturePatch/SDMXFixturePatcher.cpp
@@ -41,6 +41,8 @@ void SDMXFixturePatcher::Construct(const FArguments& InArgs)
 
 	SharedData = DMXEditorPtr.Pin()->GetFixturePatchSharedData();
 		
+	const UDMXEditorSettings* DMXEditorSettings = GetDefault<UDMXEditorSettings>();
+	const ECheckBoxState InitialDMXMonitorEnabledCheckBoxState = DMXEditorSettings->bFixturePatcherDMXMonitorEnabled ? ECheckBoxState::Checked : ECheckBoxState::Unchecked;
 	const FLinearColor BackgroundTint(0.6f, 0.6f, 0.6f, 1.0f);
 	ChildSlot			
 		[
@@ -144,7 +146,7 @@ void SDMXFixturePatcher::Construct(const FArguments& InArgs)
 						.Padding(FMargin(4.0f, 4.0f, 15.0f, 4.0f))
 						[
 							SAssignNew(EnableDMXMonitorCheckBox, SCheckBox)
-							.IsChecked(false)
+							.IsChecked(InitialDMXMonitorEnabledCheckBoxState)
 							.OnCheckStateChanged(this, &SDMXFixturePatcher::OnToggleDMXMonitorEnabled)
 						]
 					]
@@ -188,7 +190,6 @@ void SDMXFixturePatcher::Construct(const FArguments& InArgs)
 
 	ShowSelectedUniverse();
 
-	const UDMXEditorSettings* DMXEditorSettings = GetDefault<UDMXEditorSettings>();
 	SetDMXMonitorEnabled(DMXEditorSettings->bFixturePatcherDMXMonitorEnabled);
 }
 END_SLATE_FUNCTION_BUILD_OPTIMIZATION
diff --git a/Engine/Plugins/VirtualProduction/DMX/DMXEngine/Source/DMXEditor/Private/Widgets/Monitors/SDMXActivityMonitor.cpp b/Engine/Plugins/VirtualProduction/DMX/DMXEngine/Source/DMXEditor/Private/Widgets/Monitors/SDMXActivityMonitor.cpp
--- a/Engine/Plugins/VirtualProduction/DMX/DMXEngine/Source/DMXEditor/Private/Widgets/Monitors/SDMXActivityMonitor.cpp
+++ b/Engine/Plugins/VirtualProduction/DMX/DMXEngine/Source/DMXEditor/Private/Widgets/Monitors/SDMXActivityMonitor.cpp
@@ -8,6 +8,7 @@
 #include "DMXProtocolTypes.h"
 #include "IO/DMXInputPort.h"
 #include "IO/DMXOutputPort.h"
+#include "IO/DMXPortManager.h"
 #include "IO/DMXRawListener.h"
 #include "Widgets/Monitors/SDMXActivityInUniverse.h"
 #include "Widgets/Monitors/SDMXMonitorSourceSelector.h"
@@ -194,6 +195,8 @@ void SDMXActivityMonitor::Construct(const FArguments& InArgs)
 	LoadMonitorSettings();
 
 	UpdateListenerRegistration();
+
+	FDMXPortManager::Get().OnPortsChanged.AddSP(this, &SDMXActivityMonitor::UpdateListenerRegistration);
 }
 END_SLATE_FUNCTION_BUILD_OPTIMIZATION
 
diff --git a/Engine/Plugins/VirtualProduction/DMX/DMXEngine/Source/DMXEditor/Private/Widgets/Monitors/SDMXMonitorSourceSelector.cpp b/Engine/Plugins/VirtualProduction/DMX/DMXEngine/Source/DMXEditor/Private/Widgets/Monitors/SDMXMonitorSourceSelector.cpp
--- a/Engine/Plugins/VirtualProduction/DMX/DMXEngine/Source/DMXEditor/Private/Widgets/Monitors/SDMXMonitorSourceSelector.cpp
+++ b/Engine/Plugins/VirtualProduction/DMX/DMXEngine/Source/DMXEditor/Private/Widgets/Monitors/SDMXMonitorSourceSelector.cpp
@@ -137,6 +137,8 @@ void SDMXMonitorSourceSelector::Construct(const FArguments& InArgs)
 	];
 
 	UpdateMonitoredPorts();
+
+	FDMXPortManager::Get().OnPortsChanged.AddSP(this, &SDMXMonitorSourceSelector::UpdateMonitoredPorts);
 }
 END_SLATE_FUNCTION_BUILD_OPTIMIZATION
 
diff --git a/Engine/Plugins/VirtualProduction/DMX/DMXProtocol/Source/DMXProtocol/Private/IO/DMXInputPort.cpp b/Engine/Plugins/VirtualProduction/DMX/DMXProtocol/Source/DMXProtocol/Private/IO/DMXInputPort.cpp
--- a/Engine/Plugins/VirtualProduction/DMX/DMXProtocol/Source/DMXProtocol/Private/IO/DMXInputPort.cpp
+++ b/Engine/Plugins/VirtualProduction/DMX/DMXProtocol/Source/DMXProtocol/Private/IO/DMXInputPort.cpp
@@ -438,10 +438,6 @@ FDMXInputPortSharedRef FDMXInputPort::CreateFromConfig(FDMXInputPortConfig& Inpu
 
 FDMXInputPort::~FDMXInputPort()
 {
-	// All Listeners need to be explicitly removed before destruction 
-	check(RawListeners.Num() == 0);
-	check(LocalUniverseToListenerGroupMap.Num() == 0);
-
 	// Port needs be unregistered before destruction
 	check(!bRegistered);
 
diff --git a/Engine/Plugins/VirtualProduction/DMX/DMXProtocol/Source/DMXProtocol/Private/IO/DMXOutputPort.cpp b/Engine/Plugins/VirtualProduction/DMX/DMXProtocol/Source/DMXProtocol/Private/IO/DMXOutputPort.cpp
--- a/Engine/Plugins/VirtualProduction/DMX/DMXProtocol/Source/DMXProtocol/Private/IO/DMXOutputPort.cpp
+++ b/Engine/Plugins/VirtualProduction/DMX/DMXProtocol/Source/DMXProtocol/Private/IO/DMXOutputPort.cpp
@@ -558,9 +558,6 @@ FDMXOutputPortSharedRef FDMXOutputPort::CreateFromConfig(FDMXOutputPortConfig& O
 
 FDMXOutputPort::~FDMXOutputPort()
 {	
-	// All Listeners need to be explicitly removed before destruction 
-	check(RawListeners.Num() == 0);
-	
 	// Port needs be unregistered before destruction
 	check(DMXSenderArray.Num() == 0);
 
