DMX - Fix DMX entities cannot be reset when calling a setter with an invalid object

diff --git a/Engine/Plugins/VirtualProduction/DMX/DMXEngine/Source/DMXRuntime/Private/Library/DMXEntityReference.cpp b/Engine/Plugins/VirtualProduction/DMX/DMXEngine/Source/DMXRuntime/Private/Library/DMXEntityReference.cpp
--- a/Engine/Plugins/VirtualProduction/DMX/DMXEngine/Source/DMXRuntime/Private/Library/DMXEntityReference.cpp
+++ b/Engine/Plugins/VirtualProduction/DMX/DMXEngine/Source/DMXRuntime/Private/Library/DMXEntityReference.cpp
@@ -29,17 +29,14 @@ void FDMXEntityReference::SetEntity(UDMXEntity* NewEntity)
 		
 		CachedEntity = NewEntity;
 	}
-	else if (EntityId.IsValid())
+	else
 	{
-		InvalidateId();
+		DMXLibrary = nullptr;
+		EntityId.Invalidate();
+		CachedEntity.Reset();
 	}
 }
 
-void FDMXEntityReference::InvalidateId()
-{
-	EntityId.Invalidate();
-}
-
 UDMXEntity* FDMXEntityReference::GetEntity() const
 {
 	if (CachedEntity.IsValid())
diff --git a/Engine/Plugins/VirtualProduction/DMX/DMXEngine/Source/DMXRuntime/Public/Library/DMXEntityReference.h b/Engine/Plugins/VirtualProduction/DMX/DMXEngine/Source/DMXRuntime/Public/Library/DMXEntityReference.h
--- a/Engine/Plugins/VirtualProduction/DMX/DMXEngine/Source/DMXRuntime/Public/Library/DMXEntityReference.h
+++ b/Engine/Plugins/VirtualProduction/DMX/DMXEngine/Source/DMXRuntime/Public/Library/DMXEntityReference.h
@@ -76,14 +76,6 @@ public:
 protected:
 	UPROPERTY()
 	TSubclassOf<UDMXEntity> EntityType;
-
-private:
-	/** 
-	 * Set the EntityId to an invalid state.
-	 * Called automatically when the DMX Library is changed or SetEntityId is called with an Id that doesn't
-	 * belong to the current DMX Library.
-	 */
-	void InvalidateId();
 	
 private:
 	/** The entity's unique ID */
