DMX - Fix undo removing a Fixture Patch changes Fixture ID of the Patch

diff --git a/Engine/Plugins/VirtualProduction/DMX/DMXEngine/Source/DMXEditor/Private/Widgets/FixturePatch/SDMXMVRFixtureList.cpp b/Engine/Plugins/VirtualProduction/DMX/DMXEngine/Source/DMXEditor/Private/Widgets/FixturePatch/SDMXMVRFixtureList.cpp
--- a/Engine/Plugins/VirtualProduction/DMX/DMXEngine/Source/DMXEditor/Private/Widgets/FixturePatch/SDMXMVRFixtureList.cpp
+++ b/Engine/Plugins/VirtualProduction/DMX/DMXEngine/Source/DMXEditor/Private/Widgets/FixturePatch/SDMXMVRFixtureList.cpp
@@ -740,6 +740,7 @@ void SDMXMVRFixtureList::RefreshList()
 			if (!bUpdatedGeneralSceneDescription)
 			{	
 				// Ony update the General Scene Description when Nodes were never acquired, or new patches were added.
+				DMXLibrary->Modify();
 				DMXLibrary->UpdateGeneralSceneDescription();
 				bUpdatedGeneralSceneDescription = true;
 			}
