DMX - Control Console is its own experimental plugin

diff --git a/Engine/Plugins/VirtualProduction/DMX/DMXEngine/DMXEngine.uplugin b/Engine/Plugins/VirtualProduction/DMX/DMXEngine/DMXEngine.uplugin
--- a/Engine/Plugins/VirtualProduction/DMX/DMXEngine/DMXEngine.uplugin
+++ b/Engine/Plugins/VirtualProduction/DMX/DMXEngine/DMXEngine.uplugin
@@ -18,16 +18,6 @@
 			"Type": "UncookedOnly",
 			"LoadingPhase": "Default"
 		},
-		{
-			"Name": "DMXControlConsole",
-			"Type": "Runtime",
-			"LoadingPhase": "Default"
-		},
-		{
-			"Name": "DMXControlConsoleEditor",
-			"Type": "Editor",
-			"LoadingPhase": "Default"
-		},
 		{
 			"Name": "DMXEditor",
 			"Type": "Editor",
