DMX - Fix an issue where datasmith importer doesn't save ports it updated, account for discrepancies between vectorworks datasmith and mvr export filenames

diff --git a/Engine/Plugins/VirtualProduction/DMX/DMXEngine/Source/DMXEditor/Private/Factories/DMXLibraryFromMVRImportOptions.cpp b/Engine/Plugins/VirtualProduction/DMX/DMXEngine/Source/DMXEditor/Private/Factories/DMXLibraryFromMVRImportOptions.cpp
--- a/Engine/Plugins/VirtualProduction/DMX/DMXEngine/Source/DMXEditor/Private/Factories/DMXLibraryFromMVRImportOptions.cpp
+++ b/Engine/Plugins/VirtualProduction/DMX/DMXEngine/Source/DMXEditor/Private/Factories/DMXLibraryFromMVRImportOptions.cpp
@@ -47,10 +47,10 @@ void UDMXLibraryFromMVRImportOptions::ApplyOptions(UDMXLibrary* DMXLibrary)
 
 	if (bUpdateInputPort)
 	{
+		UDMXProtocolSettings* ProtocolSettings = GetMutableDefault<UDMXProtocolSettings>();
 		const TSharedPtr<FDMXInputPort> InputPort = FDMXPortManager::Get().FindInputPortByGuid(InputPortToUpdate.GetPortGuid());
 		if (InputPort.IsValid())
 		{
-			UDMXProtocolSettings* ProtocolSettings = GetMutableDefault<UDMXProtocolSettings>();
 			FDMXInputPortConfig* InputPortConfigPtr = Algo::FindByPredicate(ProtocolSettings->InputPortConfigs, [&InputPort](const FDMXInputPortConfig& InputPortConfig)
 				{
 					return InputPortConfig.GetPortGuid() == InputPort->GetPortGuid();
@@ -73,21 +73,21 @@ void UDMXLibraryFromMVRImportOptions::ApplyOptions(UDMXLibrary* DMXLibrary)
 			ConfigParams.LocalUniverseStart = FMath::Min(NewConfig.GetLocalUniverseStart(), FirstUniverse);
 			ConfigParams.NumUniverses = FMath::Max(NewConfig.GetNumUniverses(), NumUniverses);
 
-			UDMXProtocolSettings* ProtocolSettings = GetMutableDefault<UDMXProtocolSettings>();
 			ProtocolSettings->PreEditChange(UDMXProtocolSettings::StaticClass()->FindPropertyByName(GET_MEMBER_NAME_CHECKED(UDMXProtocolSettings, InputPortConfigs)));
 			ProtocolSettings->InputPortConfigs.Add(NewConfig);
 			ProtocolSettings->PostEditChange();
 		}
+		ProtocolSettings->SaveConfig();
 
 		FDMXPortManager::Get().UpdateFromProtocolSettings();
 	}
 
 	if (bUpdateOutputPort)
 	{
+		UDMXProtocolSettings* ProtocolSettings = GetMutableDefault<UDMXProtocolSettings>();
 		const TSharedPtr<FDMXOutputPort> OutputPort = FDMXPortManager::Get().FindOutputPortByGuid(OutputPortToUpdate.GetPortGuid());
 		if (OutputPort.IsValid())
 		{
-			UDMXProtocolSettings* ProtocolSettings = GetMutableDefault<UDMXProtocolSettings>();
 			FDMXOutputPortConfig* OutputPortConfigPtr = Algo::FindByPredicate(ProtocolSettings->OutputPortConfigs, [&OutputPort](const FDMXOutputPortConfig& OutputPortConfig)
 				{
 					return OutputPortConfig.GetPortGuid() == OutputPort->GetPortGuid();
@@ -110,11 +110,11 @@ void UDMXLibraryFromMVRImportOptions::ApplyOptions(UDMXLibrary* DMXLibrary)
 			ConfigParams.LocalUniverseStart = FMath::Min(NewConfig.GetLocalUniverseStart(), FirstUniverse);
 			ConfigParams.NumUniverses = FMath::Max(NewConfig.GetNumUniverses(), NumUniverses);
 
-			UDMXProtocolSettings* ProtocolSettings = GetMutableDefault<UDMXProtocolSettings>();
 			ProtocolSettings->PreEditChange(UDMXProtocolSettings::StaticClass()->FindPropertyByName(GET_MEMBER_NAME_CHECKED(UDMXProtocolSettings, OutputPortConfigs)));
 			ProtocolSettings->OutputPortConfigs.Add(NewConfig);
 			ProtocolSettings->PostEditChange();
 		}
+		ProtocolSettings->SaveConfig();
 
 		FDMXPortManager::Get().UpdateFromProtocolSettings();
 	}
diff --git a/Engine/Plugins/VirtualProduction/DMX/DatasmithMVR/Source/DatasmithMVRTranslator/Private/DatasmithMVRNativeTranslator.cpp b/Engine/Plugins/VirtualProduction/DMX/DatasmithMVR/Source/DatasmithMVRTranslator/Private/DatasmithMVRNativeTranslator.cpp
--- a/Engine/Plugins/VirtualProduction/DMX/DatasmithMVR/Source/DatasmithMVRTranslator/Private/DatasmithMVRNativeTranslator.cpp
+++ b/Engine/Plugins/VirtualProduction/DMX/DatasmithMVR/Source/DatasmithMVRTranslator/Private/DatasmithMVRNativeTranslator.cpp
@@ -95,7 +95,6 @@ bool FDatasmithMVRNativeTranslator::FindMVRFile(FString& OutMVRFilePathAndName)
 	const FString DatasmithAssetsPath = FPaths::GetPath(DatasmithFilePathAndName) / DatasmithFileName + TEXT("_Assets");
 
 	const FString MVRFilename = FPaths::GetBaseFilename(DatasmithFilePathAndName) + TEXT(".mvr");
-
 	if (FPaths::FileExists(DatasmithPath / MVRFilename))
 	{
 		OutMVRFilePathAndName = DatasmithPath / MVRFilename;
@@ -104,7 +103,20 @@ bool FDatasmithMVRNativeTranslator::FindMVRFile(FString& OutMVRFilePathAndName)
 	else if (FPaths::FileExists(DatasmithAssetsPath / MVRFilename))
 	{
 		OutMVRFilePathAndName = DatasmithAssetsPath / MVRFilename;
-		return true;;
+		return true;
+	}
+
+	// Vectorworks exports .udatasmith files with underscores, but MVR without underscores by default
+	const FString MVRFilenameNoUnderscores = MVRFilename.Replace(TEXT("_"), TEXT(" "));
+	if (FPaths::FileExists(DatasmithPath / MVRFilenameNoUnderscores))
+	{
+		OutMVRFilePathAndName = DatasmithPath / MVRFilenameNoUnderscores;
+		return true;
+	}
+	else if (FPaths::FileExists(DatasmithAssetsPath / MVRFilenameNoUnderscores))
+	{
+		OutMVRFilePathAndName = DatasmithAssetsPath / MVRFilenameNoUnderscores;
+		return true;
 	}
 
 	return false;
