DMX - Fix Level Sequence DMX tracks fail to load when Sequence and DMX Library reside in plugin content

diff --git a/Engine/Plugins/VirtualProduction/DMX/DMXEngine/Source/DMXRuntime/Private/Sequencer/MovieSceneDMXLibrarySection.cpp b/Engine/Plugins/VirtualProduction/DMX/DMXEngine/Source/DMXRuntime/Private/Sequencer/MovieSceneDMXLibrarySection.cpp
--- a/Engine/Plugins/VirtualProduction/DMX/DMXEngine/Source/DMXRuntime/Private/Sequencer/MovieSceneDMXLibrarySection.cpp
+++ b/Engine/Plugins/VirtualProduction/DMX/DMXEngine/Source/DMXRuntime/Private/Sequencer/MovieSceneDMXLibrarySection.cpp
@@ -353,11 +353,16 @@ void UMovieSceneDMXLibrarySection::Serialize(FArchive& Ar)
 				}
 			}
 		}
-
-		UpdateChannelProxy();
 	}
 }
 
+void UMovieSceneDMXLibrarySection::PostLoad()
+{
+	Super::PostLoad();
+
+	UpdateChannelProxy();
+}
+
 void UMovieSceneDMXLibrarySection::PostEditImport()
 {
 	Super::PostEditImport();
diff --git a/Engine/Plugins/VirtualProduction/DMX/DMXEngine/Source/DMXRuntime/Public/Sequencer/MovieSceneDMXLibrarySection.h b/Engine/Plugins/VirtualProduction/DMX/DMXEngine/Source/DMXRuntime/Public/Sequencer/MovieSceneDMXLibrarySection.h
--- a/Engine/Plugins/VirtualProduction/DMX/DMXEngine/Source/DMXRuntime/Public/Sequencer/MovieSceneDMXLibrarySection.h
+++ b/Engine/Plugins/VirtualProduction/DMX/DMXEngine/Source/DMXRuntime/Public/Sequencer/MovieSceneDMXLibrarySection.h
@@ -163,6 +163,7 @@ public:
 public:
 	// Begin UObject interface
 	virtual void Serialize(FArchive& Ar) override;
+	virtual void PostLoad() override;
 	virtual void PostEditImport() override;
 #if WITH_EDITOR
 	void PostEditChangeProperty(FPropertyChangedEvent& PropertyChangedEvent) override;
