Fix c++20 warnings

#rnx

diff --git a/Engine/Plugins/VirtualProduction/DMX/DMXEngine/Source/DMXEditor/Private/Widgets/SDMXChannel.cpp b/Engine/Plugins/VirtualProduction/DMX/DMXEngine/Source/DMXEditor/Private/Widgets/SDMXChannel.cpp
--- a/Engine/Plugins/VirtualProduction/DMX/DMXEngine/Source/DMXEditor/Private/Widgets/SDMXChannel.cpp
+++ b/Engine/Plugins/VirtualProduction/DMX/DMXEngine/Source/DMXEditor/Private/Widgets/SDMXChannel.cpp
@@ -177,7 +177,7 @@ FText SDMXChannel::GetValueText() const
 
 FSlateColor SDMXChannel::GetBackgroundColor() const
 {
-	const float CurrentPercent = static_cast<float>(Value) / DMX_MAX_CHANNEL_VALUE;
+	const float CurrentPercent = static_cast<float>(Value) / static_cast<float>(DMX_MAX_CHANNEL_VALUE);
 
 	// totally transparent when 0
 	if (CurrentPercent <= 0.0f)
diff --git a/Engine/Source/Editor/KismetWidgets/Public/SPinTypeSelector.h b/Engine/Source/Editor/KismetWidgets/Public/SPinTypeSelector.h
--- a/Engine/Source/Editor/KismetWidgets/Public/SPinTypeSelector.h
+++ b/Engine/Source/Editor/KismetWidgets/Public/SPinTypeSelector.h
@@ -58,7 +58,7 @@ struct FTopLevenshteinResult
 	ItemType Item;
 	float Score = INDEX_NONE;
 
-	bool IsSet() const { return Score != INDEX_NONE; }
+	bool IsSet() const { return Score != static_cast<float>(INDEX_NONE); }
 
 	void CompareAndUpdate(FStringView SearchValue, const ItemType& NewItem, FStringView NewItemValue)
 	{
