DMX - Fix opening a DMX Library asset dirties the asset

diff --git a/Engine/Plugins/VirtualProduction/DMX/DMXEngine/Source/DMXEditor/Private/Widgets/FixturePatch/SDMXMVRFixtureList.cpp b/Engine/Plugins/VirtualProduction/DMX/DMXEngine/Source/DMXEditor/Private/Widgets/FixturePatch/SDMXMVRFixtureList.cpp
--- a/Engine/Plugins/VirtualProduction/DMX/DMXEngine/Source/DMXEditor/Private/Widgets/FixturePatch/SDMXMVRFixtureList.cpp
+++ b/Engine/Plugins/VirtualProduction/DMX/DMXEngine/Source/DMXEditor/Private/Widgets/FixturePatch/SDMXMVRFixtureList.cpp
@@ -740,7 +740,8 @@ void SDMXMVRFixtureList::RefreshList()
 			if (!bUpdatedGeneralSceneDescription)
 			{	
 				// Ony update the General Scene Description when Nodes were never acquired, or new patches were added.
-				DMXLibrary->Modify();
+				constexpr bool bAlwaysMarkDirty = false;
+				DMXLibrary->Modify(bAlwaysMarkDirty);				
 				DMXLibrary->UpdateGeneralSceneDescription();
 				bUpdatedGeneralSceneDescription = true;
 			}
