DMX - Apply redesign of the Fixture Patch Editor

diff --git a/Engine/Plugins/VirtualProduction/DMX/DMXEngine/Source/DMXEditor/Private/Commands/DMXEditorCommands.cpp b/Engine/Plugins/VirtualProduction/DMX/DMXEngine/Source/DMXEditor/Private/Commands/DMXEditorCommands.cpp
--- a/Engine/Plugins/VirtualProduction/DMX/DMXEngine/Source/DMXEditor/Private/Commands/DMXEditorCommands.cpp
+++ b/Engine/Plugins/VirtualProduction/DMX/DMXEngine/Source/DMXEditor/Private/Commands/DMXEditorCommands.cpp
@@ -26,6 +26,10 @@ void FDMXEditorCommands::RegisterCommands()
 	UI_COMMAND(OpenPatchTool, "Open Patch Tool", "Open the patch tool - Useful to patch many fixtures at once.", EUserInterfaceActionType::Button, FInputChord());
 	UI_COMMAND(ToggleReceiveDMX, "Receive DMX", "Sets whether DMX is received in editor", EUserInterfaceActionType::ToggleButton, FInputChord(EModifierKey::Alt, EKeys::M));
 	UI_COMMAND(ToggleSendDMX, "Send DMX", "Sets whether DMX is sent from editor", EUserInterfaceActionType::ToggleButton, FInputChord(EModifierKey::Alt, EKeys::N));
+
+	UI_COMMAND(AutoAssignToFirstUniverseInSelection, "Auto Assign", "Auto assigns selected patches, doesn't decrement Universe of selection.", EUserInterfaceActionType::Button, FInputChord());
+	UI_COMMAND(AutoAssignToFirstReachableUniverse, "Auto Assign to first reachable Universe", "Auto assigns selected patches to the first Universe reachable by the DMX Library", EUserInterfaceActionType::Button, FInputChord());
+	UI_COMMAND(AutoAssignAfterLastAddressInLibrary, "Auto Assign after last Patch in Library", "Auto assigns selected patches after the last patch in the DMX Library", EUserInterfaceActionType::Button, FInputChord());
 }
 
 #undef LOCTEXT_NAMESPACE
diff --git a/Engine/Plugins/VirtualProduction/DMX/DMXEngine/Source/DMXEditor/Private/DMXEditorStyle.cpp b/Engine/Plugins/VirtualProduction/DMX/DMXEngine/Source/DMXEditor/Private/DMXEditorStyle.cpp
--- a/Engine/Plugins/VirtualProduction/DMX/DMXEngine/Source/DMXEditor/Private/DMXEditorStyle.cpp
+++ b/Engine/Plugins/VirtualProduction/DMX/DMXEngine/Source/DMXEditor/Private/DMXEditorStyle.cpp
@@ -6,6 +6,7 @@
 #include "Framework/Application/SlateApplication.h"
 #include "Interfaces/IPluginManager.h"
 #include "Styling/AppStyle.h"
+#include "Styling/SlateBrush.h"
 #include "Styling/SlateStyleRegistry.h"
 #include "Styling/SlateStyleMacros.h"
 #include "Styling/SlateTypes.h"
@@ -51,10 +52,10 @@ FDMXEditorStyle::FDMXEditorStyle()
 		Set("DMXEditor.RoundedPropertyBorder", new FSlateRoundedBoxBrush(FLinearColor::Black, 4.f));
 		Set("DMXEditor.WhiteRoundedPropertyBorder", new FSlateRoundedBoxBrush(FLinearColor::White, 2.f));
 
-		Set("FixturePatcher.FragmentBorder.Normal", new BORDER_BRUSH("FixturePatch_Border_Normal_8x", 2.f / 8.f));
-		Set("FixturePatcher.FragmentBorder.L", new BORDER_BRUSH("FixturePatch_Border_L_8x", 2.f / 8.f));
-		Set("FixturePatcher.FragmentBorder.R", new BORDER_BRUSH("FixturePatch_Border_R_8x", 2.f / 8.f));
-		Set("FixturePatcher.FragmentBorder.TB", new BORDER_BRUSH("FixturePatch_Border_TB_8x", 2.f / 8.f));
+		Set("FixturePatcher.FragmentBorder.Normal", new BORDER_BRUSH("FixturePatch_Border_Full_16x", 2.f / 8.f));
+		Set("FixturePatcher.FragmentBorder.L", new BORDER_BRUSH("FixturePatch_Border_L_16x", 2.f / 8.f));
+		Set("FixturePatcher.FragmentBorder.R", new BORDER_BRUSH("FixturePatch_Border_R_16x", 2.f / 8.f));
+		Set("FixturePatcher.FragmentBorder.TB", new BORDER_BRUSH("FixturePatch_Border_M_16x", 2.f / 8.f));
 	}
 
 	// Fonts
@@ -101,6 +102,8 @@ FDMXEditorStyle::FDMXEditorStyle()
 		Set("Icons.PatchTool", new IMAGE_BRUSH_SVG("PatchTool_16", Icon16x16));
 		Set("Icons.ReceiveDMX", new IMAGE_BRUSH_SVG("ToggleReceiveDMX_16", Icon16x16));
 		Set("Icons.SendDMX", new IMAGE_BRUSH_SVG("ToggleSendDMX_16", Icon16x16));
+
+		Set("Icons.WarningExclamationMark", new CORE_IMAGE_BRUSH("Old/Tiles/Outer/alertSolid", Icon16x16, FStyleColors::Warning));
 	}
 	
 	// Distribution Grid buttons
@@ -161,6 +164,8 @@ FDMXEditorStyle::FDMXEditorStyle()
 			.SetActiveHoveredBrush(CORE_IMAGE_BRUSH("Common/Selection", Icon8x8, FLinearColor(.3f, .3f, .3f, .8f)))
 			.SetSelectorFocusedBrush(CORE_BORDER_BRUSH("Common/Selector", FMargin(4.f / 16.f), FLinearColor(1.f, 1.f, 1.f, .7f)))
 		);
+
+		Set("FixturePatchFragment.ConflictBackground", new IMAGE_BRUSH("FixturePatch_ConflictBack_36x", FVector2D(36.f, 36.f), FStyleColors::Warning, ESlateBrushTileType::Both));
 	}
 
 	FSlateStyleRegistry::RegisterSlateStyle(*this);
diff --git a/Engine/Plugins/VirtualProduction/DMX/DMXEngine/Source/DMXEditor/Private/DMXFixturePatchAutoAssignUtility.cpp b/Engine/Plugins/VirtualProduction/DMX/DMXEngine/Source/DMXEditor/Private/DMXFixturePatchAutoAssignUtility.cpp
new file mode 100644
--- /dev/null
+++ b/Engine/Plugins/VirtualProduction/DMX/DMXEngine/Source/DMXEditor/Private/DMXFixturePatchAutoAssignUtility.cpp
@@ -0,0 +1,397 @@
+// Copyright Epic Games, Inc. All Rights Reserved.
+
+#include "DMXFixturePatchAutoAssignUtility.h"
+
+#include "DMXProtocolConstants.h"
+#include "IO/DMXInputPort.h"
+#include "IO/DMXOutputPort.h"
+#include "Library/DMXEntityFixturePatch.h"
+#include "Library/DMXLibrary.h"
+
+#include "ScopedTransaction.h"
+#include "Math/Range.h"
+
+#define LOCTEXT_NAMESPACE "DMXFixturePatchAutoAssignUtility"
+
+
+namespace UE::DMX::DMXFixturePatchAutoAssignUtility::Private
+{
+	/** Returns true if two patches form a stack (same absolute address, same fixture type, same mode) */
+	bool DoPatchesFormStack(UDMXEntityFixturePatch* First, UDMXEntityFixturePatch* Second)
+	{
+		if (!First || !Second)
+		{
+			return false;
+		}
+
+		uint64 FirstAbsoluteStartingChannel = ((uint64)First->GetUniverseID() * DMX_UNIVERSE_SIZE) + First->GetStartingChannel();
+		uint64 SecondAbsoluteStartingChannel = ((uint64)Second->GetUniverseID() * DMX_UNIVERSE_SIZE) + Second->GetStartingChannel();
+		if (FirstAbsoluteStartingChannel != SecondAbsoluteStartingChannel)
+		{
+			return false;
+		}
+
+		if (First->GetFixtureType() != Second->GetFixtureType())
+		{
+			return false;
+		}
+
+		if (First->GetActiveModeIndex() != Second->GetActiveModeIndex())
+		{
+			return false;
+		}
+
+		return true;
+	}
+
+	/** Builds arrays of stacked patches. A stack are patches of same fixture type, mode and absolute starting channel. */
+	TArray<TArray<UDMXEntityFixturePatch*>> BuildPatchStacks(TArray<UDMXEntityFixturePatch*> FixturePatches)
+	{
+		TArray<TArray<UDMXEntityFixturePatch*>> PatchStacks;
+
+		Algo::SortBy(FixturePatches, [](const UDMXEntityFixturePatch* FixturePatch)
+			{
+				return ((uint64)FixturePatch->GetUniverseID() * DMX_UNIVERSE_SIZE) + FixturePatch->GetStartingChannel();
+			});
+
+		TArray<UDMXEntityFixturePatch*> Stack;
+		for (UDMXEntityFixturePatch* Patch : FixturePatches)
+		{
+			if (!Stack.IsEmpty() && !DoPatchesFormStack(Stack[0], Patch))
+			{
+				PatchStacks.Add(Stack);
+				Stack.Reset();
+			}
+			
+			Stack.Add(Patch);
+		}
+
+		// Add the last stack of patches
+		PatchStacks.Add(Stack);
+		
+		return PatchStacks;
+	}
+
+	/** Finds patches in the DMX Library that are not contained in the FixturePatches array. Returns an ordered array (by addresses). */
+	TArray<UDMXEntityFixturePatch*> FindOtherPatchesInLibrary(UDMXLibrary* DMXLibrary, const TArray<UDMXEntityFixturePatch*>& FixturePatches)
+	{
+		check(DMXLibrary);
+		TArray<UDMXEntityFixturePatch*> Result = DMXLibrary->GetEntitiesTypeCast<UDMXEntityFixturePatch>();
+		Result.RemoveAll([FixturePatches](const UDMXEntityFixturePatch* FixturePatch)
+			{
+				return
+					!FixturePatch ||
+					!FixturePatch->GetFixtureType() ||
+					FixturePatches.Contains(FixturePatch);
+			});
+
+
+		Algo::SortBy(Result, [](UDMXEntityFixturePatch* FixturePatch)
+			{
+				return ((uint64)FixturePatch->GetUniverseID() * DMX_UNIVERSE_SIZE) + FixturePatch->GetStartingChannel();
+			});
+
+		return Result;
+	}
+}
+
+void FDMXFixturePatchAutoAssignUtility::AutoAssign(const TArray<UDMXEntityFixturePatch*>& FixturePatches, EDMXFixturePatchAutoAssignMode AutoAssignMode)
+{
+	FDMXFixturePatchAutoAssignUtility(FixturePatches, AutoAssignMode);
+}
+
+FDMXFixturePatchAutoAssignUtility::FDMXFixturePatchAutoAssignUtility(const TArray<UDMXEntityFixturePatch*>& FixturePatches, EDMXFixturePatchAutoAssignMode AutoAssignMode)
+{
+	if (FixturePatches.IsEmpty())
+	{
+		return;
+	}
+
+	// Ensure all patches reside in the same library
+	DMXLibrary = FixturePatches[0]->GetParentLibrary();
+	if (!ensureMsgf(DMXLibrary, TEXT("Trying to auto assign fixture patches but the DMX Library of its patches is not valid.")))
+	{
+		return;
+	}
+	check(DMXLibrary); // Assume a valid library for this operation from hereon
+
+	for (const UDMXEntityFixturePatch* FixturePatch : FixturePatches)
+	{
+		if (!ensureMsgf(FixturePatch->GetParentLibrary() == DMXLibrary, TEXT("Trying to auto assign fixture patches, but the patches don't share a common DMX Library. This is not supported")))
+		{
+			return;
+		}
+	}
+
+	// Initialize
+	using namespace UE::DMX::DMXFixturePatchAutoAssignUtility::Private;
+	OtherPatches = FindOtherPatchesInLibrary(DMXLibrary, FixturePatches);
+	AutoAssignPatchStacks = BuildPatchStacks(FixturePatches);
+
+	// Auto assign as per mode
+	switch (AutoAssignMode)
+	{
+	case EDMXFixturePatchAutoAssignMode::FirstUniverseInSelection:
+		AssignToFirstUniverseInSelection();
+		break;
+
+	case EDMXFixturePatchAutoAssignMode::FirstReachableUniverse:
+		AssignToFirstReachableUniverse();
+		break;
+
+	case EDMXFixturePatchAutoAssignMode::LastAddressInLibrary:
+		AutoAssignAfterLastAddressInLibrary();
+		break;
+
+	default:
+		checkf(0, TEXT("Trying to use unhandled auto assign mode in FDMXFixturePatchAutoAssignUtility"));
+	}
+}
+
+void FDMXFixturePatchAutoAssignUtility::AssignToFirstUniverseInSelection()
+{
+	if (!ensureMsgf(!AutoAssignPatchStacks.IsEmpty(), TEXT("Cannot auto assign patches without patches. Such cases should have been handled previously.")))
+	{
+		return;
+	}
+
+	if (!ensureMsgf(!AutoAssignPatchStacks[0].IsEmpty() || !AutoAssignPatchStacks[0][0], TEXT("No patch or invalid patch in patch stack when trying to auto assign..")))
+	{
+		return;
+	}
+
+	const uint64 MinRequiredSize = ComputeMinRequiredSize();
+
+	UDMXEntityFixturePatch* FirstPachToAutoAssign = AutoAssignPatchStacks[0][0];
+	const uint64 FirstReachableAbsoluteAddress = FirstPachToAutoAssign->GetUniverseID() * DMX_UNIVERSE_SIZE + 1;
+	const TArray<TRange<uint64>> FreeRanges = ComputeFreeRanges(FirstReachableAbsoluteAddress);
+	for (const TRange<uint64>& FreeRange : FreeRanges)
+	{
+		const uint64 SizeOfRange = FreeRange.Size<uint64>();
+		if (SizeOfRange >= MinRequiredSize)
+		{
+			if (CanAutoAssignToRange(FreeRange))
+			{
+				PerformAutoAssign(FreeRange.GetLowerBoundValue());
+				return;
+			}
+		}
+	}
+}
+
+void FDMXFixturePatchAutoAssignUtility::AssignToFirstReachableUniverse()
+{
+	if (!ensureMsgf(CanReachAnyUniverse(), TEXT("If a DMX Library has no reachable universe, options to auto-assign to first reachable universe should be disabled. Please use CanReachAnyUniverse() to disable such options.")))
+	{
+		return;
+	}
+
+	const uint64 MinRequiredSize = ComputeMinRequiredSize();
+
+	const uint64 FirstReachableAbsoluteAddress = GetFirstReachableUniverse() * DMX_UNIVERSE_SIZE + 1;
+	const TArray<TRange<uint64>> FreeRanges = ComputeFreeRanges(FirstReachableAbsoluteAddress);
+	for (const TRange<uint64>& FreeRange : FreeRanges)
+	{
+		const uint64 SizeOfRange = FreeRange.Size<uint64>();
+		if (SizeOfRange >= MinRequiredSize)
+		{
+			if (CanAutoAssignToRange(FreeRange))
+			{
+				PerformAutoAssign(FreeRange.GetLowerBoundValue());
+				return;
+			}
+		}
+	}
+}
+
+void FDMXFixturePatchAutoAssignUtility::AutoAssignAfterLastAddressInLibrary()
+{
+	if (!ensureMsgf(!AutoAssignPatchStacks.IsEmpty(), TEXT("Cannot auto assign patches without patches. Such cases should have been handled previously.")))
+	{
+		return;
+	}
+
+	if (!ensureMsgf(!AutoAssignPatchStacks[0].IsEmpty() || !AutoAssignPatchStacks[0][0], TEXT("Cannot auto assign. No patch or invalid patch in patch stack when trying to auto assign..")))
+	{
+		return;
+	}
+
+	if (!ensureMsgf(!OtherPatches.IsEmpty() && OtherPatches[0], TEXT("Cannot auto assign. Invalid 'other' patch when trying to auto assign to last address in DMX Library")))
+	{
+		return;
+	}
+
+	const uint64 MinRequiredSize = ComputeMinRequiredSize();
+
+	const uint64 AbsoluteAddressAfterLastPatchInLibrary = [this]()
+	{
+		if (OtherPatches.IsEmpty())
+		{
+			UDMXEntityFixturePatch* FirstPachToAutoAssign = AutoAssignPatchStacks[0][0];
+			return FirstPachToAutoAssign->GetUniverseID() * DMX_UNIVERSE_SIZE + 1;
+		}
+		else
+		{
+			UDMXEntityFixturePatch* LastPatchInLibrary = OtherPatches.Last();
+			return LastPatchInLibrary->GetUniverseID() * DMX_UNIVERSE_SIZE + 1;
+		}
+	}();
+	const TArray<TRange<uint64>> FreeRanges = ComputeFreeRanges(AbsoluteAddressAfterLastPatchInLibrary);
+	for (const TRange<uint64>& FreeRange : FreeRanges)
+	{
+		const uint64 SizeOfRange = FreeRange.Size<uint64>();
+		if (SizeOfRange >= MinRequiredSize)
+		{
+			if (CanAutoAssignToRange(FreeRange))
+			{
+				PerformAutoAssign(FreeRange.GetLowerBoundValue());
+				return;
+			}
+		}
+	}
+}
+
+bool FDMXFixturePatchAutoAssignUtility::CanAutoAssignToRange(const TRange<uint64>& Range) const
+{
+	const uint64 SizeOfRange = Range.Size<uint64>();
+
+	// For simplicity work with a 0 - 511 range
+	uint64 RelativeStartingChannel = (Range.GetLowerBoundValue() - 1) % DMX_UNIVERSE_SIZE;
+	uint64 PatchedChannels = 0;
+	uint64 VoidChannels = 0; // Channels that take space but cannot hold a patch
+	for (const TArray<UDMXEntityFixturePatch*>& PatchStack : AutoAssignPatchStacks)
+	{
+		check(!PatchStack.IsEmpty() && PatchStack[0]);
+
+		const int32 ChannelSpan = PatchStack[0]->GetChannelSpan();
+		if (RelativeStartingChannel + ChannelSpan > DMX_UNIVERSE_SIZE)
+		{
+			VoidChannels += DMX_UNIVERSE_SIZE - RelativeStartingChannel;
+			RelativeStartingChannel = 0;
+		}
+
+		if (PatchedChannels + VoidChannels + ChannelSpan > SizeOfRange)
+		{
+			return false;
+		}
+
+		PatchedChannels += ChannelSpan;
+	}
+
+	return true;
+}
+
+void FDMXFixturePatchAutoAssignUtility::PerformAutoAssign(uint64 AbsoluteStartingChannel)
+{
+	const FScopedTransaction AutoAssignTransaction(NSLOCTEXT("DMXFixturePatchAutoAssignUtility", "AutoAssignTransaction", "Auto assign fixture patches"));
+
+	for (const TArray<UDMXEntityFixturePatch*>& PatchStack : AutoAssignPatchStacks)
+	{
+		check(!PatchStack.IsEmpty() && PatchStack[0]);
+
+		int32 RelativeStartingChannel = AbsoluteStartingChannel % DMX_UNIVERSE_SIZE;
+		int32 Universe = AbsoluteStartingChannel / DMX_UNIVERSE_SIZE;
+
+		const int32 ChannelSpan = PatchStack[0]->GetChannelSpan();
+		if (RelativeStartingChannel + ChannelSpan - 1 > DMX_UNIVERSE_SIZE)
+		{
+			Universe++;
+			RelativeStartingChannel = 1;
+		}
+
+		for (UDMXEntityFixturePatch* Patch : PatchStack)
+		{
+			Patch->PreEditChange(UDMXEntityFixturePatch::StaticClass()->FindPropertyByName(UDMXEntityFixturePatch::GetStartingChannelPropertyNameChecked()));
+
+			Patch->SetUniverseID(Universe);
+			Patch->SetStartingChannel(RelativeStartingChannel);
+
+			Patch->PostEditChange();
+		}
+
+		AbsoluteStartingChannel = PatchStack[0]->GetUniverseID() * DMX_MAX_ADDRESS + PatchStack[0]->GetStartingChannel() + ChannelSpan;
+	}
+}
+
+uint64 FDMXFixturePatchAutoAssignUtility::ComputeMinRequiredSize() const
+{
+	uint64 MinRequiredSize = 0;
+	for (const TArray<UDMXEntityFixturePatch*>& PatchStack : AutoAssignPatchStacks)
+	{
+		check(!PatchStack.IsEmpty() && PatchStack[0]);
+		MinRequiredSize += PatchStack[0]->GetChannelSpan();
+	}
+
+	return MinRequiredSize;
+}
+
+TArray<TRange<uint64>> FDMXFixturePatchAutoAssignUtility::ComputeFreeRanges(uint64 FirstAbsoluteAddress) const
+{
+	TArray<TRange<uint64>> FreeRanges;
+	if (OtherPatches.IsEmpty())
+	{
+		FreeRanges.Add(TRange<uint64>(0, TNumericLimits<uint64>::Max()));
+		return FreeRanges;
+	}
+
+	// Add the range from the first channel to the first patch, only if the first patch is not on channel 1
+	const uint64 FirstAbsoulteStartingChannel = ((uint64)OtherPatches[0]->GetUniverseID() * DMX_UNIVERSE_SIZE) + OtherPatches[0]->GetStartingChannel();
+	if (FirstAbsoulteStartingChannel > FirstAbsoluteAddress)
+	{
+		FreeRanges.Add(TRange<uint64>(FirstAbsoluteAddress, FirstAbsoulteStartingChannel));
+	}
+
+	// Add ranges in between patches
+	const UDMXEntityFixturePatch* PreviousPatch = nullptr;
+	for (UDMXEntityFixturePatch* Patch : OtherPatches)
+	{
+		if (PreviousPatch)
+		{
+			const uint64 FreeRangeStart = ((uint64)PreviousPatch->GetUniverseID() * DMX_UNIVERSE_SIZE) + PreviousPatch->GetStartingChannel() + PreviousPatch->GetChannelSpan();
+			const uint64 FreeRangeEnd = ((uint64)Patch->GetUniverseID() * DMX_UNIVERSE_SIZE) + Patch->GetStartingChannel();
+			
+			if (FreeRangeStart >= FirstAbsoluteAddress && FreeRangeStart < FreeRangeEnd)
+			{
+				FreeRanges.Add(TRange<uint64>(FreeRangeStart, FreeRangeEnd));
+			}
+		}
+		PreviousPatch = Patch;
+	}
+
+	// Add the range from the last patch to uint64 max. This always can be added as patches cannot span up to uint64 max.
+	const uint64 LastAbsoulteEndingChannel = ((uint64)OtherPatches.Last()->GetUniverseID() * DMX_UNIVERSE_SIZE) + OtherPatches.Last()->GetStartingChannel() + OtherPatches.Last()->GetChannelSpan();
+	FreeRanges.Add(TRange<uint64>(LastAbsoulteEndingChannel, TNumericLimits<uint64>::Max()));
+
+	return FreeRanges;
+}
+
+bool FDMXFixturePatchAutoAssignUtility::CanReachAnyUniverse() const
+{
+	check(DMXLibrary);
+
+	return !(DMXLibrary->GetInputPorts().IsEmpty() && DMXLibrary->GetOutputPorts().IsEmpty());
+}
+
+int32 FDMXFixturePatchAutoAssignUtility::GetFirstReachableUniverse() const
+{
+	check(DMXLibrary);
+	
+	if (!ensureMsgf(CanReachAnyUniverse(), TEXT("If a DMX Library has no reachable universe, options to auto-assign to first reachable universe should be disabled. Please use CanReachAnyUniverse() to disable such options.")))
+	{
+		return 1;
+	}
+
+	int32 Result = TNumericLimits<int32>::Max();
+	for (const FDMXInputPortSharedRef& InputPort : DMXLibrary->GetInputPorts())
+	{
+		Result = FMath::Min(Result, InputPort->GetLocalUniverseStart());
+	}
+	for (const FDMXOutputPortSharedRef& OutputPort : DMXLibrary->GetOutputPorts())
+	{
+		Result = FMath::Min(Result, OutputPort->GetLocalUniverseStart());
+	}
+
+	return Result;
+}
+
+#undef LOCTEXT_NAMESPACE
diff --git a/Engine/Plugins/VirtualProduction/DMX/DMXEngine/Source/DMXEditor/Private/DMXFixturePatchAutoAssignUtility.h b/Engine/Plugins/VirtualProduction/DMX/DMXEngine/Source/DMXEditor/Private/DMXFixturePatchAutoAssignUtility.h
new file mode 100644
--- /dev/null
+++ b/Engine/Plugins/VirtualProduction/DMX/DMXEngine/Source/DMXEditor/Private/DMXFixturePatchAutoAssignUtility.h
@@ -0,0 +1,78 @@
+// Copyright Epic Games, Inc. All Rights Reserved.
+
+#pragma once
+
+#include "Containers/Array.h"
+#include "HAL/Platform.h"
+
+class UDMXEntityFixturePatch;
+class UDMXLibrary;
+
+
+/** Available auto assign modes enumerated */
+enum class EDMXFixturePatchAutoAssignMode : uint8
+{
+	/** Auto assign to the first universe in selection */
+	FirstUniverseInSelection,
+
+	/** Auto assign to the first universe reachable for the DMX Library */
+	FirstReachableUniverse,
+
+	/** Auto assign after the address of the last fixture patch in the library */
+	LastAddressInLibrary
+};
+
+/**
+ * Utility to auto assign fixture patches.
+ */
+class FDMXFixturePatchAutoAssignUtility
+{
+public:
+	/** 
+	 * Auto assigns FixturePatches. If two ore more patches in selection reside on the same address, 
+	 * they're auto assigned in order as they appear in the array.
+	 * 
+	 * Note, all Patches are expected to reside in the same DMX Library.
+	 */
+	static void AutoAssign(const TArray<UDMXEntityFixturePatch*>& FixturePatches, EDMXFixturePatchAutoAssignMode AutoAssignMode);
+
+private:
+	/** Private constructor. Use public FDMXFixturePatchAutoAssignUtility::AutoAssign instead. */
+	FDMXFixturePatchAutoAssignUtility(const TArray<UDMXEntityFixturePatch*>& FixturePatches, EDMXFixturePatchAutoAssignMode AutoAssignMode);
+
+	/** Auto assigns to the universe in the patches selected to be auto assigned */
+	void AssignToFirstUniverseInSelection();
+
+	/** Auto assigns to the first reachable universe in the library */
+	void AssignToFirstReachableUniverse();
+
+	/** Auto assigns after the last patch in the DMX Library */
+	void AutoAssignAfterLastAddressInLibrary();
+
+	/** Returns true if patches can be auto assigned to given range */
+	bool CanAutoAssignToRange(const TRange<uint64>& Range) const;
+
+	/** Performs the auto assign operation to AutoAssignPatchStacks */
+	void PerformAutoAssign(uint64 AbsoluteStartingChannel);
+
+	/** Computes the min size required to fit all patches. Note this may not be a sufficient size, since patches cannot span multiple channels. */
+	uint64 ComputeMinRequiredSize() const;
+
+	/** Returns an array of free ranges in this dmx library, using absolute addresses. */
+	TArray<TRange<uint64>> ComputeFreeRanges(uint64 FirstAbsoluteAddress) const;
+	
+	/** Returns true if any universe in the DMX Library can be reached */
+	bool CanReachAnyUniverse() const;
+
+	/** Returns the first reachable universe in which this auto assign operation occurs */
+	int32 GetFirstReachableUniverse() const;
+
+	/** The DMX Library in which the auto assign operation occurs */
+	UDMXLibrary* DMXLibrary = nullptr;
+
+	/** Patches to auto assign, as arrays of valid stacks. A stack are patches of same fixture type, mode and absolute starting channel */
+	TArray<TArray<UDMXEntityFixturePatch*>> AutoAssignPatchStacks;
+
+	/** Other patches in library that do not need to be auto assigned */
+	TArray<UDMXEntityFixturePatch*> OtherPatches;
+};
diff --git a/Engine/Plugins/VirtualProduction/DMX/DMXEngine/Source/DMXEditor/Private/Widgets/FixturePatch/DMXFixturePatchNode.cpp b/Engine/Plugins/VirtualProduction/DMX/DMXEngine/Source/DMXEditor/Private/Widgets/FixturePatch/DMXFixturePatchNode.cpp
--- a/Engine/Plugins/VirtualProduction/DMX/DMXEngine/Source/DMXEditor/Private/Widgets/FixturePatch/DMXFixturePatchNode.cpp
+++ b/Engine/Plugins/VirtualProduction/DMX/DMXEngine/Source/DMXEditor/Private/Widgets/FixturePatch/DMXFixturePatchNode.cpp
@@ -9,6 +9,7 @@
 #include "SDMXPatchedUniverse.h"
 #include "Library/DMXEntityFixturePatch.h"
 #include "Library/DMXLibrary.h"
+#include "MVR/Types/DMXMVRFixtureNode.h"
 
 #include "ScopedTransaction.h"
 #include "Widgets/Layout/SGridPanel.h"
@@ -22,7 +23,7 @@ class FDMXFixturePatchFragment
 {
 public:
 	/** Creates fragments of a patch spread across a grid */
-	static void CreateFragments(TWeakObjectPtr<UDMXEntityFixturePatch> InFixturePatch, int32 ChannelID, TArray<TSharedPtr<FDMXFixturePatchFragment>>& OutFragments)
+	static void CreateFragments(TWeakObjectPtr<UDMXEntityFixturePatch> InFixturePatch, TArray<TSharedPtr<FDMXFixturePatchFragment>>& OutFragments)
 	{
 		OutFragments.Reset();
 
@@ -72,10 +73,12 @@ public:
 				NewFragment->Row = Row;
 				NewFragment->ColumnSpan = ColumnSpan;
 				NewFragment->bIsConflict = NumConflicts > 0;
-				NewFragment->bIsText = false;
 
 				OutFragments.Add(NewFragment);
 
+				// Link lhs
+				NewFragment->LhsFragment = PrevFragment;
+
 				// Link rhs
 				if (PrevFragment.IsValid())
 				{
@@ -88,73 +91,8 @@ public:
 				PrevFragment = NewFragment;
 			} while (ChannelSpan > 0);
 		}
-
-		// Create a text fragment for the first range, from the start to the first increment of conflicts
-		int32 PreviousNumConflicts = TNumericLimits<int32>::Max();
-		int32 TextStartingIndex = INDEX_NONE;
-		int32 TextEndingIndex = INDEX_NONE;
-		for (const TTuple<TRange<int32>, int32>& RangeToNumConflictsPair : RangeToNumConflictsMap)
-		{
-			if (TextStartingIndex == INDEX_NONE)
-			{
-				TextStartingIndex = RangeToNumConflictsPair.Key.GetLowerBoundValue();
-			}
-
-			const int32 NumConflicts = RangeToNumConflictsPair.Value;
-			if (PreviousNumConflicts <= NumConflicts)
-			{
-				break;
-			}
-
-			TextEndingIndex = RangeToNumConflictsPair.Key.GetUpperBoundValue() - 1;
-			PreviousNumConflicts = NumConflicts;
-		}
-
-		TSharedPtr<FDMXFixturePatchFragment> NewTextFragment = MakeShared<FDMXFixturePatchFragment>();
-
-		const int32 Row = TextStartingIndex / FDMXChannelGridSpecs::NumColumns;
-		const int32 Column = TextStartingIndex % FDMXChannelGridSpecs::NumColumns;
-		int32 ChannelSpan = TextEndingIndex + 1 - TextStartingIndex;
-		int32 ColumnSpan = Column + ChannelSpan < FDMXChannelGridSpecs::NumColumns ? ChannelSpan : FDMXChannelGridSpecs::NumColumns - Column;
-
-		NewTextFragment->ChannelIndex = TextStartingIndex;
-		NewTextFragment->Column = Column;
-		NewTextFragment->Row = Row;
-		NewTextFragment->ColumnSpan = ColumnSpan;
-		NewTextFragment->bIsConflict = false;
-		NewTextFragment->bIsText = true;
-		OutFragments.Add(NewTextFragment);
-
-		// Link lhs
-		TSharedPtr<FDMXFixturePatchFragment> LhsLink = OutFragments[0];
-		for (int IndexOfFragment = 1; IndexOfFragment < OutFragments.Num(); IndexOfFragment++)
-		{
-			OutFragments[IndexOfFragment]->LhsFragment = LhsLink;
-			LhsLink = OutFragments[IndexOfFragment];
-		}
 	}
 
-	/** Gets all fragments of a patch */
-	TArray<TSharedPtr<FDMXFixturePatchFragment>> GetFragments()
-	{
-		TArray<TSharedPtr<FDMXFixturePatchFragment>> Fragments;
-		TSharedPtr<FDMXFixturePatchFragment> Root = SharedThis(this);
-		while (Root->LhsFragment.IsValid())
-		{
-			Root = LhsFragment;
-		}
-
-		TSharedPtr<FDMXFixturePatchFragment> Fragment = Root;
-		while (Fragment.IsValid())
-		{
-			Fragments.Add(Fragment);
-			Fragment = Fragment->RhsFragment;
-		}
-
-		return Fragments;
-	}
-
-
 	bool IsHead() const { return !LhsFragment.IsValid(); }
 	bool IsTail() const { return !RhsFragment.IsValid(); }
 
@@ -173,9 +111,6 @@ public:
 	/** True if this fragment conflicts with others */
 	bool bIsConflict = false;
 
-	/** True if this fragment is the text fragment */
-	bool bIsText = false;
-
 	/** The fragment to the left of this fragment */
 	TSharedPtr<FDMXFixturePatchFragment> LhsFragment;
 
@@ -259,7 +194,11 @@ private:
 		}();
 
 		// Find the conflicts
-		TArray<const UDMXEntityFixturePatch*> ConflictingFixturePatches = [FixturePatch, &OtherFixturePatches]()
+		const int32 MyChannelIndex = FixturePatch->GetStartingChannel() - 1;
+		const int32 MyChannelSpan = FixturePatch->GetChannelSpan();
+		const TRange<int32> MyRange(MyChannelIndex, MyChannelIndex + MyChannelSpan);
+
+		TArray<const UDMXEntityFixturePatch*> ConflictingFixturePatches = [FixturePatch, &OtherFixturePatches, &MyRange]()
 		{
 			TArray<const UDMXEntityFixturePatch*> Result;
 			for (const UDMXEntityFixturePatch* Other : OtherFixturePatches)
@@ -286,12 +225,10 @@ private:
 				const int32 OtherChannelSpan = Other->GetChannelSpan();
 				const TRange<int32> OtherRange(OtherChannelIndex, OtherChannelIndex + OtherChannelSpan);
 
-				if (!OtherRange.Overlaps(OtherRange))
+				if (OtherRange.Overlaps(MyRange))
 				{
-					continue;
+					Result.Add(Other);
 				}
-
-				Result.Add(Other);
 			}
 			return Result;
 		}();
@@ -307,7 +244,7 @@ TSharedPtr<FDMXFixturePatchNode> FDMXFixturePatchNode::Create(TWeakPtr<FDMXEdito
 		return nullptr;
 	}
 
-	TSharedPtr<FDMXFixturePatchNode> NewNode = MakeShared<FDMXFixturePatchNode>();
+	const TSharedRef<FDMXFixturePatchNode> NewNode = MakeShared<FDMXFixturePatchNode>();
 	NewNode->WeakDMXEditor = InDMXEditor;
 	NewNode->FixturePatch= InFixturePatch;
 
@@ -318,20 +255,12 @@ TSharedPtr<FDMXFixturePatchNode> FDMXFixturePatchNode::Create(TWeakPtr<FDMXEdito
 	}
 
 	TSharedPtr<FDMXFixturePatchSharedData> SharedData = DMXEditor->GetFixturePatchSharedData();
-	if (!SharedData.IsValid())
+	if (SharedData.IsValid())
 	{
-		return nullptr;
+		NewNode->UpdateIsSelected();
+		SharedData->OnFixturePatchSelectionChanged.AddSP(NewNode, &FDMXFixturePatchNode::UpdateIsSelected);
 	}
 
-	// Select the new node if it is selected in Fixture Patch Shared Data
-	if (SharedData->GetSelectedFixturePatches().Contains(InFixturePatch))
-	{
-		NewNode->bSelected = true;
-	}
-
-	// Listen to selection changes 
-	SharedData->OnFixturePatchSelectionChanged.AddSP(NewNode.ToSharedRef(), &FDMXFixturePatchNode::OnSelectionChanged);
-
 	return NewNode;
 }
 
@@ -342,43 +271,32 @@ void FDMXFixturePatchNode::SetAddresses(int32 NewUniverseID, int32 NewStartingCh
 		return;
 	}
 
-	UniverseID = NewUniverseID;
-	StartingChannel = NewStartingChannel;
-	ChannelSpan = NewChannelSpan;
-
 	FixturePatch->SetUniverseID(NewUniverseID);
-	FixturePatch->SetStartingChannel(StartingChannel);
+	FixturePatch->SetStartingChannel(NewStartingChannel);
 }
 
-TArray<TSharedPtr<SDMXFixturePatchFragment>> FDMXFixturePatchNode::GenerateWidgets(const TArray<TSharedPtr<FDMXFixturePatchNode>>& FixturePatchNodeGroup)
+TArray<TSharedRef<SDMXFixturePatchFragment>> FDMXFixturePatchNode::GenerateWidgets(const TSharedRef<SDMXPatchedUniverse>& OwningUniverse, const TArray<TSharedPtr<FDMXFixturePatchNode>>& FixturePatchNodeGroup)
 {
-
-	// Fragment the patch to fit the Addresses
+	// Fragment the patch to fit the grid
 	TArray<TSharedPtr<FDMXFixturePatchFragment>> NewFragments;
-	FDMXFixturePatchFragment::CreateFragments(FixturePatch, StartingChannel, NewFragments);
+	FDMXFixturePatchFragment::CreateFragments(FixturePatch, NewFragments);
 
-	FixturePatchFragmentWidgets.Reset();
+	TArray<TSharedRef<SDMXFixturePatchFragment>> Result;
 	for (const TSharedPtr<FDMXFixturePatchFragment>& Fragment : NewFragments)
 	{
-		const bool bIsHead = !Fragment->LhsFragment.IsValid();
-		const bool bIsTail = !Fragment->RhsFragment.IsValid();
-
-		const FText DisplayName = FixturePatch.IsValid() ? FText::FromString(FixturePatch->GetDisplayName()) : LOCTEXT("InvalidFixturePatchName", "Invalid Fixture Patch");
-		FixturePatchFragmentWidgets.Add(
+		Result.Add(
 			SNew(SDMXFixturePatchFragment, SharedThis(this), FixturePatchNodeGroup)
-			.DMXEditor(WeakDMXEditor)
-			.IsHead(bIsHead)
-			.IsTail(bIsTail)
-			.IsText(Fragment->bIsText)
-			.IsConflicting(Fragment->bIsConflict)
+			.bIsHead(Fragment->IsHead())
+			.bIsTail(Fragment->IsTail())
+			.bIsConflicting(Fragment->bIsConflict)
+			.StartingChannel(Fragment->ChannelIndex + 1)
 			.Column(Fragment->Column)
 			.Row(Fragment->Row)
 			.ColumnSpan(Fragment->ColumnSpan)
-			.bHighlight(bSelected)
 		);
 	}
 
-	return FixturePatchFragmentWidgets;
+	return Result;
 }
 
 bool FDMXFixturePatchNode::OccupiesChannels(int32 Channel, int32 Span) const
@@ -391,14 +309,39 @@ bool FDMXFixturePatchNode::OccupiesChannels(int32 Channel, int32 Span) const
 	return true;
 }
 
-bool FDMXFixturePatchNode::IsPatched() const
+bool FDMXFixturePatchNode::IsSelected() const
 {
-	return UniverseWidget.IsValid();
+	return bSelected;
 }
 
-bool FDMXFixturePatchNode::IsSelected() const
+FString FDMXFixturePatchNode::GetFixtureID() const
 {
-	return bSelected;
+	if (!MVRFixtureNode.IsValid())
+	{
+		UDMXLibrary* DMXLibrary = WeakDMXEditor.IsValid() ? WeakDMXEditor.Pin()->GetDMXLibrary() : nullptr;
+		if (!DMXLibrary || !FixturePatch.IsValid())
+		{
+			return FString();
+		}
+
+		TArray<UDMXMVRFixtureNode*> MVRFixtureNodes;
+		DMXLibrary->GetLazyGeneralSceneDescription()->GetFixtureNodes(MVRFixtureNodes);
+		UDMXMVRFixtureNode* const* MVRFixtureNodePtr = Algo::FindByPredicate(MVRFixtureNodes, [this](const UDMXMVRFixtureNode* Node)
+			{
+				return Node->UUID == FixturePatch->GetMVRFixtureUUID();
+			});
+		if (MVRFixtureNodePtr)
+		{
+			MVRFixtureNode = TWeakObjectPtr<UDMXMVRFixtureNode>(*MVRFixtureNodePtr);
+		}
+	}
+
+	if (MVRFixtureNode.IsValid())
+	{
+		return MVRFixtureNode->FixtureID;
+	}
+
+	return FString();
 }
 
 int32 FDMXFixturePatchNode::GetUniverseID() const
@@ -431,42 +374,21 @@ int32 FDMXFixturePatchNode::GetChannelSpan() const
 	return -1;
 }
 
-bool FDMXFixturePatchNode::NeedsUpdateGrid() const
+void FDMXFixturePatchNode::UpdateIsSelected()
 {
-	return 
-		UniverseID != GetUniverseID() ||
-		StartingChannel != GetStartingChannel() ||
-		ChannelSpan != GetChannelSpan();
-}
-
-void FDMXFixturePatchNode::OnSelectionChanged()
-{
-	const TSharedPtr<FDMXEditor> DMXEditor = WeakDMXEditor.Pin();
+	TSharedPtr<FDMXEditor> DMXEditor = WeakDMXEditor.Pin();
 	if (!DMXEditor.IsValid())
 	{
 		return;
 	}
 
-	const TSharedPtr<FDMXFixturePatchSharedData> FixturePatchSharedData = DMXEditor->GetFixturePatchSharedData();
-	TArray<TWeakObjectPtr<UDMXEntityFixturePatch>> SelectedFixturePatches = FixturePatchSharedData->GetSelectedFixturePatches();
-	if (SelectedFixturePatches.Contains(FixturePatch))
+	const TSharedPtr<FDMXFixturePatchSharedData> SharedData = DMXEditor->GetFixturePatchSharedData();
+	if (!SharedData.IsValid())
 	{
-		for (const TSharedPtr<SDMXFixturePatchFragment>& Widget : FixturePatchFragmentWidgets)
-		{
-			Widget->SetHighlight(true);
-		}
-
-		bSelected = true;
+		return;
 	}
-	else
-	{
-		for (const TSharedPtr<SDMXFixturePatchFragment>& Widget : FixturePatchFragmentWidgets)
-		{
-			Widget->SetHighlight(false);
-		}
 
-		bSelected = false;
-	}
+	bSelected = SharedData->GetSelectedFixturePatches().Contains(FixturePatch);
 }
 
 #undef LOCTEXT_NAMESPACE
diff --git a/Engine/Plugins/VirtualProduction/DMX/DMXEngine/Source/DMXEditor/Private/Widgets/FixturePatch/DMXFixturePatchNode.h b/Engine/Plugins/VirtualProduction/DMX/DMXEngine/Source/DMXEditor/Private/Widgets/FixturePatch/DMXFixturePatchNode.h
--- a/Engine/Plugins/VirtualProduction/DMX/DMXEngine/Source/DMXEditor/Private/Widgets/FixturePatch/DMXFixturePatchNode.h
+++ b/Engine/Plugins/VirtualProduction/DMX/DMXEngine/Source/DMXEditor/Private/Widgets/FixturePatch/DMXFixturePatchNode.h
@@ -10,6 +10,7 @@ class FDMXEditor;
 class SDMXFixturePatchFragment;
 class SDMXPatchedUniverse;
 class UDMXEntityFixturePatch;
+class UDMXMVRFixtureNode;
 
 
 /** 
@@ -28,25 +29,25 @@ public:
 	void SetAddresses(int32 UniverseID, int32 NewStartingChannel, int32 NewChannelSpan);
 
 	/** Generates widgets. Requires Addresses to be set via SetAddresses */
-	TArray<TSharedPtr<SDMXFixturePatchFragment>> GenerateWidgets(const TArray<TSharedPtr<FDMXFixturePatchNode>>& FixturePatchNodeGroup);
+	TArray<TSharedRef<SDMXFixturePatchFragment>> GenerateWidgets(const TSharedRef<SDMXPatchedUniverse>& OwningUniverse, const TArray<TSharedPtr<FDMXFixturePatchNode>>& FixturePatchNodeGroup);
 
-	/** Returns widgets */
-	const TArray<TSharedPtr<SDMXFixturePatchFragment>>& GetGeneratedWidgets() { return FixturePatchFragmentWidgets; };
-
-	/** Returns wether the patch uses specified channesl */
+	/** Returns wether the patch uses specified channels */
 	bool OccupiesChannels(int32 Channel, int32 Span) const;
 
-	/** Returns true if the node is patched in a universe */
-	bool IsPatched() const;
-
 	/** Returns true if the node is selected */
 	bool IsSelected() const;
 
+	/** Sets if the patch is hovered */
+	void SetIsHovered(bool bNewIsHovered) { bHovered = bNewIsHovered; }
+
+	/** Returns true if this patch is hovered */
+	bool IsHovered() const { return bHovered; }
+
 	/** Returns the fixture patches this node holds */
 	TWeakObjectPtr<UDMXEntityFixturePatch> GetFixturePatch() const { return FixturePatch; }
 
-	/** Returns the universe this node resides in */
-	const TSharedPtr<SDMXPatchedUniverse>& GetUniverseWidget() const { return UniverseWidget; }
+	/** Returns the Fixture ID of the patch */
+	FString GetFixtureID() const;
 
 	/** Returns the Universe ID the Node currently resides in. Returns a negative Value if not assigned to a Universe. */
 	int32 GetUniverseID() const;
@@ -63,37 +64,25 @@ public:
 	/** Gets the ZOrder of the Node */
 	int32 GetZOrder() const { return ZOrder; }
 
-	/** Returns true if the patch was moved on the grid and new widgets need be generated. */
-	bool NeedsUpdateGrid() const;
-
 private:
-	/** Called when a patch node was selected */
-	void OnSelectionChanged();
-
-	/** Universe the patch is assigned to */
-	TSharedPtr<SDMXPatchedUniverse> UniverseWidget;
-
-	/** Cached Universe of the patch */
-	int32 UniverseID = 0;
-
-	/** Cached starting channel of the patch */	
-	int32 StartingChannel = 0;
-
-	/** Cached channel span of the patch */
-	int32 ChannelSpan = 0;
+	/** Updates if the node is selected */
+	void UpdateIsSelected();
 
 	/** If true the node is selected */
 	bool bSelected = false;
 
+	/** True if the node is hovered */
+	bool bHovered = false;
+
 	/** ZOrder of this node */
 	int32 ZOrder = 0;
 
-	/** Keep track of Fragment Widgets, useful to handle selection */
-	TArray<TSharedPtr<SDMXFixturePatchFragment>> FixturePatchFragmentWidgets;
-
 	/** The Fixture Patch this Node stands for */
 	TWeakObjectPtr<UDMXEntityFixturePatch> FixturePatch;
 
+	/** MVR Fixture Node that corrresponds to the patch */
+	mutable TWeakObjectPtr<UDMXMVRFixtureNode> MVRFixtureNode;
+
 	/** Weak DMXEditor refrence */
 	TWeakPtr<FDMXEditor> WeakDMXEditor;
 };
diff --git a/Engine/Plugins/VirtualProduction/DMX/DMXEngine/Source/DMXEditor/Private/Widgets/FixturePatch/SDMXChannelConnector.cpp b/Engine/Plugins/VirtualProduction/DMX/DMXEngine/Source/DMXEditor/Private/Widgets/FixturePatch/SDMXChannelConnector.cpp
--- a/Engine/Plugins/VirtualProduction/DMX/DMXEngine/Source/DMXEditor/Private/Widgets/FixturePatch/SDMXChannelConnector.cpp
+++ b/Engine/Plugins/VirtualProduction/DMX/DMXEngine/Source/DMXEditor/Private/Widgets/FixturePatch/SDMXChannelConnector.cpp
@@ -10,15 +10,12 @@
 #include "Widgets/SDMXChannel.h"
 
 
-void SDMXChannelConnector::Construct(const FArguments& InArgs)
+void SDMXChannelConnector::Construct(const FArguments& InArgs, TWeakPtr<FDMXEditor> InWeakDMXEditor)
 {
-	check(InArgs._OnMouseButtonDownOnChannel.IsBound());
-	check(InArgs._OnMouseButtonUpOnChannel.IsBound());
-	check(InArgs._OnDragDetectedOnChannel.IsBound());
-	check(InArgs._OnDragEnterChannel.IsBound());
-	check(InArgs._OnDragLeaveChannel.IsBound());
-	check(InArgs._OnDropOntoChannel.IsBound());
+	WeakDMXEditor = InWeakDMXEditor;
 
+	OnHovered = InArgs._OnHovered;
+	OnUnhovered = InArgs._OnUnhovered;
 	OnMouseButtonDownOnChannel = InArgs._OnMouseButtonDownOnChannel;
 	OnMouseButtonUpOnChannel = InArgs._OnMouseButtonUpOnChannel;
 	OnDragDetectedOnChannel = InArgs._OnDragDetectedOnChannel;
@@ -29,14 +26,17 @@ void SDMXChannelConnector::Construct(const FArguments& InArgs)
 	ChannelID = InArgs._ChannelID;
 	UniverseID = InArgs._UniverseID;
 
-	DMXEditorPtr = InArgs._DMXEditor;
-
 	ChildSlot
 		[
 			SAssignNew(ChannelValueWidget, SDMXChannel)
 			.ChannelID(ChannelID)
 			.Value(InArgs._Value)	
 			.bShowChannelIDBottom(true)
+			.ChannelIDTextColor(FLinearColor::White.CopyWithNewOpacity(.6f))
+			.ValueTextColor(FLinearColor::White.CopyWithNewOpacity(.6f))
+			.Visibility(EVisibility::HitTestInvisible)
+			.ChannelIDTextColor(InArgs._ChannelIDTextColor)
+			.ValueTextColor(InArgs._ValueTextColor)
 		];
 }
 
@@ -45,19 +45,51 @@ void SDMXChannelConnector::SetValue(uint8 Value)
 	ChannelValueWidget->SetValue(Value);
 }
 
+void SDMXChannelConnector::OnMouseEnter(const FGeometry& MyGeometry, const FPointerEvent& MouseEvent)
+{
+	OnHovered.ExecuteIfBound();
+}
+
+void SDMXChannelConnector::OnMouseLeave(const FPointerEvent& MouseEvent)
+{
+	OnUnhovered.ExecuteIfBound();
+}
+
 FReply SDMXChannelConnector::OnMouseButtonDown(const FGeometry& MyGeometry, const FPointerEvent& MouseEvent)
 {
-	return OnMouseButtonDownOnChannel.Execute(ChannelValueWidget->GetChannelID(), MouseEvent);
+	if (OnMouseButtonDownOnChannel.IsBound())
+	{
+		if (MouseEvent.GetEffectingButton() == EKeys::LeftMouseButton)
+		{
+			return OnMouseButtonDownOnChannel.Execute(ChannelValueWidget->GetChannelID(), MouseEvent).DetectDrag(AsShared(), EKeys::LeftMouseButton);
+		}
+		else
+		{
+			return OnMouseButtonDownOnChannel.Execute(ChannelValueWidget->GetChannelID(), MouseEvent);
+		}
+	}
+
+	return FReply::Unhandled();
 }
 
 FReply SDMXChannelConnector::OnMouseButtonUp(const FGeometry& MyGeometry, const FPointerEvent& MouseEvent)
 {
-	return OnMouseButtonUpOnChannel.Execute(ChannelValueWidget->GetChannelID(), MouseEvent);
+	if (OnMouseButtonUpOnChannel.IsBound())
+	{
+		return OnMouseButtonUpOnChannel.Execute(ChannelValueWidget->GetChannelID(), MouseEvent);
+	}
+
+	return FReply::Unhandled();
 }
 
 FReply SDMXChannelConnector::OnDragDetected(const FGeometry& MyGeometry, const FPointerEvent& MouseEvent)
 {
-	return OnDragDetectedOnChannel.Execute(ChannelValueWidget->GetChannelID(), MouseEvent);
+	if (OnDragDetectedOnChannel.IsBound())
+	{
+		return OnDragDetectedOnChannel.Execute(ChannelValueWidget->GetChannelID(), MouseEvent);
+	}
+
+	return FReply::Unhandled();
 }
 
 void SDMXChannelConnector::OnDragEnter(const FGeometry& MyGeometry, const FDragDropEvent& DragDropEvent)
@@ -72,12 +104,17 @@ void SDMXChannelConnector::OnDragLeave(const FDragDropEvent& DragDropEvent)
 
 FReply SDMXChannelConnector::OnDrop(const FGeometry& MyGeometry, const FDragDropEvent& DragDropEvent)
 {
-	return OnDropOntoChannel.Execute(ChannelValueWidget->GetChannelID(), DragDropEvent);
+	if (OnDropOntoChannel.IsBound())
+	{
+		return OnDropOntoChannel.Execute(ChannelValueWidget->GetChannelID(), DragDropEvent);
+	}
+
+	return FReply::Unhandled();
 }
 
 UDMXLibrary* SDMXChannelConnector::GetDMXLibrary() const
 {
-	if (TSharedPtr<FDMXEditor> DMXEditor = DMXEditorPtr.Pin())
+	if (TSharedPtr<FDMXEditor> DMXEditor = WeakDMXEditor.Pin())
 	{
 		return DMXEditor->GetDMXLibrary();
 	}
diff --git a/Engine/Plugins/VirtualProduction/DMX/DMXEngine/Source/DMXEditor/Private/Widgets/FixturePatch/SDMXChannelConnector.h b/Engine/Plugins/VirtualProduction/DMX/DMXEngine/Source/DMXEditor/Private/Widgets/FixturePatch/SDMXChannelConnector.h
--- a/Engine/Plugins/VirtualProduction/DMX/DMXEngine/Source/DMXEditor/Private/Widgets/FixturePatch/SDMXChannelConnector.h
+++ b/Engine/Plugins/VirtualProduction/DMX/DMXEngine/Source/DMXEditor/Private/Widgets/FixturePatch/SDMXChannelConnector.h
@@ -29,7 +29,6 @@ public:
 		, _OnDragEnterChannel()
 		, _OnDragLeaveChannel()
 		, _OnDropOntoChannel()
-		, _DMXEditor()
 	{}
 		/** The channel ID this widget stands for */
 		SLATE_ARGUMENT(int32, ChannelID)
@@ -40,6 +39,12 @@ public:
 		/** The value of the channel */
 		SLATE_ARGUMENT(uint8, Value)
 
+		/** Called when the widget is hovered */
+		SLATE_EVENT(FSimpleDelegate, OnHovered)
+
+		/** Called when the widget is unhovered */
+		SLATE_EVENT(FSimpleDelegate, OnUnhovered)
+
 		/** Called when Mouse Button was pressed on the Channel */
 		SLATE_EVENT(FOnMouseEventWithReply, OnMouseButtonDownOnChannel)
 
@@ -58,12 +63,16 @@ public:
 		/** Called when dropped onto the Channel */
 		SLATE_EVENT(FOnDragEventWithReply, OnDropOntoChannel)
 
-		SLATE_ARGUMENT(TWeakPtr<FDMXEditor>, DMXEditor)
+		/** Color of the Channel ID Text */
+		SLATE_ATTRIBUTE(FSlateColor, ChannelIDTextColor)
+
+		/** Color of the Value Text */
+		SLATE_ATTRIBUTE(FSlateColor, ValueTextColor)
 
 	SLATE_END_ARGS()
 
 	/** Constructs this widget */
-	void Construct(const FArguments& InArgs);
+	void Construct(const FArguments& InArgs, TWeakPtr<FDMXEditor> InWeakDMXEditor);
 
 	/** Returns the ChannelID of this connector */
 	int32 GetChannelID() const { return ChannelID; }
@@ -79,6 +88,8 @@ public:
 
 protected:
 	// Begin SWidget interface	
+	virtual void OnMouseEnter(const FGeometry& MyGeometry, const FPointerEvent& MouseEvent) override;
+	virtual void OnMouseLeave(const FPointerEvent& MouseEvent) override;
 	virtual FReply OnMouseButtonDown(const FGeometry& MyGeometry, const FPointerEvent& MouseEvent) override;
 	virtual FReply OnMouseButtonUp(const FGeometry& MyGeometry, const FPointerEvent& PointerEvent) override;
 	virtual FReply OnDragDetected(const FGeometry& MyGeometry, const FPointerEvent& MouseEvent) override;
@@ -100,9 +111,11 @@ private:
 
 	int32 UniverseID;
 
-	TWeakPtr<FDMXEditor> DMXEditorPtr;
+	TWeakPtr<FDMXEditor> WeakDMXEditor;
 
 	// Slate Arguments
+	FSimpleDelegate OnHovered;
+	FSimpleDelegate OnUnhovered;
 	FOnMouseEventWithReply OnMouseButtonDownOnChannel;
 	FOnMouseEventWithReply OnMouseButtonUpOnChannel;
 	FOnMouseEventWithReply OnDragDetectedOnChannel;
diff --git a/Engine/Plugins/VirtualProduction/DMX/DMXEngine/Source/DMXEditor/Private/Widgets/FixturePatch/SDMXFixturePatchEditor.cpp b/Engine/Plugins/VirtualProduction/DMX/DMXEngine/Source/DMXEditor/Private/Widgets/FixturePatch/SDMXFixturePatchEditor.cpp
--- a/Engine/Plugins/VirtualProduction/DMX/DMXEngine/Source/DMXEditor/Private/Widgets/FixturePatch/SDMXFixturePatchEditor.cpp
+++ b/Engine/Plugins/VirtualProduction/DMX/DMXEngine/Source/DMXEditor/Private/Widgets/FixturePatch/SDMXFixturePatchEditor.cpp
@@ -3,6 +3,7 @@
 #include "Widgets/FixturePatch/SDMXFixturePatchEditor.h"
 
 #include "DMXEditor.h"
+#include "DMXEditorSettings.h"
 #include "DMXFixturePatchSharedData.h"
 #include "DMXSubsystem.h"
 #include "SDMXFixturePatcher.h"
@@ -21,6 +22,15 @@
 
 #define LOCTEXT_NAMESPACE "SDMXFixturePatcher"
 
+SDMXFixturePatchEditor::~SDMXFixturePatchEditor()
+{
+	const float LeftSideWidth = LhsRhsSplitter->SlotAt(0).GetSizeValue();
+
+	UDMXEditorSettings* const DMXEditorSettings = GetMutableDefault<UDMXEditorSettings>();
+	DMXEditorSettings->MVRFixtureListSettings.ListWidth = LeftSideWidth;
+	DMXEditorSettings->SaveConfig();
+}
+
 void SDMXFixturePatchEditor::Construct(const FArguments& InArgs)
 {
 	SDMXEntityEditor::Construct(SDMXEntityEditor::FArguments());
@@ -32,36 +42,42 @@ void SDMXFixturePatchEditor::Construct(const FArguments& InArgs)
 
 	FixturePatchDetailsView = GenerateFixturePatchDetailsView();
 
+	const UDMXEditorSettings* const DMXEditorSettings = GetDefault<UDMXEditorSettings>();
+	const float LeftSideWidth = FMath::Clamp(DMXEditorSettings->MVRFixtureListSettings.ListWidth, 0.1f, 0.9f);
+	const float RightSideWidth = FMath::Max(1.f - DMXEditorSettings->MVRFixtureListSettings.ListWidth, .1f);
+
 	ChildSlot
 	.VAlign(VAlign_Fill)
 	.HAlign(HAlign_Fill)
 	[
-		SNew(SSplitter)
+		SAssignNew(LhsRhsSplitter, SSplitter)
 		.Orientation(EOrientation::Orient_Horizontal)
 		.ResizeMode(ESplitterResizeMode::FixedPosition)
 		
 		// Left, MVR Fixture List
 		+ SSplitter::Slot()	
+		.Value(LeftSideWidth)
 		[
 			SAssignNew(MVRFixtureList, SDMXMVRFixtureList, DMXEditorPtr)
 		]
 
 		// Right, Fixture Patcher and Details
 		+ SSplitter::Slot()	
+		.Value(RightSideWidth)
 		[
 			SNew(SSplitter)
 			.Orientation(EOrientation::Orient_Vertical)
 			.ResizeMode(ESplitterResizeMode::FixedPosition)
 
 			+SSplitter::Slot()			
-			.Value(0.6f)
+			.Value(.618f)
 			[
 				SAssignNew(FixturePatcher, SDMXFixturePatcher)
 				.DMXEditor(DMXEditorPtr)
 			]
 	
 			+SSplitter::Slot()
-			.Value(0.4f)
+			.Value(.382f)
 			[
 				FixturePatchDetailsView.ToSharedRef()
 			]
diff --git a/Engine/Plugins/VirtualProduction/DMX/DMXEngine/Source/DMXEditor/Private/Widgets/FixturePatch/SDMXFixturePatchEditor.h b/Engine/Plugins/VirtualProduction/DMX/DMXEngine/Source/DMXEditor/Private/Widgets/FixturePatch/SDMXFixturePatchEditor.h
--- a/Engine/Plugins/VirtualProduction/DMX/DMXEngine/Source/DMXEditor/Private/Widgets/FixturePatch/SDMXFixturePatchEditor.h
+++ b/Engine/Plugins/VirtualProduction/DMX/DMXEngine/Source/DMXEditor/Private/Widgets/FixturePatch/SDMXFixturePatchEditor.h
@@ -15,10 +15,11 @@ class UDMXEntityFixturePatch;
 
 struct FPropertyChangedEvent;
 class IDetailsView;
+class SSplitter;
 
 
 /** Editor for Fixture Patches */
-class SDMXFixturePatchEditor
+class SDMXFixturePatchEditor final
 	: public SDMXEntityEditor
 {
 public:
@@ -31,6 +32,9 @@ public:
 	SLATE_END_ARGS()
 
 public:	
+	/** Destructor */
+	~SDMXFixturePatchEditor();
+
 	/** Constructs this widget */
 	void Construct(const FArguments& InArgs);
 
@@ -63,6 +67,9 @@ private:
 	/** Details View for the selected Fixture Patches */
 	TSharedPtr<IDetailsView> FixturePatchDetailsView;
 
+	/** The main splitter that divides the view in an left and right side */
+	TSharedPtr<SSplitter> LhsRhsSplitter;
+
 	/** Widget where the user can drag drop fixture patches */
 	TSharedPtr<SDMXFixturePatcher> FixturePatcher;
 
diff --git a/Engine/Plugins/VirtualProduction/DMX/DMXEngine/Source/DMXEditor/Private/Widgets/FixturePatch/SDMXFixturePatchFragment.cpp b/Engine/Plugins/VirtualProduction/DMX/DMXEngine/Source/DMXEditor/Private/Widgets/FixturePatch/SDMXFixturePatchFragment.cpp
--- a/Engine/Plugins/VirtualProduction/DMX/DMXEngine/Source/DMXEditor/Private/Widgets/FixturePatch/SDMXFixturePatchFragment.cpp
+++ b/Engine/Plugins/VirtualProduction/DMX/DMXEngine/Source/DMXEditor/Private/Widgets/FixturePatch/SDMXFixturePatchFragment.cpp
@@ -3,191 +3,363 @@
 #include "SDMXFixturePatchFragment.h"
 
 #include "DMXEditor.h"
+#include "DMXEditorSettings.h"
 #include "DMXEditorStyle.h"
 #include "DMXFixturePatchNode.h"
-#include "DMXFixturePatchEditorDefinitions.h"
 #include "DMXFixturePatchSharedData.h"
+#include "SDMXChannelConnector.h"
 #include "Library/DMXEntityFixturePatch.h"
 
-#include "Styling/AppStyle.h"
-#include "SlateOptMacros.h"
-#include "Widgets/Images/SImage.h"
-#include "Widgets/Layout/SBox.h"
+#include "Rendering/DrawElementPayloads.h"
+#include "Widgets/SBoxPanel.h"
 #include "Widgets/SOverlay.h"
-#include "Widgets/Layout/SBorder.h"
-#include "Widgets/Layout/SSpacer.h"
-#include "Widgets/Layout/SWidgetSwitcher.h"
-#include "Widgets/Text/STextBlock.h"
 
-#define LOCTEXT_NAMESPACE "SDMXFixturePatchFragment"
 
+#define LOCTEXT_NAMESPACE "SDMXFixturePatchFragment"
 
-void SDMXFixturePatchFragment::Construct(const FArguments& InArgs, const TSharedPtr<FDMXFixturePatchNode>& InFixturePatchNode, const TArray<TSharedPtr<FDMXFixturePatchNode>>& InFixturePatchNodeGroup)
+/** A widget that draws the color of a fixture patch depending on its state */
+class SDMXFixturePatchFragmentColor
+	: public SCompoundWidget
 {
-	check(InFixturePatchNode.IsValid());
-	check(InFixturePatchNodeGroup.Contains(InFixturePatchNode));
-	check(Column != -1);
-	check(Row != -1);
-	check(ColumnSpan != -1);
+public:
+	SLATE_BEGIN_ARGS(SDMXFixturePatchFragmentColor)
+		: _bIsTopmost(false)
+		, _bIsConflicting(false)
+	{}
+		SLATE_DEFAULT_SLOT(FArguments, Content)
 
-	ColorBrush = MakeShared<FSlateColorBrush>(FLinearColor::Red);
+		SLATE_ARGUMENT(bool, bIsTopmost)
 
-	FixturePatchNode = InFixturePatchNode;
-	FixturePatchNodeGroup = InFixturePatchNodeGroup;
-	DMXEditorPtr = InArgs._DMXEditor;
-	bIsHead = InArgs._IsHead;
-	bIsTail = InArgs._IsTail;
-	bIsText = InArgs._IsText;
-	bIsConflicting = InArgs._IsConflicting;
-	Column = InArgs._Column;
-	Row = InArgs._Row;
-	ColumnSpan = InArgs._ColumnSpan;
-	bHighlight = InArgs._bHighlight;
+		SLATE_ARGUMENT(bool, bIsConflicting)
 
-	if (bIsText)
-	{
-		SAssignNew(ContentBorder, SBorder)
-			.BorderImage(FAppStyle::GetBrush("NoBorder"))
-			.HAlign(HAlign_Fill)
-			.VAlign(VAlign_Fill)
-			.Visibility(EVisibility::SelfHitTestInvisible);
-	}
-	else
+		SLATE_ATTRIBUTE(bool, IsNodeHovered)
+
+	SLATE_END_ARGS()
+			
+	/** Constructs this widget */
+	void Construct(const FArguments& InArgs, const TSharedPtr<FDMXFixturePatchNode>& InFixturePatchNode)
 	{
-		SAssignNew(ContentBorder, SBorder)
-			.Padding(1.f)
-			.HAlign(HAlign_Fill)
-			.VAlign(VAlign_Fill)
-			.Visibility(EVisibility::SelfHitTestInvisible)
-			.BorderImage_Lambda([this]()
-				{
-					if (UDMXEntityFixturePatch* FixturePatch = FixturePatchNode->GetFixturePatch().Get())
-					{
-						if (bIsConflicting)
-						{
-							ColorBrush->TintColor = FixturePatch->EditorColor.CopyWithNewOpacity(.25f / FixturePatchNodeGroup.Num());
-						}
-						else if (bIsTopmost)
-						{
-							ColorBrush->TintColor = FixturePatch->EditorColor.CopyWithNewOpacity(.5f);
-						}
-						else
-						{
-							ColorBrush->TintColor = FLinearColor::Transparent;
-						}
-					}
-					check(ColorBrush.IsValid());
-					return ColorBrush.Get();
-				});
+		FixturePatchNode = InFixturePatchNode;
+
+		bIsTopmost = InArgs._bIsTopmost;
+		bIsConflicting = InArgs._bIsConflicting;
+		IsNodeHovered = InArgs._IsNodeHovered;
+
+		ConflictBrush = FDMXEditorStyle::Get().GetBrush("FixturePatchFragment.ConflictBackground");
+		SelectedBrush = FDMXEditorStyle::Get().GetBrush("DMXEditor.WhiteBrush");
+
+		ChildSlot
+			[
+				InArgs._Content.Widget
+			];
 	}
 
-	ChildSlot
+	/** Sets the content of the widget */
+	void SetContent(const TSharedRef<SWidget>& InContent)
+	{
+		ChildSlot
 		[
-			SNew(SBox) 	
-			.Visibility(EVisibility::SelfHitTestInvisible)
-			.MaxDesiredHeight(1.0f) // Prevents from some wiggling issues
-			.MaxDesiredWidth(1.0f) // Prevents from some wiggling issues
-			.Padding_Lambda([this]()
-				{
-					return BorderPadding;
-				})
-			[
-				SNew(SBorder)
-				.Padding(1.f)
-				.HAlign(HAlign_Fill)
-				.VAlign(VAlign_Fill)
-				.Visibility(EVisibility::SelfHitTestInvisible)
-				.BorderImage(&GetBorderImage())
-				.BorderBackgroundColor_Lambda([this]()
-					{
-						return bHighlight ? FLinearColor(.8f, .8f, .8f) : FLinearColor(.2f, .2f, .2f);
-					})
-				[	
-					ContentBorder.ToSharedRef()
-				]
-			]
+			InContent
 		];
+	}
 
-	SetVisibility(EVisibility::SelfHitTestInvisible);
-	Refresh(InFixturePatchNodeGroup);
-}
+protected:
+	//~ Begin SWidget interface
+	int32 OnPaint(const FPaintArgs& Args, const FGeometry& AllottedGeometry, const FSlateRect& MyCullingRect, FSlateWindowElementList& OutDrawElements, int32 LayerId, const FWidgetStyle& InWidgetStyle, bool bParentEnabled) const
+	{
+		if (!bIsTopmost)
+		{
+			return LayerId + 1;
+		}
 
-void SDMXFixturePatchFragment::Refresh(const TArray<TSharedPtr<FDMXFixturePatchNode>>& InFixturePatchNodeGroup)
-{
-	FixturePatchNodeGroup = InFixturePatchNodeGroup;
-	
-	// Update if it is topmost first so further methods can rely on bIsTopmost
-	UpdateIsTopmost();
+		const UDMXEntityFixturePatch* FixturePatch = FixturePatchNode->GetFixturePatch().Get();
+		if (!FixturePatch)
+		{
+			return LayerId + 1;
+		}
 
-	UpdateFixturePatchNameText();
-	UpdateNodeGroupText();
-	UpdateBorderPadding();
+		const bool bEnabled = ShouldBeEnabled(bParentEnabled);
+		if (bIsConflicting)
+		{
+			const FLinearColor BackgroundColor = ConflictBrush->GetTint(InWidgetStyle).CopyWithNewOpacity(.3f);
 
-	if (bIsText)
-	{
-		// Redraw text
-		ContentBorder->SetContent(CreateTextWidget());
+			FSlateDrawElement::MakeBox(
+				OutDrawElements,
+				LayerId,
+				AllottedGeometry.ToPaintGeometry(),
+				ConflictBrush,
+				ESlateDrawEffect::None,
+				BackgroundColor);
+		}
+		else if (FixturePatchNode->IsSelected())
+		{
+			const FLinearColor BackgroundColor = FixturePatch->EditorColor.CopyWithNewOpacity(.5f);
+
+			FSlateDrawElement::MakeBox(
+				OutDrawElements,
+				LayerId,
+				AllottedGeometry.ToPaintGeometry(),
+				SelectedBrush,
+				ESlateDrawEffect::None,
+				BackgroundColor);
+		}
+		else 
+		{
+			const float InvertedLuminanceClamped = FMath::Clamp(FMath::Abs(1.f - FixturePatch->EditorColor.GetLuminance()), .4f, .8f);
+			const float Opacity = IsNodeHovered.Get() ? InvertedLuminanceClamped : InvertedLuminanceClamped / 2.f;
+
+			TArray<FSlateGradientStop> GradientStops;
+			GradientStops.Add(FSlateGradientStop(FVector2D::ZeroVector, FixturePatch->EditorColor.CopyWithNewOpacity(Opacity)));
+			GradientStops.Add(FSlateGradientStop(AllottedGeometry.GetLocalSize(), FixturePatch->EditorColor.CopyWithNewOpacity(Opacity / 4.f)));
+			 
+			FSlateDrawElement::MakeGradient(
+				OutDrawElements,
+				LayerId,
+				AllottedGeometry.ToPaintGeometry(),
+				GradientStops,
+				EOrientation::Orient_Horizontal,
+				ESlateDrawEffect::None | ESlateDrawEffect::NoGamma
+			);
+		}
+
+		return SCompoundWidget::OnPaint(Args, AllottedGeometry, MyCullingRect, OutDrawElements, LayerId, InWidgetStyle, bParentEnabled);
 	}
-}
+	//~ End SWidget interface
+
+private:
+	/** Topmost attribute */
+	bool bIsTopmost = false;
+
+	/** Is conflicting attribute */
+	bool bIsConflicting = false;
+
+	/** True if parent is hovered */
+	TAttribute<bool> IsNodeHovered = false;
 
-void SDMXFixturePatchFragment::OnFixturePatchSharedDataSelectedFixturePatch()
+	/** Brush used to draw the selected color */
+	const FSlateBrush* SelectedBrush = nullptr;
+
+	/** Brush used to draw the warning color */
+	const FSlateBrush* ConflictBrush = nullptr;
+
+	/** Fixture Patch Node being displayed by this widget */
+	TSharedPtr<FDMXFixturePatchNode> FixturePatchNode;
+};
+
+/** Fixture patch text widget */
+class SDMXFixturePatchFragmentText
+	: public SCompoundWidget
 {
-	TSharedPtr<FDMXEditor> DMXEditor = DMXEditorPtr.Pin();
-	if (!DMXEditor.IsValid())
+public:
+	SLATE_BEGIN_ARGS(SDMXFixturePatchFragmentText)
+	{}
+
+	SLATE_END_ARGS()
+
+	/** Constructs this widget */
+	void Construct(const FArguments& InArgs, const TSharedPtr<FDMXFixturePatchNode>& InFixturePatchNode, const TArray<TSharedPtr<FDMXFixturePatchNode>>& InFixturePatchNodeGroup)
 	{
-		return;
+		FixturePatchNode = InFixturePatchNode;
+		FixturePatchNodeGroup = InFixturePatchNodeGroup;
+
+		InitializeNodeGroupText();
+
+		ChildSlot
+		.HAlign(HAlign_Left)
+		.VAlign(VAlign_Top)
+		.Padding(2.f)
+		[
+			SNew(SBorder)
+			.HAlign(HAlign_Left)
+			.VAlign(VAlign_Top)
+			.Padding(FMargin(2.f, 2.f, 0.f, 2.f))
+			.BorderImage(FDMXEditorStyle::Get().GetBrush("DMXEditor.RoundedPropertyBorder"))
+			[
+				SNew(SHorizontalBox)
+				.Clipping(EWidgetClipping::ClipToBoundsAlways)
+				.Visibility(EVisibility::Visible)
+
+				+ SHorizontalBox::Slot()
+				.HAlign(HAlign_Left)
+				.VAlign(VAlign_Top)
+				.AutoWidth()
+				[
+					SNew(STextBlock)
+					.Text(this, &SDMXFixturePatchFragmentText::GetNameText)
+					.TextStyle(FAppStyle::Get(), "SmallText")
+					.ColorAndOpacity(FLinearColor::White)
+				]
+				
+				+ SHorizontalBox::Slot()
+				.HAlign(HAlign_Left)
+				.VAlign(VAlign_Top)
+				.AutoWidth()
+				.Padding(FMargin(10.f, 0.f, 0.f, 0.f))
+				[
+					SNew(STextBlock)
+					.Text_Lambda([this]()
+						{
+							return NodeGroupText;
+						})
+					.TextStyle(FAppStyle::Get(), "SmallText")
+					.ColorAndOpacity(FLinearColor(.5f, .5f, .5f))
+				]
+			]
+		];
 	}
 
-	const TSharedPtr<FDMXFixturePatchSharedData> SharedData = DMXEditor->GetFixturePatchSharedData();
-	if (!SharedData.IsValid())
+private:
+	/** Updates the Fixture Patch Name Text */
+	FText GetNameText() const
 	{
-		return;
+		if (UDMXEntityFixturePatch* FixturePatch = FixturePatchNode->GetFixturePatch().Get())
+		{
+			const UDMXEditorSettings* DMXEditorSettings = GetDefault<UDMXEditorSettings>();
+
+			const EDMXFixturePatcherNameDisplayMode NameDisplayMode = DMXEditorSettings->FixturePatcherSettings.FixturePatchNameDisplayMode;
+			if (NameDisplayMode == EDMXFixturePatcherNameDisplayMode::FixtureIDAndFixturePatchName)
+			{
+				return FText::FromString(FString::Printf(TEXT("%s - %s"), *FixturePatchNode->GetFixtureID(), *FixturePatch->GetDisplayName()));
+			}
+			else if (NameDisplayMode == EDMXFixturePatcherNameDisplayMode::FixtureID)
+			{
+				return FText::FromString(FixturePatchNode->GetFixtureID());
+			}
+			else if (NameDisplayMode == EDMXFixturePatcherNameDisplayMode::FixturePatchName)
+			{
+				return FText::FromString(FixturePatch->GetDisplayName());
+			}
+			else
+			{
+				checkf(0, TEXT("Unhandled Fixture Patch Name Display Mode"));
+			}
+		}
+
+		return FText::GetEmpty();
 	}
-	const TArray<TWeakObjectPtr<UDMXEntityFixturePatch>>& SelectedFixturePatches = SharedData->GetSelectedFixturePatches();
-}
 
-void SDMXFixturePatchFragment::UpdateFixturePatchNameText()
-{
-	if (UDMXEntityFixturePatch* Patch = FixturePatchNode->GetFixturePatch().Get())
+	/** Updates the Node Group Text */
+	void InitializeNodeGroupText()
 	{
-		FixturePatchNameText = FText::FromString(Patch->GetDisplayName());
+		NodeGroupText =
+			FixturePatchNodeGroup.Num() > 1 ?
+			FText::Format(LOCTEXT("NodeGroupText", "(and {0} more)"), FixturePatchNodeGroup.Num() - 1) :
+			FText::GetEmpty();
 	}
-}
 
-void SDMXFixturePatchFragment::UpdateNodeGroupText()
-{
-	NodeGroupText =
-		FixturePatchNodeGroup.Num() > 1 ?
-		FText::Format(LOCTEXT("NodeGroupText", "(and {0} more)"), FixturePatchNodeGroup.Num() - 1) :
-		FText::GetEmpty();
-}
+	/** The Node Group Text the Node displays */
+	FText NodeGroupText;
+
+	/** Fixture Patch Node being displayed by this widget */
+	TSharedPtr<FDMXFixturePatchNode> FixturePatchNode;
+
+	/** The group of nodes (same Universe, same Channel) this widget belongs to */
+	TArray<TSharedPtr<FDMXFixturePatchNode>> FixturePatchNodeGroup;
+};
 
-void SDMXFixturePatchFragment::UpdateIsTopmost()
+
+void SDMXFixturePatchFragment::Construct(const FArguments& InArgs, const TSharedPtr<FDMXFixturePatchNode>& InFixturePatchNode, const TArray<TSharedPtr<FDMXFixturePatchNode>>& InFixturePatchNodeGroup)
 {
-	if (!ensureMsgf(FixturePatchNodeGroup.Num() > 0, TEXT("Unexpected empty node group for Fixture Patch Fragment Widget.")))
-	{
-		return;
-	}
+	check(InFixturePatchNode.IsValid());
+	check(InFixturePatchNodeGroup.Contains(InFixturePatchNode));
+	check(InArgs._StartingChannel != -1);
+	check(InArgs._Column != -1);
+	check(InArgs._Row != -1);
+	check(InArgs._ColumnSpan != -1);
 
-	bIsTopmost = !FixturePatchNodeGroup.ContainsByPredicate([this](const TSharedPtr<FDMXFixturePatchNode>& OtherNode)
+	FixturePatchNode = InFixturePatchNode;
+	FixturePatchNodeGroup = InFixturePatchNodeGroup;
+	StartingChannel = InArgs._StartingChannel;
+	Column = InArgs._Column;
+	Row = InArgs._Row;
+	ColumnSpan = InArgs._ColumnSpan;
+
+	bIsHead = InArgs._bIsHead;
+	bIsTail = InArgs._bIsTail;
+	bIsConflicting = InArgs._bIsConflicting;
+
+	SetVisibility(EVisibility::HitTestInvisible);
+
+	InitializeIsTopmost();
+	InitializeBorderPadding();
+
+	ChildSlot
+	[
+		SAssignNew(ContentOverlay, SOverlay)
+	
+		+ SOverlay::Slot()
+		[
+			SNew(SDMXFixturePatchFragmentColor, FixturePatchNode)
+			.bIsTopmost(bIsTopmost)
+			.bIsConflicting(bIsConflicting)
+			.IsNodeHovered_Lambda([this]()
+				{
+					return FixturePatchNode->IsHovered();
+				})
+		]
+		
+		+ SOverlay::Slot()
+		[
+			SNew(SBox)
+			.Padding(BorderPadding)
+			[
+				SNew(SBorder)
+				.Padding(0.f)
+				.HAlign(HAlign_Fill)
+				.VAlign(VAlign_Fill)
+				.BorderImage(GetBorderBrush())
+				.BorderBackgroundColor(this, &SDMXFixturePatchFragment::GetBorderColor)
+			]
+		]
+
+		+ SOverlay::Slot()
+		[
+			SAssignNew(ChannelsHorizontalBox, SHorizontalBox)
+		]
+	];
+
+	if (bIsTopmost)
+	{	
+		// Add fixture patch text, only for the topmost head node
+		if (bIsHead)
 		{
-			return OtherNode->GetZOrder() > FixturePatchNode->GetZOrder();
-		});
+			ContentOverlay->AddSlot(3)
+			[
+				SNew(SBox)
+				.MaxDesiredHeight(1.0f) // Prevents from resizing the widget to text size
+				.MaxDesiredWidth(1.0f) // Prevents from resizing the widget to text size
+				.Padding(BorderPadding)
+				[
+					SNew(SDMXFixturePatchFragmentText, FixturePatchNode, FixturePatchNodeGroup)
+				]
+			];
+		}
 
-	if (FixturePatchNodeGroup.Num() > 1)
-	{
-		UpdateFixturePatchNameText();
+		// Add channel connectors
+		for (int32 Channel = StartingChannel; Channel < StartingChannel + ColumnSpan; Channel++)
+		{
+			ChannelsHorizontalBox->AddSlot()
+				.HAlign(HAlign_Fill)
+				.VAlign(VAlign_Fill)
+				[
+					SNew(SDMXChannelConnector, DMXEditorPtr)
+					.ChannelID(Channel)
+					.UniverseID(FixturePatchNode->GetUniverseID())
+					.ChannelIDTextColor(this, &SDMXFixturePatchFragment::GetChannelIDTextColor)
+					.ValueTextColor(this, &SDMXFixturePatchFragment::GetValueTextColor)
+				];
+		}
 	}
 }
 
-void SDMXFixturePatchFragment::UpdateBorderPadding()
+void SDMXFixturePatchFragment::InitializeBorderPadding()
 {
 	if (bIsHead && FixturePatchNodeGroup.Num() > 1)
 	{
 		// Sort by ZOrder ascending
 		FixturePatchNodeGroup.StableSort([](const TSharedPtr<FDMXFixturePatchNode>& NodeA, const TSharedPtr<FDMXFixturePatchNode>& NodeB)
 			{
-				return 
+				return
 					NodeA->GetZOrder() < NodeB->GetZOrder();
 			});
 
@@ -205,90 +377,89 @@ void SDMXFixturePatchFragment::UpdateBorderPadding()
 	}
 }
 
-const FSlateBrush& SDMXFixturePatchFragment::GetBorderImage() const
+void SDMXFixturePatchFragment::InitializeIsTopmost()
 {
-	const FSlateBrush* BorderImage = nullptr;
-	if (bIsText)
-	{
-		BorderImage = FCoreStyle::Get().GetBrush("NoBorder");
-	}
-	else if (bIsConflicting)
+	if (!ensureMsgf(FixturePatchNodeGroup.Num() > 0, TEXT("Unexpected empty node group for Fixture Patch Fragment Widget.")))
 	{
-		BorderImage = FAppStyle::GetBrush("Graph.Node.DevelopmentBanner");
+		return;
 	}
-	else if (bIsHead && bIsTail)
+
+	bIsTopmost = !FixturePatchNodeGroup.ContainsByPredicate([this](const TSharedPtr<FDMXFixturePatchNode>& OtherNode)
+		{
+			const bool bOtherHasHigherZOrder = OtherNode->GetZOrder() > FixturePatchNode->GetZOrder();
+			const bool bOtherNodeIntersects = OtherNode->GetStartingChannel() + OtherNode->GetChannelSpan() > StartingChannel;
+
+			return bOtherHasHigherZOrder && bOtherNodeIntersects;
+		});
+}
+
+const FSlateBrush* SDMXFixturePatchFragment::GetBorderBrush() const
+{
+	if (bIsHead && bIsTail)
 	{
-		BorderImage = FDMXEditorStyle::Get().GetBrush("FixturePatcher.FragmentBorder.Normal");
+		return FDMXEditorStyle::Get().GetBrush("FixturePatcher.FragmentBorder.Normal");
 	}
 	else if (bIsHead)
 	{
-		BorderImage = FDMXEditorStyle::Get().GetBrush("FixturePatcher.FragmentBorder.L");
+		return FDMXEditorStyle::Get().GetBrush("FixturePatcher.FragmentBorder.L");
 	}
 	else if (bIsTail)
 	{
-		BorderImage = FDMXEditorStyle::Get().GetBrush("FixturePatcher.FragmentBorder.R");
+		return FDMXEditorStyle::Get().GetBrush("FixturePatcher.FragmentBorder.R");
 	}
 	else
 	{
-		BorderImage = FDMXEditorStyle::Get().GetBrush("FixturePatcher.FragmentBorder.TB");
+		return FDMXEditorStyle::Get().GetBrush("FixturePatcher.FragmentBorder.TB");
 	}
-	check(BorderImage);
+}
 
-	return *BorderImage;
+FSlateColor SDMXFixturePatchFragment::GetBorderColor() const
+{
+	// Note all colors need be without oppacity so they draw over the conflict background brush
+	if (FixturePatchNode->IsSelected())
+	{
+		return FLinearColor::White;
+	}
+	else if (UDMXEntityFixturePatch* Patch = FixturePatchNode->GetFixturePatch().Get())
+	{
+		return (Patch->EditorColor * 3.f + FLinearColor::White) / 4.f; 
+	}
+	else
+	{
+		return FLinearColor(1.f, 0.f, 1.f);
+	}
 }
 
-TSharedRef<SWidget> SDMXFixturePatchFragment::CreateTextWidget()
+FSlateColor SDMXFixturePatchFragment::GetChannelIDTextColor() const
 {
-	if (bIsTopmost)
+	if (bIsConflicting)
 	{
-		return
-			SNew(SBox)
-			.HAlign(HAlign_Left)
-			.VAlign(VAlign_Top)
-			.Visibility(EVisibility::HitTestInvisible)
-			[
-				SNew(SBorder)
-				.HAlign(HAlign_Left)
-				.VAlign(VAlign_Top)
-				.Padding(FMargin(2.f, 2.f, 0.f, 2.f))
-				.BorderImage(FDMXEditorStyle::Get().GetBrush("DMXEditor.RoundedPropertyBorder"))
-				[
-					SNew(SHorizontalBox)
-					.Clipping(EWidgetClipping::ClipToBoundsAlways)
-
-					+ SHorizontalBox::Slot()
-					.HAlign(HAlign_Left)
-					.VAlign(VAlign_Top)
-					.AutoWidth()
-					[
-						SNew(STextBlock)
-						.Text_Lambda([this]()
-							{
-								return FixturePatchNameText;
-							})
-						.TextStyle(FAppStyle::Get(), "SmallText")
-						.ColorAndOpacity(FLinearColor::White)
-					]
-				
-					+ SHorizontalBox::Slot()
-					.HAlign(HAlign_Left)
-					.VAlign(VAlign_Top)
-					.AutoWidth()
-					.Padding(FMargin(10.f, 0.f, 0.f, 0.f))
-					[
-						SNew(STextBlock)
-						.Text_Lambda([this]()
-							{
-								return NodeGroupText;
-							})
-						.TextStyle(FAppStyle::Get(), "SmallText")
-						.ColorAndOpacity(FLinearColor(.5f, .5f, .5f))
-					]
-				]
-			];
+		return FLinearColor::Black;
+	}
+	else if (FixturePatchNode->IsHovered() || FixturePatchNode->IsSelected())
+	{
+		return FLinearColor::White;
+	}
+	else
+	{
+		return FLinearColor(.8f, .8f, .8f, 1.f);
 	}
+}
 
-	return SNullWidget::NullWidget;
+FSlateColor SDMXFixturePatchFragment::GetValueTextColor() const
+{
+	if (bIsConflicting)
+	{
+		return FLinearColor::Black;
+	}
+	else if (FixturePatchNode->IsHovered() || FixturePatchNode->IsSelected())
+	{
+		return FLinearColor::White;
+	}
+	else
+	{
+		return FLinearColor(.6f, .6f, .6f, 1.f);
+	}
 }
 
 #undef LOCTEXT_NAMESPACE
diff --git a/Engine/Plugins/VirtualProduction/DMX/DMXEngine/Source/DMXEditor/Private/Widgets/FixturePatch/SDMXFixturePatchFragment.h b/Engine/Plugins/VirtualProduction/DMX/DMXEngine/Source/DMXEditor/Private/Widgets/FixturePatch/SDMXFixturePatchFragment.h
--- a/Engine/Plugins/VirtualProduction/DMX/DMXEngine/Source/DMXEditor/Private/Widgets/FixturePatch/SDMXFixturePatchFragment.h
+++ b/Engine/Plugins/VirtualProduction/DMX/DMXEngine/Source/DMXEditor/Private/Widgets/FixturePatch/SDMXFixturePatchFragment.h
@@ -2,109 +2,96 @@
 
 #pragma once
 
-#include "CoreMinimal.h"
-#include "Styling/SlateColor.h"
-#include "Templates/SharedPointer.h"
-#include "UObject/WeakObjectPtrTemplates.h"
 #include "Widgets/DeclarativeSyntaxSupport.h"
 #include "Widgets/SCompoundWidget.h"
 
 class FDMXEditor;
 class FDMXFixturePatchNode;
-class UDMXEntityFixturePatch;
-class UDMXLibrary;
 
-struct FSlateColorBrush;
-class SBorder;
-template <typename OptionType> class SComboBox;
+class SHorizontalBox;
+class SOverlay;
 
 
 class SDMXFixturePatchFragment
-	: public SCompoundWidget 
+	: public SCompoundWidget
 {
+	DECLARE_DELEGATE_RetVal_TwoParams(FReply, FOnMouseEventWithReply, uint32 /** ChannelID */, const FPointerEvent&);
+
+	DECLARE_DELEGATE_TwoParams(FOnDragEvent, uint32 /** ChannelID */, const FDragDropEvent&);
+
+	DECLARE_DELEGATE_RetVal_TwoParams(FReply, FOnDragEventWithReply, uint32 /** ChannelID */, const FDragDropEvent&);
+
 public:
 	SLATE_BEGIN_ARGS(SDMXFixturePatchFragment)
-		: _DMXEditor(nullptr)
-		, _IsHead(true)
-		, _IsTail(true)
-		, _IsText(false)
-		, _IsConflicting(false)
+		: _bIsHead(true)
+		, _bIsTail(true)
+		, _bIsConflicting(false)
 		, _Column(-1)
 		, _Row(-1)
 		, _ColumnSpan(1)
-		, _bHighlight(false)
 	{}
-		SLATE_ARGUMENT(TWeakPtr<FDMXEditor>, DMXEditor)
-
-		SLATE_ARGUMENT(bool, IsHead)
+		/** If true displays as head of a patch */
+		SLATE_ARGUMENT(bool, bIsHead)
 
-		SLATE_ARGUMENT(bool, IsTail)
+		/** If true displays as tail of a patch */
+		SLATE_ARGUMENT(bool, bIsTail)
 
-		SLATE_ARGUMENT(bool, IsText)
+		/** If true displays a conflict */
+		SLATE_ARGUMENT(bool, bIsConflicting)
 
-		SLATE_ARGUMENT(bool, IsConflicting)
+		/** The starting channel of the fragment */
+		SLATE_ARGUMENT(int32, StartingChannel)
 
+		/** The column in which the fragment should be displayed */
 		SLATE_ARGUMENT(int32, Column)
 
+		/** The row in which the fragment should be displayed */
 		SLATE_ARGUMENT(int32, Row)
 
+		/** The column span of the fragment */
 		SLATE_ARGUMENT(int32, ColumnSpan)
 
-		SLATE_ARGUMENT(bool, bHighlight)
-
 	SLATE_END_ARGS()
 
 	/** Constructs this widget */
 	void Construct(const FArguments& InArgs, const TSharedPtr<FDMXFixturePatchNode>& InFixturePatchNode, const TArray<TSharedPtr<FDMXFixturePatchNode>>& InFixturePatchNodeGroup);
 
-	/** Refreshes the widget */
-	void Refresh(const TArray<TSharedPtr<FDMXFixturePatchNode>>& InFixturePatchNodeGroup);
-
-	/** Sets wether the fragment should be highlit */
-	void SetHighlight(bool bEnabled) { bHighlight = bEnabled; }
-
 	/** Sets if the fragment is conflicting with others */
-	void SetConflicting(bool bConflicting) { bIsConflicting = bConflicting; if (Column == 0 && Row == 0 && bIsConflicting) UE_LOG(LogTemp, Warning, TEXT("%s set to conflicting"), *FixturePatchNameText.ToString()); }
+	void SetConflicting(bool bConflicting) { bIsConflicting = bConflicting; }
 
-	/** Returns the Column of the fragment */
-	int32 GetColumn() const { return Column; }
+	/** Gets the fixture patch node this widget refers to */
+	const TSharedPtr<FDMXFixturePatchNode>& GetFixturePatchNode() const { return FixturePatchNode; }
 
-	/** Sets the Column of the fragment */
-	void SetColumn(int32 NewColumn) { Column = NewColumn; }
+	/** Starting Channel of the Fragment */
+	int32 StartingChannel = -1;
 
-	/** Returns the Row of the fragment */
-	int32 GetRow() const { return Row; }
-
-	/** Sets the Row of the fragment */
-	void SetRow(int32 NewRow) { Row = NewRow; }
+	/** Column of the Fragment */
+	int32 Column = -1;
 
-	/** Returns the Column span of the fragment */
-	int32 GetColumnSpan() const { return ColumnSpan; }
+	/** Row of the Fragment */
+	int32 Row = -1;
 
-	/** Sets the Column span of the fragment */
-	void SetColumnSpan(int32 NewColumnSpan) { ColumnSpan = NewColumnSpan; }
+	/** Column span of the Fragment */
+	int32 ColumnSpan = 1;
 
 private:
-	/** Called when Fixture Patch Shared Data selected a Fixture Patch */
-	void OnFixturePatchSharedDataSelectedFixturePatch();
+	/** Updates the Padding of the outermost Border */
+	void InitializeBorderPadding();
 
-	/** Updates the Fixture Patch Name Text */
-	void UpdateFixturePatchNameText();
+	/** Updates the bTopmost member from the ZOrder of the node */
+	void InitializeIsTopmost();
 
-	/** Updates the Node Group Text */
-	void UpdateNodeGroupText();
+	/** Returns the brush for the border around the widget */
+	const FSlateBrush* GetBorderBrush() const;
 
-	/** Updates the Padding of the outermost Border */
-	void UpdateBorderPadding();
+	/** Returns the border color of the border around the widget */
+	FSlateColor GetBorderColor() const;
 
-	/** Updates the bTopmost memeber from the ZOrder of the node */
-	void UpdateIsTopmost();
+	/** Returns the color of the channel ID texts */
+	FSlateColor GetChannelIDTextColor() const;
 
-	/** Returns the Border image */
-	const FSlateBrush& GetBorderImage() const;
-
-	/** Gets the Text Widget */
-	TSharedRef<SWidget> CreateTextWidget();
+	/** Returns the color of the Value texts */
+	FSlateColor GetValueTextColor() const;
 
 	/** If true, the patch has the highest ZOrder in its node group */
 	bool bIsTopmost = false;
@@ -115,38 +102,17 @@ private:
 	/** If true this is the Tail of the Fixture Patch */
 	bool bIsTail = false;
 
-	/** If true, this fragment shows text */
-	bool bIsText;
-
-	/** Column of the Fragment */
-	int32 Column;
-
-	/** Row of the Fragment */
-	int32 Row;
-
-	/** Column span of the Fragment */
-	int32 ColumnSpan;
-
-	/** If true, the Widget shows a highlight */
-	bool bHighlight = false;
-
 	/** If true, the Widget displays a conflict background */
 	bool bIsConflicting = false;
 
-	/** The Fixture Patch Name Text the Node displays */
-	FText FixturePatchNameText;
-
-	/** The Node Group Text the Node displays */
-	FText NodeGroupText;
-
 	/** The padding of the outermost border */
 	FMargin BorderPadding;
-	
-	/** The border that holds the color or text content */
-	TSharedPtr<SBorder> ContentBorder;
 
-	/** Color Brush the widget uses */
-	TSharedPtr<FSlateColorBrush> ColorBrush;
+	/** Horizontal box containing the channels of this fragment */
+	TSharedPtr<SHorizontalBox> ChannelsHorizontalBox;
+
+	/** Overlay for content */
+	TSharedPtr<SOverlay> ContentOverlay;
 
 	/** Fixture Patch Node being displayed by this widget */
 	TSharedPtr<FDMXFixturePatchNode> FixturePatchNode;
diff --git a/Engine/Plugins/VirtualProduction/DMX/DMXEngine/Source/DMXEditor/Private/Widgets/FixturePatch/SDMXFixturePatcher.cpp b/Engine/Plugins/VirtualProduction/DMX/DMXEngine/Source/DMXEditor/Private/Widgets/FixturePatch/SDMXFixturePatcher.cpp
--- a/Engine/Plugins/VirtualProduction/DMX/DMXEngine/Source/DMXEditor/Private/Widgets/FixturePatch/SDMXFixturePatcher.cpp
+++ b/Engine/Plugins/VirtualProduction/DMX/DMXEngine/Source/DMXEditor/Private/Widgets/FixturePatch/SDMXFixturePatcher.cpp
@@ -13,16 +13,17 @@
 #include "IO/DMXPortManager.h"
 #include "Library/DMXEntityFixturePatch.h"
 #include "Library/DMXLibrary.h"
-#include "Widgets/Layout/SBox.h"
-#include "Widgets/SNullWidget.h"
 
 #include "Algo/RemoveIf.h"
 #include "Styling/AppStyle.h"
 #include "ScopedTransaction.h"
 #include "SlateOptMacros.h"
+#include "Widgets/SNullWidget.h"
 #include "Widgets/Docking/SDockTab.h"
 #include "Widgets/Input/SCheckBox.h"
+#include "Widgets/Input/SComboButton.h"
 #include "Widgets/Input/SNumericEntryBox.h"
+#include "Widgets/Layout/SBox.h"
 #include "Widgets/Layout/SScrollBox.h"
 #include "Widgets/Layout/SSeparator.h"
 
@@ -40,9 +41,7 @@ void SDMXFixturePatcher::Construct(const FArguments& InArgs)
 	}
 
 	SharedData = DMXEditorPtr.Pin()->GetFixturePatchSharedData();
-		
-	const UDMXEditorSettings* DMXEditorSettings = GetDefault<UDMXEditorSettings>();
-	const ECheckBoxState InitialDMXMonitorEnabledCheckBoxState = DMXEditorSettings->bFixturePatcherDMXMonitorEnabled ? ECheckBoxState::Checked : ECheckBoxState::Unchecked;
+
 	const FLinearColor BackgroundTint(0.6f, 0.6f, 0.6f, 1.0f);
 	ChildSlot			
 		[
@@ -52,8 +51,6 @@ void SDMXFixturePatcher::Construct(const FArguments& InArgs)
 			[
 				SNew(SVerticalBox)				
 
-				// Settings area
-
 				+ SVerticalBox::Slot()
 				.HAlign(HAlign_Fill)
 				.AutoHeight()
@@ -64,96 +61,10 @@ void SDMXFixturePatcher::Construct(const FArguments& InArgs)
 					.BorderBackgroundColor(BackgroundTint)
 					.BorderImage(FAppStyle::GetBrush("DetailsView.CategoryTop"))
 					[
-						SNew(SHorizontalBox)			
-
-						+ SHorizontalBox::Slot()						
-						.AutoWidth()
-						.VAlign(VAlign_Center)
-						.Padding(FMargin(4.0f, 4.0f, 15.0f, 4.0f))
-						[
-							SNew(STextBlock)							
-							.MinDesiredWidth(75.0f)
-							.Font(FAppStyle::GetFontStyle("PropertyWindow.NormalFont"))
-							.TextStyle(FAppStyle::Get(), "DetailsView.CategoryTextStyle")
-							.IsEnabled(this, &SDMXFixturePatcher::IsUniverseSelectionEnabled)
-							.Text(LOCTEXT("UniverseSelectorLabel", "Universe"))
-						]
-
-						+ SHorizontalBox::Slot()
-						.AutoWidth()
-						.VAlign(VAlign_Center)
-						.Padding(FMargin(4.0f, 4.0f, 15.0f, 4.0f))
-						[
-							SNew(SBox)
-							.MinDesiredWidth(210.0f)
-							.MaxDesiredWidth(420.0f)
-							[
-								SNew(SSpinBox<int32>)								
-								.SliderExponent(1000.0f)								
-								.MinSliderValue(0)
-								.MaxSliderValue(DMX_MAX_UNIVERSE - 1)
-								.MinValue(0)
-								.MaxValue(DMX_MAX_UNIVERSE - 1)
-								.IsEnabled(this, &SDMXFixturePatcher::IsUniverseSelectionEnabled)
-								.Value(this, &SDMXFixturePatcher::GetSelectedUniverse)
-								.OnValueChanged(this, &SDMXFixturePatcher::SelectUniverse)
-							]
-						]
-						
-						+ SHorizontalBox::Slot()
-						.AutoWidth()
-						.VAlign(VAlign_Center)
-						.Padding(FMargin(4.0f, 4.0f, 15.0f, 4.0f))						
-						[
-							SNew(SSeparator)
-							.Orientation(EOrientation::Orient_Vertical)
-						]
-
-						+ SHorizontalBox::Slot()
-						.AutoWidth()
-						.VAlign(VAlign_Center)
-						.Padding(FMargin(4.0f, 4.0f, 15.0f, 4.0f))
-						[
-							SNew(STextBlock)
-							.Font(FAppStyle::GetFontStyle("PropertyWindow.NormalFont"))
-							.Text(LOCTEXT("ShowAllPatchedUniversesLabel", "Show all patched Universes"))
-						]
-
-						+ SHorizontalBox::Slot()
-						.AutoWidth()
-						.VAlign(VAlign_Center)
-						.Padding(FMargin(4.0f, 4.0f, 15.0f, 4.0f))
-						[
-							SAssignNew(ShowAllUniversesCheckBox, SCheckBox)
-							.IsChecked(false)
-							.OnCheckStateChanged(this, &SDMXFixturePatcher::OnToggleDisplayAllUniverses)
-						]
-
-						+ SHorizontalBox::Slot()
-						.AutoWidth()
-						.VAlign(VAlign_Center)
-						.Padding(FMargin(4.0f, 4.0f, 15.0f, 4.0f))
-						[
-							SNew(STextBlock)
-							.Font(FAppStyle::GetFontStyle("PropertyWindow.NormalFont"))
-							.Text(LOCTEXT("EnableInputMonitorLabel", "Monitor DMX Inputs"))
-							.ToolTipText(LOCTEXT("EnableInputMonitorTooltip", "If checked, monitors DMX Input Ports used in this DMX Library"))
-						]
-
-						+ SHorizontalBox::Slot()
-						.AutoWidth()
-						.VAlign(VAlign_Center)
-						.Padding(FMargin(4.0f, 4.0f, 15.0f, 4.0f))
-						[
-							SAssignNew(EnableDMXMonitorCheckBox, SCheckBox)
-							.IsChecked(InitialDMXMonitorEnabledCheckBoxState)
-							.OnCheckStateChanged(this, &SDMXFixturePatcher::OnToggleDMXMonitorEnabled)
-						]
+						CreateToolbar()
 					]
 				]
 
-				// Patched Universes
-
 				+ SVerticalBox::Slot()
 				.HAlign(HAlign_Left)
 				.VAlign(VAlign_Fill)
@@ -190,7 +101,8 @@ void SDMXFixturePatcher::Construct(const FArguments& InArgs)
 
 	ShowSelectedUniverse();
 
-	SetDMXMonitorEnabled(DMXEditorSettings->bFixturePatcherDMXMonitorEnabled);
+	const UDMXEditorSettings* DMXEditorSettings = GetDefault<UDMXEditorSettings>();
+	SetDMXMonitorEnabled(DMXEditorSettings->FixturePatcherSettings.bMonitorEnabled);
 }
 END_SLATE_FUNCTION_BUILD_OPTIMIZATION
 
@@ -245,6 +157,155 @@ FReply SDMXFixturePatcher::OnDrop(const FGeometry& MyGeometry, const FDragDropEv
 	return FReply::Unhandled();
 }
 
+TSharedRef<SWidget> SDMXFixturePatcher::CreateToolbar()
+{
+	TSharedPtr<FUICommandList> CommandList = nullptr;
+	FToolBarBuilder	ToolbarBuilder(CommandList, FMultiBoxCustomization::None);
+
+	ToolbarBuilder.BeginSection("UniverseSelection");
+	{
+		ToolbarBuilder.AddWidget(
+			SNew(STextBlock)
+			.MinDesiredWidth(75.0f)
+			.Font(FAppStyle::GetFontStyle("PropertyWindow.NormalFont"))
+			.TextStyle(FAppStyle::Get(), "DetailsView.CategoryTextStyle")
+			.IsEnabled(this, &SDMXFixturePatcher::IsUniverseSelectionEnabled)
+			.Text(LOCTEXT("UniverseSelectorLabel", "Universe"))
+		);
+
+		ToolbarBuilder.AddWidget(
+			SNew(SBox)
+			.MinDesiredWidth(210.0f)
+			.MaxDesiredWidth(420.0f)
+			[
+				SNew(SSpinBox<int32>)
+				.SliderExponent(1000.0f)
+				.MinSliderValue(0)
+				.MaxSliderValue(DMX_MAX_UNIVERSE - 1)
+				.MinValue(0)
+				.MaxValue(DMX_MAX_UNIVERSE - 1)
+				.IsEnabled(this, &SDMXFixturePatcher::IsUniverseSelectionEnabled)
+				.Value(this, &SDMXFixturePatcher::GetSelectedUniverse)
+				.OnValueChanged(this, &SDMXFixturePatcher::SelectUniverse)
+			]
+		);
+	}
+	ToolbarBuilder.EndSection();
+
+	ToolbarBuilder.BeginSection("ExpandedSettings");
+	{		
+		ToolbarBuilder.AddWidget(
+			SAssignNew(ShowAllUniversesCheckBox, SCheckBox)
+			.IsChecked(false)
+			.OnCheckStateChanged(this, &SDMXFixturePatcher::OnToggleDisplayAllUniverses)
+			[
+				SNew(STextBlock)
+				.Font(FAppStyle::GetFontStyle("PropertyWindow.NormalFont"))
+				.Text(LOCTEXT("ShowAllPatchedUniversesLabel", "Show all patched Universes"))
+			]
+		);
+
+		const UDMXEditorSettings* DMXEditorSettings = GetDefault<UDMXEditorSettings>();
+		const ECheckBoxState InitialDMXMonitorEnabledCheckBoxState = DMXEditorSettings->FixturePatcherSettings.bMonitorEnabled ? ECheckBoxState::Checked : ECheckBoxState::Unchecked;
+
+		ToolbarBuilder.AddWidget(
+			SAssignNew(EnableDMXMonitorCheckBox, SCheckBox)
+			.IsChecked(InitialDMXMonitorEnabledCheckBoxState)
+			.OnCheckStateChanged(this, &SDMXFixturePatcher::OnToggleDMXMonitorEnabled)
+			[
+				SNew(STextBlock)
+				.Font(FAppStyle::GetFontStyle("PropertyWindow.NormalFont"))
+				.Text(LOCTEXT("EnableInputMonitorLabel", "Monitor DMX Inputs"))
+				.ToolTipText(LOCTEXT("EnableInputMonitorTooltip", "If checked, monitors DMX Input Ports used in this DMX Library"))
+			]
+		);
+	}
+	ToolbarBuilder.EndSection();
+
+	ToolbarBuilder.BeginSection("AdvancedSettings");
+	{
+		ToolbarBuilder.AddWidget(
+			SNew(SComboButton)
+			.ComboButtonStyle(FAppStyle::Get(), "SimpleComboButtonWithIcon") // Use the tool bar item style for this button
+			.OnGetMenuContent(this, &SDMXFixturePatcher::CreateDisplaySettingsMenu)
+			.HasDownArrow(false)
+			.ButtonContent()
+			[
+				SNew(SImage)
+				.ColorAndOpacity(FSlateColor::UseForeground())
+				.Image(FAppStyle::Get().GetBrush("Icons.Settings"))
+			]
+		);
+	}
+	ToolbarBuilder.EndSection();
+
+	return ToolbarBuilder.MakeWidget();
+}
+
+TSharedRef<SWidget> SDMXFixturePatcher::CreateDisplaySettingsMenu()
+{
+	constexpr bool bCloseMenuWindowAfterMenuSelection = true;
+	FMenuBuilder MenuBuilder(bCloseMenuWindowAfterMenuSelection, nullptr);
+
+	MenuBuilder.BeginSection("PatchNameSection", LOCTEXT("PatchNameSectionLabel", "Displayed Patch Name"));
+	{
+		MenuBuilder.AddMenuEntry(
+			LOCTEXT("FixtureID", "Fixture ID"),
+			LOCTEXT("FixtureIDTooltip", "Displays the Fixture IDs of the Patches"),
+			FSlateIcon(),
+			FUIAction(
+				FExecuteAction::CreateSP(this, &SDMXFixturePatcher::SelectFixturePatchNameDisplayMode, EDMXFixturePatcherNameDisplayMode::FixtureID),
+				FCanExecuteAction(),
+				FIsActionChecked::CreateSP(this, &SDMXFixturePatcher::IsUsingFixturePatchNameDisplayMode, EDMXFixturePatcherNameDisplayMode::FixtureID)
+			),
+			NAME_None,
+			EUserInterfaceActionType::RadioButton
+		);
+
+		MenuBuilder.AddMenuEntry(
+			LOCTEXT("DisplayName", "Patch Name"),
+			LOCTEXT("DisplayNameTooltip", "Displays the Names of the Patches"),
+			FSlateIcon(),
+			FUIAction(
+				FExecuteAction::CreateSP(this, &SDMXFixturePatcher::SelectFixturePatchNameDisplayMode, EDMXFixturePatcherNameDisplayMode::FixturePatchName),
+				FCanExecuteAction(),
+				FIsActionChecked::CreateSP(this, &SDMXFixturePatcher::IsUsingFixturePatchNameDisplayMode, EDMXFixturePatcherNameDisplayMode::FixturePatchName)
+			),
+			NAME_None,
+			EUserInterfaceActionType::RadioButton
+		);
+
+		MenuBuilder.AddMenuEntry(
+			LOCTEXT("FixtureIDAndName", "Fixture ID and Patch Name"),
+			LOCTEXT("FixtureIDAndNameTooltip", "Displays the Fixture IDs and Names of the Patches"),
+			FSlateIcon(),
+			FUIAction(
+				FExecuteAction::CreateSP(this, &SDMXFixturePatcher::SelectFixturePatchNameDisplayMode, EDMXFixturePatcherNameDisplayMode::FixtureIDAndFixturePatchName),
+				FCanExecuteAction(),
+				FIsActionChecked::CreateSP(this, &SDMXFixturePatcher::IsUsingFixturePatchNameDisplayMode, EDMXFixturePatcherNameDisplayMode::FixtureIDAndFixturePatchName)
+			),
+			NAME_None,
+			EUserInterfaceActionType::RadioButton
+		);
+	}
+	MenuBuilder.EndSection();
+
+	return MenuBuilder.MakeWidget();
+}
+
+void SDMXFixturePatcher::SelectFixturePatchNameDisplayMode(EDMXFixturePatcherNameDisplayMode DisplayMode)
+{
+	UDMXEditorSettings* EditorSettings = GetMutableDefault<UDMXEditorSettings>();
+	EditorSettings->FixturePatcherSettings.FixturePatchNameDisplayMode = DisplayMode;
+	EditorSettings->SaveConfig();
+}
+
+bool SDMXFixturePatcher::IsUsingFixturePatchNameDisplayMode(EDMXFixturePatcherNameDisplayMode DisplayMode) const
+{
+	const UDMXEditorSettings* EditorSettings = GetDefault<UDMXEditorSettings>();
+	return DisplayMode == EditorSettings->FixturePatcherSettings.FixturePatchNameDisplayMode;
+}
+
 void SDMXFixturePatcher::OnDragEnterChannel(int32 UniverseID, int32 ChannelID, const FDragDropEvent& DragDropEvent)
 {
 	// Same as dropping onto a channel, but without FReply
@@ -616,7 +677,7 @@ void SDMXFixturePatcher::OnToggleDisplayAllUniverses(ECheckBoxState CheckboxStat
 void SDMXFixturePatcher::SetDMXMonitorEnabled(bool bEnabled)
 {
 	UDMXEditorSettings* DMXEditorSettings = GetMutableDefault<UDMXEditorSettings>();
-	DMXEditorSettings->bFixturePatcherDMXMonitorEnabled = bEnabled;
+	DMXEditorSettings->FixturePatcherSettings.bMonitorEnabled = bEnabled;
 	DMXEditorSettings->SaveConfig();
 
 	for (const TTuple<int32, TSharedPtr<SDMXPatchedUniverse>>& UniverseToWidgetPair : PatchedUniversesByID)
diff --git a/Engine/Plugins/VirtualProduction/DMX/DMXEngine/Source/DMXEditor/Private/Widgets/FixturePatch/SDMXFixturePatcher.h b/Engine/Plugins/VirtualProduction/DMX/DMXEngine/Source/DMXEditor/Private/Widgets/FixturePatch/SDMXFixturePatcher.h
--- a/Engine/Plugins/VirtualProduction/DMX/DMXEngine/Source/DMXEditor/Private/Widgets/FixturePatch/SDMXFixturePatcher.h
+++ b/Engine/Plugins/VirtualProduction/DMX/DMXEngine/Source/DMXEditor/Private/Widgets/FixturePatch/SDMXFixturePatcher.h
@@ -9,6 +9,7 @@
 #include "Widgets/DeclarativeSyntaxSupport.h"
 #include "Widgets/SCompoundWidget.h"
 
+enum class EDMXFixturePatcherNameDisplayMode : uint8;
 class FDMXEditor;
 class FDMXEntityFixturePatchDragDropOperation;
 class FDMXFixturePatchNode;
@@ -55,6 +56,19 @@ protected:
 	virtual FReply OnDrop(const FGeometry& MyGeometry, const FDragDropEvent& DragDropEvent) override;
 	// ~End SWidget Interface
 
+private:
+	/** */
+	TSharedRef<SWidget> CreateToolbar();
+
+	/** Creates the settings menu */
+	TSharedRef<SWidget> CreateDisplaySettingsMenu();
+
+	/** Selects the fixture patch name display mode */
+	void SelectFixturePatchNameDisplayMode(EDMXFixturePatcherNameDisplayMode DisplayMode);
+	
+	/** Returns true if the widget is currently using specified name display mode */
+	bool IsUsingFixturePatchNameDisplayMode(EDMXFixturePatcherNameDisplayMode DisplayMode) const;
+
 	/** Called when drag enters a channel */
 	void OnDragEnterChannel(int32 UniverseID, int32 ChannelID, const FDragDropEvent& DragDropEvent);
 
diff --git a/Engine/Plugins/VirtualProduction/DMX/DMXEngine/Source/DMXEditor/Private/Widgets/FixturePatch/SDMXMVRFixtureList.cpp b/Engine/Plugins/VirtualProduction/DMX/DMXEngine/Source/DMXEditor/Private/Widgets/FixturePatch/SDMXMVRFixtureList.cpp
--- a/Engine/Plugins/VirtualProduction/DMX/DMXEngine/Source/DMXEditor/Private/Widgets/FixturePatch/SDMXMVRFixtureList.cpp
+++ b/Engine/Plugins/VirtualProduction/DMX/DMXEngine/Source/DMXEditor/Private/Widgets/FixturePatch/SDMXMVRFixtureList.cpp
@@ -5,6 +5,7 @@
 #include "DMXEditor.h"
 #include "DMXEditorSettings.h"
 #include "DMXEditorUtils.h"
+#include "DMXFixturePatchAutoAssignUtility.h"
 #include "DMXFixturePatchSharedData.h"
 #include "DMXRuntimeUtils.h"
 #include "Commands/DMXEditorCommands.h"
@@ -556,9 +557,9 @@ private:
 	TArray<TSharedPtr<FDMXMVRFixtureListItem>> Items;
 };
 
-
-const FName FDMXMVRFixtureListCollumnIDs::Status = "Status";
+const FName FDMXMVRFixtureListCollumnIDs::EditorColor = "EditorColor";
 const FName FDMXMVRFixtureListCollumnIDs::FixturePatchName = "FixturePatchName";
+const FName FDMXMVRFixtureListCollumnIDs::Status = "Status";
 const FName FDMXMVRFixtureListCollumnIDs::FixtureID = "FixtureID";
 const FName FDMXMVRFixtureListCollumnIDs::FixtureType = "FixtureType";
 const FName FDMXMVRFixtureListCollumnIDs::Mode = "Mode";
@@ -940,41 +941,54 @@ void SDMXMVRFixtureList::AdoptSelectionFromFixturePatchSharedData()
 	}
 }
 
-void SDMXMVRFixtureList::AutoAssignFixturePatches()
+void SDMXMVRFixtureList::AutoAssignFixturePatches(EDMXFixturePatchAutoAssignMode AutoAssignMode)
 {
-	if (FixturePatchSharedData.IsValid())
+	TArray<UDMXEntityFixturePatch*> FixturePatchesToAutoAssign;
+	const TArray<TWeakObjectPtr<UDMXEntityFixturePatch>> SelectedFixturePatches = FixturePatchSharedData->GetSelectedFixturePatches();
+	for (TWeakObjectPtr<UDMXEntityFixturePatch> FixturePatch : SelectedFixturePatches)
 	{
-		TArray<UDMXEntityFixturePatch*> FixturePatchesToAutoAssign;
-		const TArray<TWeakObjectPtr<UDMXEntityFixturePatch>> SelectedFixturePatches = FixturePatchSharedData->GetSelectedFixturePatches();
-		for (TWeakObjectPtr<UDMXEntityFixturePatch> FixturePatch : SelectedFixturePatches)
+		if (FixturePatch.IsValid())
 		{
-			if (FixturePatch.IsValid())
-			{
-				FixturePatchesToAutoAssign.Add(FixturePatch.Get());
-			}
+			FixturePatchesToAutoAssign.Add(FixturePatch.Get());
 		}
+	}
 
-		if (FixturePatchesToAutoAssign.IsEmpty())
-		{
-			return;
-		}
+	if (FixturePatchesToAutoAssign.IsEmpty())
+	{
+		return;
+	}
 
-		constexpr bool bAllowDecrementUniverse = false;
-		constexpr bool bAllowDecrementChannels = true;
-		FDMXEditorUtils::AutoAssignedChannels(bAllowDecrementUniverse, bAllowDecrementChannels, FixturePatchesToAutoAssign);
-		FixturePatchSharedData->SelectUniverse(FixturePatchesToAutoAssign[0]->GetUniverseID());
+	FDMXFixturePatchAutoAssignUtility::AutoAssign(FixturePatchesToAutoAssign, AutoAssignMode);
+	FixturePatchSharedData->SelectUniverse(FixturePatchesToAutoAssign[0]->GetUniverseID());
 
-		RequestListRefresh();
+	RequestListRefresh();
+}
+
+bool SDMXMVRFixtureList::DoesDMXLibraryHaveReachableUniverses() const
+{
+	UDMXLibrary* DMXLibrary = WeakDMXEditor.IsValid() ? WeakDMXEditor.Pin()->GetDMXLibrary() : nullptr;
+	if (DMXLibrary)
+	{
+		return !DMXLibrary->GetInputPorts().IsEmpty() && !DMXLibrary->GetOutputPorts().IsEmpty();
 	}
+	return false;
 }
 
 TSharedRef<SHeaderRow> SDMXMVRFixtureList::GenerateHeaderRow()
 {
 	const float StatusColumnWidth = FMath::Max(FAppStyle::GetBrush("Icons.Warning")->GetImageSize().X + 6.f, FAppStyle::GetBrush("Icons.Error")->GetImageSize().X + 6.f);
+	const float EditorColorColumnWidth = StatusColumnWidth;
 
 	HeaderRow = SNew(SHeaderRow);
 	SHeaderRow::FColumn::FArguments ColumnArgs;
 
+	HeaderRow->AddColumn(
+		SHeaderRow::FColumn::FArguments()
+		.ColumnId(FDMXMVRFixtureListCollumnIDs::EditorColor)
+		.DefaultLabel(LOCTEXT("EditorColorColumnLabel", ""))
+		.FixedWidth(EditorColorColumnWidth)
+	);
+
 	HeaderRow->AddColumn(
 		SHeaderRow::FColumn::FArguments()
 		.ColumnId(FDMXMVRFixtureListCollumnIDs::FixturePatchName)
@@ -1215,17 +1229,14 @@ TSharedPtr<SWidget> SDMXMVRFixtureList::OnContextMenuOpening()
 	FMenuBuilder MenuBuilder(bCloseWindowAfterMenuSelection, CommandList);
 
 	// Auto Assign Section
-	MenuBuilder.BeginSection("AutoAssignSection", LOCTEXT("AutoAssignSection", "Auto-Assign"));
+	if (!FixturePatchSharedData->GetSelectedFixturePatches().IsEmpty())
 	{
-		// Auto Assign Entry
-		const FUIAction Action(FExecuteAction::CreateSP(this, &SDMXMVRFixtureList::AutoAssignFixturePatches));
-
-		const FText AutoAssignText = LOCTEXT("AutoAssignContextMenuEntry", "Auto-Assign Selection");
-		const TSharedRef<SWidget> Widget =
-			SNew(STextBlock)
-			.Text(AutoAssignText);
-
-		MenuBuilder.AddMenuEntry(Action, Widget);
+		MenuBuilder.BeginSection("AutoAssignSection", LOCTEXT("AutoAssignSection", "Auto-Assign"));
+		{
+			MenuBuilder.AddMenuEntry(FDMXEditorCommands::Get().AutoAssignToFirstUniverseInSelection);
+			MenuBuilder.AddMenuEntry(FDMXEditorCommands::Get().AutoAssignToFirstReachableUniverse);
+			MenuBuilder.AddMenuEntry(FDMXEditorCommands::Get().AutoAssignAfterLastAddressInLibrary);
+		}
 		MenuBuilder.EndSection();
 	}
 
@@ -1250,7 +1261,6 @@ void SDMXMVRFixtureList::RegisterCommands()
 		return;
 	}
 
-	// listen to common editor shortcuts for copy/paste etc
 	CommandList = MakeShared<FUICommandList>();
 	CommandList->MapAction(FGenericCommands::Get().Cut, 
 		FUIAction(
@@ -1283,6 +1293,24 @@ void SDMXMVRFixtureList::RegisterCommands()
 			EUIActionRepeatMode::RepeatEnabled
 		)
 	);
+	CommandList->MapAction
+	(
+		FDMXEditorCommands::Get().AutoAssignToFirstUniverseInSelection,
+		FExecuteAction::CreateSP(this, &SDMXMVRFixtureList::AutoAssignFixturePatches, EDMXFixturePatchAutoAssignMode::FirstUniverseInSelection)
+	);
+
+	CommandList->MapAction
+	(
+		FDMXEditorCommands::Get().AutoAssignToFirstReachableUniverse,
+		FExecuteAction::CreateSP(this, &SDMXMVRFixtureList::AutoAssignFixturePatches, EDMXFixturePatchAutoAssignMode::FirstReachableUniverse),
+		FCanExecuteAction::CreateSP(this, &SDMXMVRFixtureList::DoesDMXLibraryHaveReachableUniverses)
+	);
+
+	CommandList->MapAction
+	(
+		FDMXEditorCommands::Get().AutoAssignAfterLastAddressInLibrary,
+		FExecuteAction::CreateSP(this, &SDMXMVRFixtureList::AutoAssignFixturePatches, EDMXFixturePatchAutoAssignMode::LastAddressInLibrary)
+	);
 }
 
 bool SDMXMVRFixtureList::CanCutItems() const
diff --git a/Engine/Plugins/VirtualProduction/DMX/DMXEngine/Source/DMXEditor/Private/Widgets/FixturePatch/SDMXMVRFixtureList.h b/Engine/Plugins/VirtualProduction/DMX/DMXEngine/Source/DMXEditor/Private/Widgets/FixturePatch/SDMXMVRFixtureList.h
--- a/Engine/Plugins/VirtualProduction/DMX/DMXEngine/Source/DMXEditor/Private/Widgets/FixturePatch/SDMXMVRFixtureList.h
+++ b/Engine/Plugins/VirtualProduction/DMX/DMXEngine/Source/DMXEditor/Private/Widgets/FixturePatch/SDMXMVRFixtureList.h
@@ -8,6 +8,7 @@
 #include "Widgets/SCompoundWidget.h"
 #include "Widgets/Views/SHeaderRow.h"
 
+enum class EDMXFixturePatchAutoAssignMode : uint8;
 class FDMXEditor;
 class FDMXMVRFixtureListItem;
 class FDMXFixturePatchSharedData;
@@ -31,8 +32,9 @@ class STableViewBase;
 /** Collumn IDs in the Fixture Patch List */
 struct FDMXMVRFixtureListCollumnIDs
 {
-	static const FName Status;
+	static const FName EditorColor;
 	static const FName FixturePatchName;
+	static const FName Status;
 	static const FName FixtureID;
 	static const FName FixtureType;
 	static const FName Mode;
@@ -113,7 +115,10 @@ private:
 	void AdoptSelectionFromFixturePatchSharedData();
 
 	/** Auto assigns selected Fixture Patches */
-	void AutoAssignFixturePatches();
+	void AutoAssignFixturePatches(EDMXFixturePatchAutoAssignMode AutoAssignMode);
+
+	/** Returns true if the DMX Library has reachable Universes */
+	bool DoesDMXLibraryHaveReachableUniverses() const;
 
 	/** Sets keyboard focus on this widget */
 	void SetKeyboardFocus();
diff --git a/Engine/Plugins/VirtualProduction/DMX/DMXEngine/Source/DMXEditor/Private/Widgets/FixturePatch/SDMXMVRFixtureListRow.cpp b/Engine/Plugins/VirtualProduction/DMX/DMXEngine/Source/DMXEditor/Private/Widgets/FixturePatch/SDMXMVRFixtureListRow.cpp
--- a/Engine/Plugins/VirtualProduction/DMX/DMXEngine/Source/DMXEditor/Private/Widgets/FixturePatch/SDMXMVRFixtureListRow.cpp
+++ b/Engine/Plugins/VirtualProduction/DMX/DMXEngine/Source/DMXEditor/Private/Widgets/FixturePatch/SDMXMVRFixtureListRow.cpp
@@ -63,6 +63,7 @@ public:
 				[
 					SNew(STextBlock)
 					.Text(this, &SDMXMVRFixtureFixtureTypePicker::GetSelectedItemText)
+					.Font(FAppStyle::GetFontStyle(TEXT("PropertyWindow.NormalFont")))
 				]
 			];
 
@@ -276,6 +277,7 @@ public:
 						})
 					.OptionsSource(&ModeNames)
 					.OnSelectionChanged(this, &SDMXMVRFixtureModePicker::OnSelectionChanged)
+					.Font(FAppStyle::GetFontStyle(TEXT("PropertyWindow.NormalFont")))
 				]
 
 				+ SVerticalBox::Slot()
@@ -413,20 +415,10 @@ void SDMXMVRFixtureListRow::Construct(const FArguments& InArgs, const TSharedRef
 	OnRowRequestsStatusRefresh = InArgs._OnRowRequestsStatusRefresh;
 	OnRowRequestsListRefresh = InArgs._OnRowRequestsListRefresh;
 	IsSelected = InArgs._IsSelected;
-	
-	SetBorderBackgroundColor(TAttribute<FSlateColor>::CreateLambda([this]()
-		{
-			if (Item.IsValid())
-			{
-				return Item->GetBackgroundColor();
-			}
-			return FLinearColor::Black;
-		}));
-
 
 	SMultiColumnTableRow<TSharedPtr<FDMXMVRFixtureListItem>>::Construct(
 		FSuperRowType::FArguments()
-		.Style(&FDMXEditorStyle::Get().GetWidgetStyle<FTableRowStyle>("MVRFixtureList.Row")), 
+		.Style(&FDMXEditorStyle::Get().GetWidgetStyle<FTableRowStyle>("FixturePatchList.Row")),
 		InOwnerTable);
 }
 
@@ -440,35 +432,56 @@ void SDMXMVRFixtureListRow::EnterFixturePatchNameEditingMode()
 
 TSharedRef<SWidget> SDMXMVRFixtureListRow::GenerateWidgetForColumn(const FName& ColumnName)
 {
+	if (ColumnName == FDMXMVRFixtureListCollumnIDs::EditorColor)
+	{
+		return GenerateEditorColorWidget();
+	}
 	if (ColumnName == FDMXMVRFixtureListCollumnIDs::FixturePatchName)
 	{
-		return GenerateFixturePatchNameRow();
+		return GenerateFixturePatchNameWidget();
 	}
 	else if (ColumnName == FDMXMVRFixtureListCollumnIDs::Status)
 	{
-		return GenerateStatusRow();
+		return GenerateStatusWidget();
 	}
 	else if (ColumnName == FDMXMVRFixtureListCollumnIDs::FixtureID)
 	{
-		return GenerateFixtureIDRow();
+		return GenerateFixtureIDWidget();
 	}
 	else if (ColumnName == FDMXMVRFixtureListCollumnIDs::FixtureType)
 	{
-		return GenerateFixtureTypeRow();
+		return GenerateFixtureTypeWidget();
 	}
 	else if (ColumnName == FDMXMVRFixtureListCollumnIDs::Mode)
 	{
-		return GenerateModeRow();
+		return GenerateModeWidget();
 	}
 	else if (ColumnName == FDMXMVRFixtureListCollumnIDs::Patch)
 	{
-		return GeneratePatchRow();
+		return GeneratePatchWidget();
 	}
 
 	return SNullWidget::NullWidget;
 }
 
-TSharedRef<SWidget> SDMXMVRFixtureListRow::GenerateFixturePatchNameRow()
+TSharedRef<SWidget> SDMXMVRFixtureListRow::GenerateEditorColorWidget()
+{
+	return
+		SNew(SBorder)
+		.HAlign(HAlign_Fill)
+		.Padding(5.f, 2.f)
+		.BorderImage(FAppStyle::GetBrush("NoBorder"))
+		[
+			SNew(SImage)
+			.Image(FDMXEditorStyle::Get().GetBrush("DMXEditor.WhiteRoundedPropertyBorder"))
+			.ColorAndOpacity_Lambda([this]()
+			{
+				return Item->GetBackgroundColor();
+			})
+		];
+}
+
+TSharedRef<SWidget> SDMXMVRFixtureListRow::GenerateFixturePatchNameWidget()
 {
 	return
 		SAssignNew(FixturePatchNameBorder, SBorder)
@@ -489,7 +502,7 @@ TSharedRef<SWidget> SDMXMVRFixtureListRow::GenerateFixturePatchNameRow()
 					const FString FixturePatchName = Item->GetFixturePatchName();
 					return FText::FromString(FixturePatchName);
 				})
-				.Font(FCoreStyle::GetDefaultFontStyle("Regular", 10))
+				.Font(FAppStyle::GetFontStyle(TEXT("PropertyWindow.NormalFont")))
 				.OnTextCommitted(this, &SDMXMVRFixtureListRow::OnFixturePatchNameCommitted)
 				.IsSelected(IsSelected)
 			]
@@ -521,7 +534,7 @@ void SDMXMVRFixtureListRow::OnFixturePatchNameCommitted(const FText& InNewText,
 	FixturePatchNameTextBlock->SetText(FText::FromString(ResultingName));
 }
 
-TSharedRef<SWidget> SDMXMVRFixtureListRow::GenerateStatusRow()
+TSharedRef<SWidget> SDMXMVRFixtureListRow::GenerateStatusWidget()
 {
 	return
 		SNew(SBox)
@@ -531,14 +544,9 @@ TSharedRef<SWidget> SDMXMVRFixtureListRow::GenerateStatusRow()
 			SNew(SImage)
 			.Image_Lambda([this]()
 				{
-					if (!Item->ErrorStatusText.IsEmpty())
-					{
-						return FAppStyle::GetBrush("Icons.Error");
-					}
-
-					if (!Item->WarningStatusText.IsEmpty())
+					if (!Item->ErrorStatusText.IsEmpty() || !Item->WarningStatusText.IsEmpty())
 					{
-						return FAppStyle::GetBrush("Icons.Warning");
+						return FDMXEditorStyle::Get().GetBrush("Icons.WarningExclamationMark");
 					}
 
 					static const FSlateBrush EmptyBrush = FSlateNoResource();
@@ -560,7 +568,7 @@ TSharedRef<SWidget> SDMXMVRFixtureListRow::GenerateStatusRow()
 		];
 }
 
-TSharedRef<SWidget> SDMXMVRFixtureListRow::GenerateFixtureIDRow()
+TSharedRef<SWidget> SDMXMVRFixtureListRow::GenerateFixtureIDWidget()
 {
 	return
 		SNew(SBorder)
@@ -580,7 +588,7 @@ TSharedRef<SWidget> SDMXMVRFixtureListRow::GenerateFixtureIDRow()
 				{
 					return FText::FromString(Item->GetFixtureID());
 				})
-				.Font(FCoreStyle::GetDefaultFontStyle("Regular", 10))
+				.Font(FAppStyle::GetFontStyle(TEXT("PropertyWindow.NormalFont")))
 				.OnTextCommitted(this, &SDMXMVRFixtureListRow::OnFixtureIDCommitted)
 				.IsSelected(IsSelected)
 			]
@@ -615,7 +623,7 @@ void SDMXMVRFixtureListRow::OnFixtureIDCommitted(const FText& InNewText, ETextCo
 	}
 }
 
-TSharedRef<SWidget> SDMXMVRFixtureListRow::GenerateFixtureTypeRow()
+TSharedRef<SWidget> SDMXMVRFixtureListRow::GenerateFixtureTypeWidget()
 {
 	UDMXLibrary* DMXLibrary = Item->GetDMXLibrary();
 	if (!ensureAlwaysMsgf(DMXLibrary, TEXT("Tried to set fixture type for MVR Fixture, but fixture type is invalid.")))
@@ -646,7 +654,7 @@ void SDMXMVRFixtureListRow::OnFixtureTypeSelected(UDMXEntityFixtureType* Selecte
 	Item->SetFixtureType(SelectedFixtureType);
 }
 
-TSharedRef<SWidget> SDMXMVRFixtureListRow::GenerateModeRow()
+TSharedRef<SWidget> SDMXMVRFixtureListRow::GenerateModeWidget()
 {
 	const TSharedRef<SDMXMVRFixtureModePicker> ModePicker =
 		SNew(SDMXMVRFixtureModePicker)
@@ -673,7 +681,7 @@ void SDMXMVRFixtureListRow::OnModeSelected(int32 SelectedModeIndex)
 	Item->SetModeIndex(SelectedModeIndex);
 }
 
-TSharedRef<SWidget> SDMXMVRFixtureListRow::GeneratePatchRow()
+TSharedRef<SWidget> SDMXMVRFixtureListRow::GeneratePatchWidget()
 {
 	return
 		SNew(SBorder)
@@ -695,6 +703,7 @@ TSharedRef<SWidget> SDMXMVRFixtureListRow::GeneratePatchRow()
 						const int32 StartingAddress = Item->GetAddress();
 						return FText::Format(LOCTEXT("AddressesText", "{0}.{1}"), UniverseID, StartingAddress);
 					})
+				.Font(FAppStyle::GetFontStyle(TEXT("PropertyWindow.NormalFont")))
 				.OnTextCommitted(this, &SDMXMVRFixtureListRow::OnPatchCommitted)
 				.IsSelected(IsSelected)
 			]
diff --git a/Engine/Plugins/VirtualProduction/DMX/DMXEngine/Source/DMXEditor/Private/Widgets/FixturePatch/SDMXMVRFixtureListRow.h b/Engine/Plugins/VirtualProduction/DMX/DMXEngine/Source/DMXEditor/Private/Widgets/FixturePatch/SDMXMVRFixtureListRow.h
--- a/Engine/Plugins/VirtualProduction/DMX/DMXEngine/Source/DMXEditor/Private/Widgets/FixturePatch/SDMXMVRFixtureListRow.h
+++ b/Engine/Plugins/VirtualProduction/DMX/DMXEngine/Source/DMXEditor/Private/Widgets/FixturePatch/SDMXMVRFixtureListRow.h
@@ -2,13 +2,12 @@
 
 #pragma once
 
-#include "Library/DMXEntityReference.h"
-
-#include "CoreMinimal.h"
+#include "Widgets/Views/SListView.h"
 #include "Widgets/Views/STableRow.h"
 
 enum class EDMXMVRFixtureListEditMode : uint8;
 class UDMXEntityFixturePatch;
+class UDMXEntityFixtureType;
 class FDMXMVRFixtureListItem;
 
 class SInlineEditableTextBlock;
@@ -47,8 +46,14 @@ protected:
 	//~ End SMultiColumnTableRow interface
 
 private:
-	/** Generates the row that displays the Fixture Patch Name */
-	TSharedRef<SWidget> GenerateFixturePatchNameRow();
+	/** Generates the Widget that displays the Editor Color */
+	TSharedRef<SWidget> GenerateEditorColorWidget();
+
+	/** Gets Fixture Patch editor color */
+	FSlateColor GetFixturePatchEditorColor() const;
+
+	/** Generates the Widget that displays the Fixture Patch Name */
+	TSharedRef<SWidget> GenerateFixturePatchNameWidget();
 
 	/** Called when the Fixture Patch Name Border was double-clicked */
 	FReply OnFixturePatchNameBorderDoubleClicked(const FGeometry& InMyGeometry, const FPointerEvent& InMouseEvent);
@@ -56,11 +61,11 @@ private:
 	/** Called when a Fixture Patch Name was committed */
 	void OnFixturePatchNameCommitted(const FText& InNewText, ETextCommit::Type InTextCommit);
 
-	/** Generates the row that displays the Status */
-	TSharedRef<SWidget> GenerateStatusRow();
+	/** Generates the widget that displays the Status */
+	TSharedRef<SWidget> GenerateStatusWidget();
 
-	/** Generates the row that displays the Fixture ID */
-	TSharedRef<SWidget> GenerateFixtureIDRow();
+	/** Generates the Widget that displays the Fixture ID */
+	TSharedRef<SWidget> GenerateFixtureIDWidget();
 
 	/** Called when the Fixture ID Border was double-clicked */
 	FReply OnFixtureIDBorderDoubleClicked(const FGeometry& InMyGeometry, const FPointerEvent& InMouseEvent);
@@ -68,20 +73,20 @@ private:
 	/** Called when a Fixture ID was committed */
 	void OnFixtureIDCommitted(const FText& InNewText, ETextCommit::Type InTextCommit);
 
-	/** Generates the row that displays the Fixture Type */
-	TSharedRef<SWidget> GenerateFixtureTypeRow();
+	/** Generates the Widget that displays the Fixture Type */
+	TSharedRef<SWidget> GenerateFixtureTypeWidget();
 
 	/** Called when a Fixture Type was selected */
 	void OnFixtureTypeSelected(UDMXEntityFixtureType* SelectedFixtureType);
 
-	/** Generates the row that displays the Mode */
-	TSharedRef<SWidget> GenerateModeRow();
+	/** Generates the Widget that displays the Mode */
+	TSharedRef<SWidget> GenerateModeWidget();
 
 	/** Called when a Mode was selected */
 	void OnModeSelected(int32 SelectedModeIndex);
 
-	/** Generates the row that displays the Patch */
-	TSharedRef<SWidget> GeneratePatchRow();
+	/** Generates the Widget that displays the Patch */
+	TSharedRef<SWidget> GeneratePatchWidget();
 
 	/** Called when the Patch Border was double-clicked */
 	FReply OnPatchBorderDoubleClicked(const FGeometry& InMyGeometry, const FPointerEvent& InMouseEvent);
diff --git a/Engine/Plugins/VirtualProduction/DMX/DMXEngine/Source/DMXEditor/Private/Widgets/FixturePatch/SDMXMVRFixtureListToolbar.cpp b/Engine/Plugins/VirtualProduction/DMX/DMXEngine/Source/DMXEditor/Private/Widgets/FixturePatch/SDMXMVRFixtureListToolbar.cpp
--- a/Engine/Plugins/VirtualProduction/DMX/DMXEngine/Source/DMXEditor/Private/Widgets/FixturePatch/SDMXMVRFixtureListToolbar.cpp
+++ b/Engine/Plugins/VirtualProduction/DMX/DMXEngine/Source/DMXEditor/Private/Widgets/FixturePatch/SDMXMVRFixtureListToolbar.cpp
@@ -6,6 +6,7 @@
 #include "DMXEditorUtils.h"
 #include "DMXFixturePatchSharedData.h"
 #include "DMXMVRFixtureListItem.h"
+#include "DMXFixturePatchAutoAssignUtility.h"
 #include "Library/DMXEntityFixturePatch.h"
 #include "Library/DMXEntityFixtureType.h"
 #include "Commands/DMXEditorCommands.h"
@@ -351,9 +352,7 @@ void SDMXMVRFixtureListToolbar::OnAddNewMVRFixtureClicked(UDMXEntity* InSelected
 	}
 
 	// Auto assign
-	constexpr bool bAllowDecrementUniverse = false;
-	constexpr bool bAllowDecrementChannels = false;
-	FDMXEditorUtils::AutoAssignedChannels(bAllowDecrementUniverse, bAllowDecrementChannels, NewFixturePatches);
+	FDMXFixturePatchAutoAssignUtility::AutoAssign(NewFixturePatches, EDMXFixturePatchAutoAssignMode::LastAddressInLibrary);
 
 	DMXLibrary->PostEditChange();
 
diff --git a/Engine/Plugins/VirtualProduction/DMX/DMXEngine/Source/DMXEditor/Private/Widgets/FixturePatch/SDMXPatchedUniverse.cpp b/Engine/Plugins/VirtualProduction/DMX/DMXEngine/Source/DMXEditor/Private/Widgets/FixturePatch/SDMXPatchedUniverse.cpp
--- a/Engine/Plugins/VirtualProduction/DMX/DMXEngine/Source/DMXEditor/Private/Widgets/FixturePatch/SDMXPatchedUniverse.cpp
+++ b/Engine/Plugins/VirtualProduction/DMX/DMXEngine/Source/DMXEditor/Private/Widgets/FixturePatch/SDMXPatchedUniverse.cpp
@@ -3,14 +3,15 @@
 #include "SDMXPatchedUniverse.h"
 
 #include "DMXEditor.h"
-#include "DMXEditorUtils.h"
+#include "DMXFixturePatchAutoAssignUtility.h"
 #include "DMXFixturePatchEditorDefinitions.h"
 #include "DMXFixturePatchNode.h"
 #include "DMXFixturePatchSharedData.h"
-#include "Framework/Application/SlateApplication.h"
 #include "SDMXChannelConnector.h"
 #include "SDMXFixturePatchFragment.h"
+#include "Commands/DMXEditorCommands.h"
 #include "DragDrop/DMXEntityFixturePatchDragDropOp.h"
+#include "Framework/Application/SlateApplication.h"
 #include "IO/DMXInputPort.h"
 #include "IO/DMXOutputPort.h"
 #include "IO/DMXPortManager.h"
@@ -94,7 +95,7 @@ void SDMXPatchedUniverse::Construct(const FArguments& InArgs)
 			.MinDesiredWidth(920.0f) // Avoids a slate issue, where SScaleBox ignores border size hence scales			
 			.Padding(FMargin(4.0f, 4.0f, 4.0f, 8.0f))
 			[
-				SAssignNew(Grid, SGridPanel)												
+				SAssignNew(Grid, SGridPanel)										
 			]
 		]
 	];
@@ -147,35 +148,32 @@ void SDMXPatchedUniverse::CreateChannelConnectors()
 		int32 Channel = ChannelIndex + FDMXChannelGridSpecs::ChannelIDOffset;
 
 		TSharedRef<SDMXChannelConnector> ChannelPatchWidget =
-			SNew(SDMXChannelConnector)
+			SNew(SDMXChannelConnector, DMXEditorPtr.Pin())
 			.ChannelID(Channel)
 			.Value(0)
+			.OnHovered(this, &SDMXPatchedUniverse::HandleOnChannelHovered, Channel)
+			.OnUnhovered(this, &SDMXPatchedUniverse::HandleOnChannelUnhovered, Channel)
 			.OnMouseButtonDownOnChannel(this, &SDMXPatchedUniverse::HandleOnMouseButtonDownOnChannel)
 			.OnMouseButtonUpOnChannel(this, &SDMXPatchedUniverse::HandleOnMouseButtonUpOnChannel)
 			.OnDragDetectedOnChannel(this, &SDMXPatchedUniverse::HandleOnDragDetectedOnChannel)
 			.OnDragEnterChannel(this, &SDMXPatchedUniverse::HandleDragEnterChannel)
 			.OnDragLeaveChannel(this, &SDMXPatchedUniverse::HandleDragLeaveChannel)
-			.OnDropOntoChannel(this, &SDMXPatchedUniverse::HandleDropOntoChannel)
-			.DMXEditor(DMXEditorPtr);
+			.OnDropOntoChannel(this, &SDMXPatchedUniverse::HandleDropOntoChannel);
 
 		ChannelConnectors.Add(ChannelPatchWidget);
 
 		Grid->AddSlot(Column, Row)
 			.HAlign(HAlign_Fill)
 			.VAlign(VAlign_Fill)
+			.Layer(0)
 			[			
 				ChannelPatchWidget
 			];
 	}
 }
 
-bool SDMXPatchedUniverse::FindOrAdd(const TSharedPtr<FDMXFixturePatchNode>& Node)
+bool SDMXPatchedUniverse::FindOrAdd(const TSharedRef<FDMXFixturePatchNode>& Node)
 {
-	if (!Node.IsValid())
-	{
-		return false;
-	}
-
 	UDMXEntityFixturePatch* FixturePatch = Node->GetFixturePatch().Get();
 	if (!FixturePatch)
 	{
@@ -189,17 +187,19 @@ bool SDMXPatchedUniverse::FindOrAdd(const TSharedPtr<FDMXFixturePatchNode>& Node
 	return true;
 }
 
-void SDMXPatchedUniverse::Remove(const TSharedPtr<FDMXFixturePatchNode>& Node)
+void SDMXPatchedUniverse::Remove(const TSharedRef<FDMXFixturePatchNode>& Node)
 {
 	if (PatchedNodes.Contains(Node))
 	{
 		PatchedNodes.RemoveSingle(Node);
 
-		for (const TSharedPtr<SDMXFixturePatchFragment>& Widget : Node->GetGeneratedWidgets())
+		const TArray<TSharedRef<SDMXFixturePatchFragment>> CachedWidgets = FixturePatchFragmentWidgets;
+		for (const TSharedRef<SDMXFixturePatchFragment>& Widget : CachedWidgets)
 		{
-			if (Widget.IsValid())
+			if (Widget->GetFixturePatchNode() == Node)
 			{
-				Grid->RemoveSlot(Widget.ToSharedRef());
+				Grid->RemoveSlot(Widget);
+				FixturePatchFragmentWidgets.Remove(Widget);
 			}
 		}
 		RequestRefresh();
@@ -239,7 +239,7 @@ bool SDMXPatchedUniverse::CanAssignFixturePatch(TWeakObjectPtr<UDMXEntityFixture
 	check(StartingChannel + ChannelSpan - 1 <= DMX_UNIVERSE_SIZE);
 
 	// Test for overlapping nodes in this universe
-	for (const TSharedPtr<FDMXFixturePatchNode>& Node : PatchedNodes)
+	for (const TSharedRef<FDMXFixturePatchNode>& Node : PatchedNodes)
 	{
 		if (Node->GetFixturePatch() == FixturePatch)
 		{
@@ -266,7 +266,7 @@ bool SDMXPatchedUniverse::CanAssignNode(const TSharedPtr<FDMXFixturePatchNode>&
 
 TSharedPtr<FDMXFixturePatchNode> SDMXPatchedUniverse::FindPatchNode(TWeakObjectPtr<UDMXEntityFixturePatch> FixturePatch) const
 {
-	const TSharedPtr<FDMXFixturePatchNode>* NodePtr = PatchedNodes.FindByPredicate([&FixturePatch](const TSharedPtr<FDMXFixturePatchNode>& Node)
+	const TSharedRef<FDMXFixturePatchNode>* NodePtr = PatchedNodes.FindByPredicate([&FixturePatch](const TSharedPtr<FDMXFixturePatchNode>& Node)
 		{
 			if (Node->GetFixturePatch() == FixturePatch)
 			{
@@ -290,7 +290,7 @@ TSharedPtr<FDMXFixturePatchNode> SDMXPatchedUniverse::FindPatchNodeOfType(UDMXEn
 		return nullptr;
 	}
 
-	for (const TSharedPtr<FDMXFixturePatchNode>& PatchNode : PatchedNodes)
+	for (const TSharedRef<FDMXFixturePatchNode>& PatchNode : PatchedNodes)
 	{
 		if (PatchNode == IgoredNode)
 		{
@@ -309,6 +309,8 @@ TSharedPtr<FDMXFixturePatchNode> SDMXPatchedUniverse::FindPatchNodeOfType(UDMXEn
 
 void SDMXPatchedUniverse::Tick(const FGeometry& AllottedGeometry, const double InCurrentTime, const float InDeltaTime) 
 {
+	UpdateNodesHoveredState();
+
 	if (bMonitorInputs)
 	{
 		UpdateMonitor();
@@ -322,7 +324,7 @@ void SDMXPatchedUniverse::OnFixturePatchChanged(const UDMXEntityFixturePatch* Fi
 		return;
 	}
 
-	const TSharedPtr<FDMXFixturePatchNode>* NodePtr = PatchedNodes.FindByPredicate([FixturePatch](const TSharedPtr<FDMXFixturePatchNode>& Node)
+	const TSharedRef<FDMXFixturePatchNode>* NodePtr = PatchedNodes.FindByPredicate([FixturePatch](const TSharedPtr<FDMXFixturePatchNode>& Node)
 		{
 			return Node->GetFixturePatch() == FixturePatch;
 		});
@@ -358,8 +360,8 @@ void SDMXPatchedUniverse::SetUniverseIDInternal(int32 NewUniverseID)
 	if (Library)
 	{
 		// Unpatch all nodes
-		TArray<TSharedPtr<FDMXFixturePatchNode>> CachedPatchedNodes = PatchedNodes;
-		for (const TSharedPtr<FDMXFixturePatchNode>& Node : CachedPatchedNodes)
+		TArray<TSharedRef<FDMXFixturePatchNode>> CachedPatchedNodes = PatchedNodes;
+		for (const TSharedRef<FDMXFixturePatchNode>& Node : CachedPatchedNodes)
 		{
 			Remove(Node);
 		}
@@ -394,13 +396,12 @@ void SDMXPatchedUniverse::SetUniverseIDInternal(int32 NewUniverseID)
 			{
 				Node = FDMXFixturePatchNode::Create(DMXEditorPtr, FixturePatch);
 			}
-			check(Node.IsValid());
 
-			FindOrAdd(Node);
+			FindOrAdd(Node.ToSharedRef());
 		}
 
 		// Update the channel connectors' Universe ID
-		for (const TSharedPtr<SDMXChannelConnector>& Connector : ChannelConnectors)
+		for (const TSharedRef<SDMXChannelConnector>& Connector : ChannelConnectors)
 		{
 			Connector->SetUniverseID(NewUniverseID);
 		}
@@ -416,15 +417,15 @@ void SDMXPatchedUniverse::RefreshInternal()
 	UpdateZOrderOfNodes();
 
 	// Remove old widgets from the grid
-	for (const TSharedPtr<SDMXFixturePatchFragment>& OldWidget : FixturePatchWidgets)
+	for (const TSharedRef<SDMXFixturePatchFragment>& OldWidget : FixturePatchFragmentWidgets)
 	{
-		Grid->RemoveSlot(OldWidget.ToSharedRef());
+		Grid->RemoveSlot(OldWidget);
 	}
-	FixturePatchWidgets.Reset();
+	FixturePatchFragmentWidgets.Reset();
 
 	// Create Groups of Nodes with same Universe/Address
 	TMap<int32, TArray<TSharedPtr<FDMXFixturePatchNode>>> AddressToNodeGroupMap;
-	for (const TSharedPtr<FDMXFixturePatchNode>& Node : PatchedNodes)
+	for (const TSharedRef<FDMXFixturePatchNode>& Node : PatchedNodes)
 	{
 		if (Node->GetChannelSpan() == 0 || Node->GetChannelSpan() > DMX_MAX_ADDRESS)
 		{
@@ -437,26 +438,44 @@ void SDMXPatchedUniverse::RefreshInternal()
 	}
 
 	// Generate new widgets
-	for (const TSharedPtr<FDMXFixturePatchNode>& Node : PatchedNodes)
+	for (const TSharedRef<FDMXFixturePatchNode>& Node : PatchedNodes)
 	{
 		if (const TArray<TSharedPtr<FDMXFixturePatchNode>>* const NodeGroup = AddressToNodeGroupMap.Find(Node->GetStartingChannel()))
 		{
-			FixturePatchWidgets.Append(Node->GenerateWidgets(*NodeGroup));
+			FixturePatchFragmentWidgets.Append(Node->GenerateWidgets(StaticCastSharedRef<SDMXPatchedUniverse>(AsShared()), *NodeGroup));
 		}
 	}
 
 	// Add widgets to the grid
-	for (const TSharedPtr<SDMXFixturePatchFragment>& NewWidget : FixturePatchWidgets)
+	for (const TSharedRef<SDMXFixturePatchFragment>& NewWidget : FixturePatchFragmentWidgets)
 	{
-		Grid->AddSlot(NewWidget->GetColumn(), NewWidget->GetRow())
-			.ColumnSpan(NewWidget->GetColumnSpan())
+		Grid->AddSlot(NewWidget->Column, NewWidget->Row)
+			.ColumnSpan(NewWidget->ColumnSpan)
 			.Layer(1)
 			[
-				NewWidget.ToSharedRef()
+				NewWidget
 			];
 	}
 }
 
+void SDMXPatchedUniverse::UpdateNodesHoveredState()
+{
+	const int32 HoveredChannelID = HoveredChannel.IsValid() ? HoveredChannel->GetChannelID() : -1;
+	for (const TSharedRef<FDMXFixturePatchNode>& Node : PatchedNodes)
+	{
+		Node->SetIsHovered(false);
+
+		if (HoveredChannelID > 0)
+		{
+			constexpr int32 ChannelSpan = 1;
+			if (Node->OccupiesChannels(HoveredChannelID, ChannelSpan))
+			{
+				Node->SetIsHovered(true);
+			}
+		}
+	}
+}
+
 void SDMXPatchedUniverse::UpdateMonitor()
 {
 	UDMXLibrary* DMXLibrary = GetDMXLibrary();
@@ -466,7 +485,6 @@ void SDMXPatchedUniverse::UpdateMonitor()
 		return;
 	}
 
-
 	const TSet<FDMXInputPortSharedRef>& InputPorts = DMXLibrary->GetInputPorts();
 	for (const FDMXInputPortSharedRef& InputPort : InputPorts)
 	{
@@ -487,12 +505,26 @@ void SDMXPatchedUniverse::UpdateMonitor()
 
 void SDMXPatchedUniverse::ResetMonitor()
 {
-	for (const TSharedPtr<SDMXChannelConnector>& Channel : ChannelConnectors)
+	for (const TSharedRef<SDMXChannelConnector>& Channel : ChannelConnectors)
 	{
 		Channel->SetValue(0);
 	}
 }
 
+
+void SDMXPatchedUniverse::HandleOnChannelHovered(int32 Channel)
+{
+	HoveredChannel = ChannelConnectors[Channel];
+}
+
+void SDMXPatchedUniverse::HandleOnChannelUnhovered(int32 Channel)
+{
+	if (HoveredChannel == ChannelConnectors[Channel])
+	{
+		HoveredChannel.Reset();
+	}
+}
+
 FReply SDMXPatchedUniverse::HandleOnMouseButtonDownOnChannel(uint32 Channel, const FPointerEvent& PointerEvent)
 {
 	if (PointerEvent.GetEffectingButton() == EKeys::LeftMouseButton)
@@ -580,7 +612,7 @@ FReply SDMXPatchedUniverse::HandleOnMouseButtonDownOnChannel(uint32 Channel, con
 				SharedData->SelectFixturePatch(ClickedPatch);
 			}
 
-			return FReply::Handled().DetectDrag(ChannelConnectors[Channel - 1].ToSharedRef(), EKeys::LeftMouseButton);
+			return FReply::Handled();
 		}
 	}
 
@@ -607,9 +639,17 @@ FReply SDMXPatchedUniverse::HandleOnMouseButtonUpOnChannel(uint32 Channel, const
 			}
 		}
 	}
-
-	if (PointerEvent.GetEffectingButton() == EKeys::RightMouseButton)
+	else if (PointerEvent.GetEffectingButton() == EKeys::RightMouseButton)
 	{
+		// Select the patchh under the cursor if the current selection is not under the cursor
+		UDMXEntityFixturePatch* FixturePatchesOnChannel = GetTopmostFixturePatchOnChannel(Channel);
+		const TArray<TWeakObjectPtr<UDMXEntityFixturePatch>>& SelectedFixturePatches = SharedData->GetSelectedFixturePatches();
+		if (FixturePatchesOnChannel && !SelectedFixturePatches.Contains(FixturePatchesOnChannel))
+		{
+			SharedData->SelectFixturePatch(FixturePatchesOnChannel);
+		}
+
+		// Open context menu
 		const FVector2D& ContextMenuSummonLocation = PointerEvent.GetScreenSpacePosition();
 
 		TSharedRef<SWidget> MenuContent = CreateContextMenu(Channel);
@@ -705,9 +745,9 @@ void SDMXPatchedUniverse::UpdateZOrderOfNodes()
 		return;
 	}
 
-	TArray<TSharedPtr<FDMXFixturePatchNode>> Nodes = PatchedNodes;
+	TArray<TSharedRef<FDMXFixturePatchNode>> Nodes = PatchedNodes;
 	// Sort by ZOrder descending
-	Nodes.StableSort([](const TSharedPtr<FDMXFixturePatchNode>& NodeA, const TSharedPtr<FDMXFixturePatchNode>& NodeB)
+	Nodes.StableSort([](const TSharedRef<FDMXFixturePatchNode>& NodeA, const TSharedRef<FDMXFixturePatchNode>& NodeB)
 		{
 			if (NodeA->GetStartingChannel() > NodeB->GetStartingChannel())
 			{
@@ -733,13 +773,10 @@ void SDMXPatchedUniverse::UpdateZOrderOfNodes()
 		});
 
 	int32 ZOrder = Nodes.Num();
-	for (const TSharedPtr<FDMXFixturePatchNode>& Node : Nodes)
+	for (const TSharedRef<FDMXFixturePatchNode>& Node : Nodes)
 	{
-		if (Node.IsValid())
-		{
-			Node->SetZOrder(ZOrder);
-			--ZOrder;
-		}
+		Node->SetZOrder(ZOrder);
+		--ZOrder;
 	}
 }
 
@@ -800,7 +837,7 @@ void SDMXPatchedUniverse::OnSelectionChanged()
 	const TArray<TWeakObjectPtr<UDMXEntityFixturePatch>> FixturePatches = SharedData->GetSelectedFixturePatches();
 	const bool bSelectionChanged = [&FixturePatches, this]()
 		{
-			for (const TSharedPtr<FDMXFixturePatchNode>& Node : PatchedNodes)
+			for (const TSharedRef<FDMXFixturePatchNode>& Node : PatchedNodes)
 			{
 				if (Node->IsSelected() ^ FixturePatches.Contains(Node->GetFixturePatch()))
 				{
@@ -816,32 +853,36 @@ void SDMXPatchedUniverse::OnSelectionChanged()
 	}
 }
 
-void SDMXPatchedUniverse::AutoAssignFixturePatches()
+void SDMXPatchedUniverse::AutoAssignFixturePatches(EDMXFixturePatchAutoAssignMode AutoAssignMode)
 {
-	if (SharedData.IsValid())
+	TArray<UDMXEntityFixturePatch*> FixturePatchesToAutoAssign;
+	const TArray<TWeakObjectPtr<UDMXEntityFixturePatch>> SelectedFixturePatches = SharedData->GetSelectedFixturePatches();
+	for (TWeakObjectPtr<UDMXEntityFixturePatch> FixturePatch : SelectedFixturePatches)
 	{
-		TArray<UDMXEntityFixturePatch*> FixturePatchesToAutoAssign;
-		const TArray<TWeakObjectPtr<UDMXEntityFixturePatch>> SelectedFixturePatches = SharedData->GetSelectedFixturePatches();
-		for (TWeakObjectPtr<UDMXEntityFixturePatch> FixturePatch : SelectedFixturePatches)
+		if (FixturePatch.IsValid())
 		{
-			if (FixturePatch.IsValid())
-			{
-				FixturePatchesToAutoAssign.Add(FixturePatch.Get());
-			}
+			FixturePatchesToAutoAssign.Add(FixturePatch.Get());
 		}
+	}
 
-		if (FixturePatchesToAutoAssign.IsEmpty())
-		{
-			return;
-		}
+	if (FixturePatchesToAutoAssign.IsEmpty())
+	{
+		return;
+	}
 
-		constexpr bool bAllowDecrementUniverse = false;
-		constexpr bool bAllowDecrementChannels = true;
-		FDMXEditorUtils::AutoAssignedChannels(bAllowDecrementUniverse, bAllowDecrementChannels, FixturePatchesToAutoAssign);
-		SharedData->SelectUniverse(FixturePatchesToAutoAssign[0]->GetUniverseID());
+	FDMXFixturePatchAutoAssignUtility::AutoAssign(FixturePatchesToAutoAssign, AutoAssignMode);
+	SharedData->SelectUniverse(FixturePatchesToAutoAssign[0]->GetUniverseID());
 
-		RequestRefresh();
+	RequestRefresh();
+}
+
+bool SDMXPatchedUniverse::DoesDMXLibraryHaveReachableUniverses() const
+{
+	if (UDMXLibrary* DMXLibrary = GetDMXLibrary())
+	{
+		return !DMXLibrary->GetInputPorts().IsEmpty() && !DMXLibrary->GetOutputPorts().IsEmpty();
 	}
+	return false;
 }
 
 UDMXEntityFixturePatch* SDMXPatchedUniverse::GetTopmostFixturePatchOnChannel(uint32 Channel) const
@@ -876,8 +917,8 @@ UDMXEntityFixturePatch* SDMXPatchedUniverse::GetTopmostFixturePatchOnChannel(uin
 
 TArray<UDMXEntityFixturePatch*> SDMXPatchedUniverse::GetFixturePatchesOnChannel(uint32 Channel) const
 {
-	TArray<TSharedPtr<FDMXFixturePatchNode>> Nodes(PatchedNodes);
-	Nodes.RemoveAll([Channel](const TSharedPtr<FDMXFixturePatchNode>& Node)
+	TArray<TSharedRef<FDMXFixturePatchNode>> Nodes(PatchedNodes);
+	Nodes.RemoveAll([Channel](const TSharedRef<FDMXFixturePatchNode>& Node)
 		{
 			if (UDMXEntityFixturePatch* FixturePatch = Node->GetFixturePatch().Get())
 			{
@@ -890,7 +931,7 @@ TArray<UDMXEntityFixturePatch*> SDMXPatchedUniverse::GetFixturePatchesOnChannel(
 		});
 
 	TArray<UDMXEntityFixturePatch*> Result;
-	for (const TSharedPtr<FDMXFixturePatchNode>& Node : Nodes)
+	for (const TSharedRef<FDMXFixturePatchNode>& Node : Nodes)
 	{
 		if (UDMXEntityFixturePatch* FixturePatch = Node->GetFixturePatch().Get())
 		{
@@ -919,6 +960,25 @@ void SDMXPatchedUniverse::RegisterCommands()
 
 	// listen to common editor shortcuts for copy/paste etc
 	CommandList = MakeShared<FUICommandList>();
+
+	CommandList->MapAction
+	(
+		FDMXEditorCommands::Get().AutoAssignToFirstUniverseInSelection,
+		FExecuteAction::CreateSP(this, &SDMXPatchedUniverse::AutoAssignFixturePatches, EDMXFixturePatchAutoAssignMode::FirstUniverseInSelection)
+	);
+
+	CommandList->MapAction
+	(
+		FDMXEditorCommands::Get().AutoAssignToFirstReachableUniverse,
+		FExecuteAction::CreateSP(this, &SDMXPatchedUniverse::AutoAssignFixturePatches, EDMXFixturePatchAutoAssignMode::FirstReachableUniverse),
+		FCanExecuteAction::CreateSP(this, &SDMXPatchedUniverse::DoesDMXLibraryHaveReachableUniverses)
+	);
+
+	CommandList->MapAction
+	(
+		FDMXEditorCommands::Get().AutoAssignAfterLastAddressInLibrary,
+		FExecuteAction::CreateSP(this, &SDMXPatchedUniverse::AutoAssignFixturePatches, EDMXFixturePatchAutoAssignMode::LastAddressInLibrary)
+	);
 }
 
 TSharedRef<SWidget> SDMXPatchedUniverse::CreateContextMenu(int32 Channel)
@@ -929,29 +989,22 @@ TSharedRef<SWidget> SDMXPatchedUniverse::CreateContextMenu(int32 Channel)
 		return SNullWidget::NullWidget;
 	}
 
-	TArray<UDMXEntityFixturePatch*> FixturePatchesOnChannel = GetFixturePatchesOnChannel(Channel);
-	if (FixturePatchesOnChannel.IsEmpty())
-	{
-		return SNullWidget::NullWidget;
-	}
-
 	constexpr bool bCloseWindowAfterMenuSelection = true;
 	FMenuBuilder MenuBuilder(bCloseWindowAfterMenuSelection, CommandList);
 
 	// Auto Assign Section
-	MenuBuilder.BeginSection("AutoAssignSection", LOCTEXT("AutoAssignSection", "Auto-Assign"));
+	if (!SharedData->GetSelectedFixturePatches().IsEmpty())
 	{
-		const FUIAction Action(FExecuteAction::CreateSP(this, &SDMXPatchedUniverse::AutoAssignFixturePatches));
-
-		const FText AutoAssignText = LOCTEXT("AutoAssignContextMenuEntry", "Auto-Assign Selection");
-		const TSharedRef<SWidget> Widget =
-			SNew(STextBlock)
-			.Text(AutoAssignText);
-
-		MenuBuilder.AddMenuEntry(Action, Widget);
+		MenuBuilder.BeginSection("AutoAssignSection", LOCTEXT("AutoAssignSection", "Auto-Assign"));
+		{
+			MenuBuilder.AddMenuEntry(FDMXEditorCommands::Get().AutoAssignToFirstUniverseInSelection);
+			MenuBuilder.AddMenuEntry(FDMXEditorCommands::Get().AutoAssignToFirstReachableUniverse);
+			MenuBuilder.AddMenuEntry(FDMXEditorCommands::Get().AutoAssignAfterLastAddressInLibrary);
+		}
 		MenuBuilder.EndSection();
 	}
 
+	TArray<UDMXEntityFixturePatch*> FixturePatchesOnChannel = GetFixturePatchesOnChannel(Channel);
 	if (FixturePatchesOnChannel.Num() > 1)
 	{
 		// Select section
diff --git a/Engine/Plugins/VirtualProduction/DMX/DMXEngine/Source/DMXEditor/Private/Widgets/FixturePatch/SDMXPatchedUniverse.h b/Engine/Plugins/VirtualProduction/DMX/DMXEngine/Source/DMXEditor/Private/Widgets/FixturePatch/SDMXPatchedUniverse.h
--- a/Engine/Plugins/VirtualProduction/DMX/DMXEngine/Source/DMXEditor/Private/Widgets/FixturePatch/SDMXPatchedUniverse.h
+++ b/Engine/Plugins/VirtualProduction/DMX/DMXEngine/Source/DMXEditor/Private/Widgets/FixturePatch/SDMXPatchedUniverse.h
@@ -12,6 +12,7 @@
 #include "Widgets/DeclarativeSyntaxSupport.h"
 #include "Widgets/SCompoundWidget.h"
 
+enum class EDMXFixturePatchAutoAssignMode : uint8;
 class FDMXEditor;
 class FDMXFixturePatchNode;
 class FDMXFixturePatchSharedData;
@@ -81,10 +82,10 @@ public:
 	void SetMonitorInputsEnabled(bool bEnabled);
 
 	/** Finds or adds a Node to this Universe */
-	bool FindOrAdd(const TSharedPtr<FDMXFixturePatchNode>& Node);
+	bool FindOrAdd(const TSharedRef<FDMXFixturePatchNode>& Node);
 
 	/** Removes a Node from this Universe */
-	void Remove(const TSharedPtr<FDMXFixturePatchNode>& Node);
+	void Remove(const TSharedRef<FDMXFixturePatchNode>& Node);
 
 	/** Returns wether the patch can be patched to its current channels */
 	bool CanAssignFixturePatch(TWeakObjectPtr<UDMXEntityFixturePatch> TestedPatch) const;
@@ -105,7 +106,7 @@ public:
 	int32 GetUniverseID() const { return UniverseID; }
 
 	/** Gets all nodes patched to this universe */
-	const TArray<TSharedPtr<FDMXFixturePatchNode>>& GetPatchedNodes() const { return PatchedNodes; }
+	const TArray<TSharedRef<FDMXFixturePatchNode>>& GetPatchedNodes() const { return PatchedNodes; }
 
 protected:
 	//~ Begin SWidget Interface
@@ -125,6 +126,9 @@ private:
 	/** Refreshes the widget directly */
 	void RefreshInternal();
 
+	/** Updates the hover state of all nodes */
+	void UpdateNodesHoveredState();
+
 	/** Updates the monitor */
 	void UpdateMonitor();
 
@@ -137,6 +141,12 @@ private:
 	/** Returns the name of the universe displayed */
 	FText GetHeaderText() const;
 	
+	/** Called when a channel is hovered */
+	void HandleOnChannelHovered(int32 Channel);
+
+	/** Called when a channel is unhovered */
+	void HandleOnChannelUnhovered(int32 Channel);
+
 	/** Handles when a mouse button was pressed on a Channel */
 	FReply HandleOnMouseButtonDownOnChannel(uint32 Channel, const FPointerEvent& PointerEvent);
 
@@ -168,7 +178,10 @@ private:
 	void OnSelectionChanged();
 
 	/** Auto assigns selected Fixture Patches */
-	void AutoAssignFixturePatches();
+	void AutoAssignFixturePatches(EDMXFixturePatchAutoAssignMode AutoAssignMode);
+
+	/** Returns true if the DMX Library has reachable Universes */
+	bool DoesDMXLibraryHaveReachableUniverses() const;
 
 	/** Returns the Fixture Patch that is topmost under Channel */
 	UDMXEntityFixturePatch* GetTopmostFixturePatchOnChannel(uint32 Channel) const;
@@ -182,6 +195,9 @@ private:
 	/** The universe being displayed */
 	int32 UniverseID;
 
+	/** The hovered channel, or unset if no channel is hovered */
+	TSharedPtr<SDMXChannelConnector> HoveredChannel;
+
 	/** If true, monitors inputs */
 	bool bMonitorInputs = false;
 
@@ -195,10 +211,10 @@ private:
 	TSharedPtr<SGridPanel> Grid;
 
 	/** Patches in the grid */
-	TArray<TSharedPtr<FDMXFixturePatchNode>> PatchedNodes;
+	TArray<TSharedRef<FDMXFixturePatchNode>> PatchedNodes;
 
 	/** The Channel connectors in this universe */
-	TArray<TSharedPtr<SDMXChannelConnector>> ChannelConnectors;
+	TArray<TSharedRef<SDMXChannelConnector>> ChannelConnectors;
 
 	/** Delegate executed when drag enters a Channel */
 	FOnDragOverChannel OnDragEnterChannel;
@@ -213,7 +229,7 @@ private:
 	TWeakObjectPtr<UDMXEntityFixturePatch> ShiftSelectAnchorPatch;
 
 	/** The Fixture Patch Widgets that are currently being displalyed */
-	TArray<TSharedPtr<SDMXFixturePatchFragment>> FixturePatchWidgets;
+	TArray<TSharedRef<SDMXFixturePatchFragment>> FixturePatchFragmentWidgets;
 
 	/** Timer handle for the Request Refresh method */
 	FTimerHandle RequestRefreshTimerHandle;
diff --git a/Engine/Plugins/VirtualProduction/DMX/DMXEngine/Source/DMXEditor/Private/Widgets/SDMXChannel.cpp b/Engine/Plugins/VirtualProduction/DMX/DMXEngine/Source/DMXEditor/Private/Widgets/SDMXChannel.cpp
--- a/Engine/Plugins/VirtualProduction/DMX/DMXEngine/Source/DMXEditor/Private/Widgets/SDMXChannel.cpp
+++ b/Engine/Plugins/VirtualProduction/DMX/DMXEngine/Source/DMXEditor/Private/Widgets/SDMXChannel.cpp
@@ -2,28 +2,26 @@
 
 #include "Widgets/SDMXChannel.h"
 
-#include "Interfaces/IDMXProtocol.h"
 #include "DMXEditorLog.h"
 #include "DMXProtocolConstants.h"
 #include "DMXEditorStyle.h"
 
+#include "Brushes/SlateColorBrush.h"
+#include "Engine/Font.h"
 #include "Widgets/SOverlay.h"
 #include "Widgets/Text/STextBlock.h"
 #include "Widgets/SBoxPanel.h"
 #include "Widgets/Layout/SSeparator.h"
 #include "Widgets/Notifications/SProgressBar.h"
 #include "Widgets/Images/SImage.h"
-#include "Brushes/SlateColorBrush.h"
+#include "Widgets/Layout/SBox.h"
 #include "UObject/ConstructorHelpers.h"
-#include "Engine/Font.h"
+
 
 #define LOCTEXT_NAMESPACE "SDMXChannel"
 
 const float SDMXChannel::NewValueChangedAnimDuration = 0.8f;
 
-const FLinearColor SDMXChannel::IDColor = FLinearColor(1,1,1, 0.6f);
-const FLinearColor SDMXChannel::ValueColor = FLinearColor(1,1,1, 0.9f);
-
 void SDMXChannel::Construct(const FArguments& InArgs)
 {
 	SetVisibility(EVisibility::SelfHitTestInvisible);
@@ -33,24 +31,24 @@ void SDMXChannel::Construct(const FArguments& InArgs)
 	Value = InArgs._Value;
 	NewValueFreshness = 0.0f;
 
-	const float PaddingInfo = 3.0f;
+	constexpr float PaddingInfo = 3.0f;
 
 
 	ChannelIDTextBlock =
 		SNew(STextBlock)
 		.Text(GetChannelIDText())
-		.ColorAndOpacity(FSlateColor(IDColor))
+		.ColorAndOpacity(InArgs._ChannelIDTextColor)
 		.MinDesiredWidth(24.0f)
 		.Justification(ETextJustify::Center)
-		.Font(FDMXEditorStyle::Get().GetFontStyle("DMXEditor.Font.InputChannelID"));
+		.Font(FAppStyle::GetFontStyle(TEXT("PropertyWindow.NormalFont")));
 
 	ChannelValueTextBlock =
 		SNew(STextBlock)
 		.Text(GetValueText())
-		.ColorAndOpacity(FSlateColor(ValueColor))
+		.ColorAndOpacity(InArgs._ValueTextColor)
 		.MinDesiredWidth(24.0f)
 		.Justification(ETextJustify::Center)
-		.Font(FDMXEditorStyle::Get().GetFontStyle("DMXEditor.Font.InputChannelValue"));
+		.Font(FAppStyle::GetFontStyle(TEXT("PropertyWindow.NormalFont")));
 
 	// Arrange widgets depending on bInvertChannelDisplay argument
 	TSharedPtr<STextBlock> UpperWidget;
@@ -69,52 +67,47 @@ void SDMXChannel::Construct(const FArguments& InArgs)
 
 	ChildSlot
 	[
-		// root
-		SNew(SOverlay)
-
-		+ SOverlay::Slot()
-		.HAlign(HAlign_Fill)
-		.VAlign(VAlign_Fill)
-		[
-			// Background color image
-			SAssignNew(BarColorBorder, SImage)
-			.Image(FDMXEditorStyle::Get().GetBrush(TEXT("DMXEditor.WhiteBrush")))
-			.ColorAndOpacity(GetBackgroundColor())
-		]
-
-		// Info
-		+ SOverlay::Slot()
-		.HAlign(HAlign_Fill)
-		.VAlign(VAlign_Fill)
-		.Padding(PaddingInfo)
+		SNew(SBox)
+		.WidthOverride(36.f)
+		.HeightOverride(36.f)
 		[
-			SNew(SVerticalBox)
-
-			// ID Label
-			+ SVerticalBox::Slot()
-			.FillHeight(0.5f)
-			.VAlign(VAlign_Center)
-			.HAlign(HAlign_Center)
-			[
-				UpperWidget.ToSharedRef()
-			]
+			SNew(SOverlay)
 
-			// Separator
-			+ SVerticalBox::Slot()
-			.AutoHeight()
+			// Background color image
+			+ SOverlay::Slot()
 			.HAlign(HAlign_Fill)
+			.VAlign(VAlign_Fill)
 			[
-				SNew(SSeparator)
-				.Orientation(Orient_Horizontal)
+				SAssignNew(BarColorBorder, SImage)
+				.Image(FDMXEditorStyle::Get().GetBrush(TEXT("DMXEditor.WhiteBrush")))
+				.ColorAndOpacity(GetBackgroundColor())
 			]
 
-			// Value Label
-			+ SVerticalBox::Slot()
-			.FillHeight(0.5f)
-			.VAlign(VAlign_Center)
-			.HAlign(HAlign_Center)
+			// Info
+			+ SOverlay::Slot()
+			.HAlign(HAlign_Fill)
+			.VAlign(VAlign_Fill)
+			.Padding(PaddingInfo)
 			[
-				LowerWidget.ToSharedRef()
+				SNew(SVerticalBox)
+
+				// ID Label
+				+ SVerticalBox::Slot()
+				.FillHeight(0.5f)
+				.HAlign(HAlign_Center)
+				.VAlign(VAlign_Center)
+				[
+					UpperWidget.ToSharedRef()
+				]
+
+				// Value Label
+				+ SVerticalBox::Slot()
+				.FillHeight(0.5f)
+				.HAlign(HAlign_Center)
+				.VAlign(VAlign_Center)
+				[
+					LowerWidget.ToSharedRef()
+				]
 			]
 		]
 	];
diff --git a/Engine/Plugins/VirtualProduction/DMX/DMXEngine/Source/DMXEditor/Private/Widgets/SDMXChannel.h b/Engine/Plugins/VirtualProduction/DMX/DMXEngine/Source/DMXEditor/Private/Widgets/SDMXChannel.h
--- a/Engine/Plugins/VirtualProduction/DMX/DMXEngine/Source/DMXEditor/Private/Widgets/SDMXChannel.h
+++ b/Engine/Plugins/VirtualProduction/DMX/DMXEngine/Source/DMXEditor/Private/Widgets/SDMXChannel.h
@@ -28,6 +28,8 @@ public:
 		: _ChannelID(0)
 		, _Value(0)
 		, _bShowChannelIDBottom(false)
+		, _ChannelIDTextColor(FLinearColor::White.CopyWithNewOpacity(.6f))
+		, _ValueTextColor(FLinearColor::White.CopyWithNewOpacity(.9f))
 		{}
 
 		/** The channel ID this widget represents */
@@ -36,16 +38,20 @@ public:
 		/** The current value from the channel */
 		SLATE_ARGUMENT(uint8, Value)
 
-
 		/** If true, draws the channel number below the value */
 		SLATE_ARGUMENT(bool, bShowChannelIDBottom)
 
+		/** Color of the Channel ID Text */
+		SLATE_ATTRIBUTE(FSlateColor, ChannelIDTextColor)
+
+		/** Color of the Value Text */
+		SLATE_ATTRIBUTE(FSlateColor, ValueTextColor)
+
 	SLATE_END_ARGS()
 
 	/** Constructs the widget */
 	void Construct(const FArguments& InArgs);
 
-public:
 	/** Sets the channel ID this widget represents */
 	void SetChannelID(uint32 NewChannelID);
 
@@ -92,14 +98,6 @@ private:
 	/** Used to stop the animation timer once the animation is completed */
 	TWeakPtr<FActiveTimerHandle> AnimationTimerHandle;
 
-	// ~ VISUAL STYLE CONSTANTS
-
-	/** Color of the ID label */
-	static const FLinearColor IDColor;
-
-	/** Color of the Value label */
-	static const FLinearColor ValueColor;
-
 private:
 	/** Returns the channel ID in Text form to display it in the UI */
 	FText GetChannelIDText() const;
diff --git a/Engine/Plugins/VirtualProduction/DMX/DMXEngine/Source/DMXEditor/Public/Commands/DMXEditorCommands.h b/Engine/Plugins/VirtualProduction/DMX/DMXEngine/Source/DMXEditor/Public/Commands/DMXEditorCommands.h
--- a/Engine/Plugins/VirtualProduction/DMX/DMXEngine/Source/DMXEditor/Public/Commands/DMXEditorCommands.h
+++ b/Engine/Plugins/VirtualProduction/DMX/DMXEngine/Source/DMXEditor/Public/Commands/DMXEditorCommands.h
@@ -35,4 +35,9 @@ public:
 	TSharedPtr<FUICommandInfo> OpenPatchTool;
 	TSharedPtr<FUICommandInfo> ToggleReceiveDMX;
 	TSharedPtr<FUICommandInfo> ToggleSendDMX;
+
+	// Auto assign
+	TSharedPtr<FUICommandInfo> AutoAssignToFirstUniverseInSelection;
+	TSharedPtr<FUICommandInfo> AutoAssignToFirstReachableUniverse;
+	TSharedPtr<FUICommandInfo> AutoAssignAfterLastAddressInLibrary;
 };
diff --git a/Engine/Plugins/VirtualProduction/DMX/DMXEngine/Source/DMXEditor/Public/DMXEditorSettings.h b/Engine/Plugins/VirtualProduction/DMX/DMXEngine/Source/DMXEditor/Public/DMXEditorSettings.h
--- a/Engine/Plugins/VirtualProduction/DMX/DMXEngine/Source/DMXEditor/Public/DMXEditorSettings.h
+++ b/Engine/Plugins/VirtualProduction/DMX/DMXEngine/Source/DMXEditor/Public/DMXEditorSettings.h
@@ -15,6 +15,9 @@ struct FDMXMVRFixtureListSettings
 {
 	GENERATED_BODY()
 
+	/** Width of the entier list. The right side should take the rest of the space */
+	float ListWidth = .5f;
+
 	/** Width of the Fixture ID column */
 	UPROPERTY()
 	float FixtureIDColumnWidth = 0.f;
@@ -32,6 +35,28 @@ struct FDMXMVRFixtureListSettings
 	float PatchColumnWidth = 0.f;
 };
 
+UENUM()
+enum class EDMXFixturePatcherNameDisplayMode : uint8
+{
+	FixtureIDAndFixturePatchName,
+	FixtureID,
+	FixturePatchName
+};
+
+/** Settings for the Fixture Patcher */
+USTRUCT()
+struct FDMXMVRFixturePatcherSettings
+{
+	GENERATED_BODY()
+
+	UPROPERTY()
+	bool bMonitorEnabled = false;
+
+	UPROPERTY()
+	EDMXFixturePatcherNameDisplayMode FixturePatchNameDisplayMode;
+};
+
+
 /** Settings for the Fixture Type Functions Editor */
 USTRUCT()
 struct FDMXFixtureTypeFunctionsEditorSettings
@@ -117,7 +142,7 @@ struct FDMXMonitorSourceDescriptor
 };
 
 /** Settings that holds editor configurations. Not accessible in Project Settings. TODO: Idealy rename to UDMXEditorConfiguration */
-UCLASS(Config = DMXEditor, DefaultConfig, meta = (DisplayName = "DMXEditor"))
+UCLASS(Config = DMXEditor)
 class DMXEDITOR_API UDMXEditorSettings : public UObject
 {
 	GENERATED_BODY()
@@ -135,17 +160,21 @@ public:
 	UPROPERTY(Config)
 	FString LastMVRExportPath;
 
-	UPROPERTY(Config)
-	FDMXMVRFixtureListSettings MVRFixtureListSettings;
-
 	// DMX Library
 public:
+	/** Deprecated 5.3, moved to FDMXMVRFixturePatcherSettings */
 	UPROPERTY(Config)
-	bool bFixturePatcherDMXMonitorEnabled = false;
+	bool bFixturePatcherDMXMonitorEnabled_DEPRECATED = false;
 
 	UPROPERTY(Config)
 	FDMXFixtureTypeFunctionsEditorSettings FixtureTypeFunctionsEditorSettings;
 
+	UPROPERTY(Config)
+	FDMXMVRFixtureListSettings MVRFixtureListSettings;
+
+	UPROPERTY(Config)
+	FDMXMVRFixturePatcherSettings FixturePatcherSettings;
+
 	// Output Console (DEPRECATED 5.1)
 public:
 #if WITH_EDITORONLY_DATA
diff --git a/Engine/Plugins/VirtualProduction/DMX/DMXEngine/Source/DMXEditor/Public/DMXEditorUtils.h b/Engine/Plugins/VirtualProduction/DMX/DMXEngine/Source/DMXEditor/Public/DMXEditorUtils.h
--- a/Engine/Plugins/VirtualProduction/DMX/DMXEngine/Source/DMXEditor/Public/DMXEditorUtils.h
+++ b/Engine/Plugins/VirtualProduction/DMX/DMXEngine/Source/DMXEditor/Public/DMXEditorUtils.h
@@ -105,6 +105,7 @@ public:
 	 * @param bCanChangePatchUniverses	If true, the patches can lower their absoulte channel
 	 * @param FixturePatches			The patches that want their channels to be auto assigned
 	 */
+	UE_DEPRECATED(5.2, "AutoAssignedChannels is deprecated without replacement. Please create your own auto assign implementation. See FDMXFixturePatchAutoAssignUtility for an example.")
 	static void AutoAssignedChannels(bool bAllowDecrementUniverse, bool bAllowDecrementChannels, TArray<UDMXEntityFixturePatch*> FixturePatches);
 	
 	/**
