Fixed mismatched copy constructor/assignment operator for FDMXOptionalBool, FDMXOptionalInt32, FDMXOptionalFloat, FDMXOptionalGuid, FDMXOptionalVector2D, FDMXOptionalTransform, FDMXColorCIE1931xyY and FDMXOptionalString.

diff --git a/Engine/Plugins/VirtualProduction/DMX/DMXEngine/Source/DMXRuntime/Public/DMXOptionalTypes.h b/Engine/Plugins/VirtualProduction/DMX/DMXEngine/Source/DMXRuntime/Public/DMXOptionalTypes.h
--- a/Engine/Plugins/VirtualProduction/DMX/DMXEngine/Source/DMXRuntime/Public/DMXOptionalTypes.h
+++ b/Engine/Plugins/VirtualProduction/DMX/DMXEngine/Source/DMXRuntime/Public/DMXOptionalTypes.h
@@ -27,11 +27,6 @@ struct DMXRUNTIME_API FDMXOptionalBool
 	/** Returns the Value of the optional, only call when the value is set. */
 	const bool& GetValue() const { return Value.GetValue(); }
 
-	const FDMXOptionalBool& operator=(const FDMXOptionalBool& Rhs)
-	{
-		Value = Rhs.Value;
-		return *this;
-	}
 	const FDMXOptionalBool& operator=(const bool& Transform)
 	{
 		Value = Transform;
@@ -98,11 +93,6 @@ struct DMXRUNTIME_API FDMXOptionalInt32
 	/** Returns the Value of the optional, only call when the value is set. */
 	const int32& GetValue() const { return Value.GetValue(); }
 
-	const FDMXOptionalInt32& operator=(const FDMXOptionalInt32& Rhs)
-	{
-		Value = Rhs.Value;
-		return *this;
-	}
 	const FDMXOptionalInt32& operator=(const int32& Transform)
 	{
 		Value = Transform;
@@ -175,11 +165,6 @@ struct DMXRUNTIME_API FDMXOptionalFloat
 	/** Returns the Value of the optional, only call when the value is set. */
 	const float& GetValue() const { return Value.GetValue(); }
 
-	const FDMXOptionalFloat& operator=(const FDMXOptionalFloat& Rhs)
-	{
-		Value = Rhs.Value;
-		return *this;
-	}
 	const FDMXOptionalFloat& operator=(const float& Transform)
 	{
 		Value = Transform;
@@ -246,11 +231,6 @@ struct DMXRUNTIME_API FDMXOptionalGuid
 	/** Returns the Value of the optional, only call when the value is set. */
 	const FGuid& GetValue() const { return Value.GetValue(); }
 
-	const FDMXOptionalGuid& operator=(const FDMXOptionalGuid& Rhs)
-	{
-		Value = Rhs.Value;
-		return *this;
-	}
 	const FDMXOptionalGuid& operator=(const FGuid& Transform)
 	{
 		Value = Transform;
@@ -324,11 +304,6 @@ struct DMXRUNTIME_API FDMXOptionalVector2D
 	/** Returns the Value of the optional, only call when the value is set. */
 	const FVector2D& GetValue() const { return Value.GetValue(); }
 
-	const FDMXOptionalVector2D& operator=(const FDMXOptionalVector2D& Rhs)
-	{
-		Value = Rhs.Value;
-		return *this;
-	}
 	const FDMXOptionalVector2D& operator=(const FVector2D& Vector2D)
 	{
 		Value = Vector2D;
@@ -401,11 +376,6 @@ struct DMXRUNTIME_API FDMXOptionalTransform
 	/** Returns the Value of the optional, only call when the value is set. */
 	const FTransform& GetValue() const { return Value.GetValue(); }
 
-	const FDMXOptionalTransform& operator=(const FDMXOptionalTransform& Rhs)
-	{
-		Value = Rhs.Value;
-		return *this;
-	}
 	const FDMXOptionalTransform& operator=(const FTransform& Transform)
 	{
 		Value = Transform;
@@ -478,11 +448,6 @@ struct DMXRUNTIME_API FDMXOptionalColorCIE1931
 	/** Returns the Value of the optional, only call when the value is set. */
 	const FDMXColorCIE1931xyY& GetValue() const { return Value.GetValue(); }
 
-	const FDMXOptionalColorCIE1931& operator=(const FDMXOptionalColorCIE1931& Rhs)
-	{
-		Value = Rhs.Value;
-		return *this;
-	}
 	const FDMXOptionalColorCIE1931& operator=(const FDMXColorCIE1931xyY& Transform)
 	{
 		Value = Transform;
@@ -549,11 +514,6 @@ struct DMXRUNTIME_API FDMXOptionalString
 	/** Returns the Value of the optional, only call when the value is set. */
 	const FString& GetValue() const { return Value.GetValue(); }
 
-	const FDMXOptionalString& operator=(const FDMXOptionalString& Rhs)
-	{
-		Value = Rhs.Value;
-		return *this;
-	}
 	const FDMXOptionalString& operator=(const FString& Transform)
 	{
 		Value = Transform;
