DMX - Fix fixture patch conflict background is transparent instead of yellow in fixture patcher

diff --git a/Engine/Plugins/VirtualProduction/DMX/DMXEngine/Source/DMXEditor/Private/Widgets/FixturePatch/SDMXFixturePatchFragment.cpp b/Engine/Plugins/VirtualProduction/DMX/DMXEngine/Source/DMXEditor/Private/Widgets/FixturePatch/SDMXFixturePatchFragment.cpp
--- a/Engine/Plugins/VirtualProduction/DMX/DMXEngine/Source/DMXEditor/Private/Widgets/FixturePatch/SDMXFixturePatchFragment.cpp
+++ b/Engine/Plugins/VirtualProduction/DMX/DMXEngine/Source/DMXEditor/Private/Widgets/FixturePatch/SDMXFixturePatchFragment.cpp
@@ -81,7 +81,7 @@ protected:
 		const bool bEnabled = ShouldBeEnabled(bParentEnabled);
 		if (bIsConflicting)
 		{
-			const FLinearColor BackgroundColor = ConflictBrush->GetTint(InWidgetStyle).CopyWithNewOpacity(.3f);
+			const FLinearColor BackgroundColor = ConflictBrush->GetTint(InWidgetStyle);
 
 			FSlateDrawElement::MakeBox(
 				OutDrawElements,
@@ -105,13 +105,14 @@ protected:
 		}
 		else 
 		{
-			const float InvertedLuminanceClamped = FMath::Clamp(FMath::Abs(1.f - FixturePatch->EditorColor.GetLuminance()), .4f, .8f);
+			const float InvertedLuminanceClamped = FMath::Clamp(FMath::Abs(1.f - FixturePatch->EditorColor.GetLuminance()), .5f, .8f);
 			const float Opacity = IsNodeHovered.Get() ? InvertedLuminanceClamped : InvertedLuminanceClamped / 2.f;
 
 			TArray<FSlateGradientStop> GradientStops;
 			GradientStops.Add(FSlateGradientStop(FVector2D::ZeroVector, FixturePatch->EditorColor.CopyWithNewOpacity(Opacity)));
 			GradientStops.Add(FSlateGradientStop(AllottedGeometry.GetLocalSize(), FixturePatch->EditorColor.CopyWithNewOpacity(Opacity / 4.f)));
-			 
+			GradientStops.Add(FSlateGradientStop(AllottedGeometry.GetLocalSize(), FLinearColor::Transparent));
+
 			FSlateDrawElement::MakeGradient(
 				OutDrawElements,
 				LayerId,
@@ -422,7 +423,8 @@ FSlateColor SDMXFixturePatchFragment::GetBorderColor() const
 	}
 	else if (UDMXEntityFixturePatch* Patch = FixturePatchNode->GetFixturePatch().Get())
 	{
-		return (Patch->EditorColor * 3.f + FLinearColor::White) / 4.f; 
+		const FLinearColor WhiteAdjustedColor = (Patch->EditorColor * 2.f + FLinearColor::White) / 3.f;
+		return WhiteAdjustedColor.CopyWithNewOpacity(0.8f);
 	}
 	else
 	{
diff --git a/Engine/Plugins/VirtualProduction/DMX/DMXEngine/Source/DMXEditor/Private/Widgets/FixturePatch/SDMXPatchedUniverse.cpp b/Engine/Plugins/VirtualProduction/DMX/DMXEngine/Source/DMXEditor/Private/Widgets/FixturePatch/SDMXPatchedUniverse.cpp
--- a/Engine/Plugins/VirtualProduction/DMX/DMXEngine/Source/DMXEditor/Private/Widgets/FixturePatch/SDMXPatchedUniverse.cpp
+++ b/Engine/Plugins/VirtualProduction/DMX/DMXEngine/Source/DMXEditor/Private/Widgets/FixturePatch/SDMXPatchedUniverse.cpp
@@ -460,7 +460,7 @@ void SDMXPatchedUniverse::RefreshInternal()
 
 void SDMXPatchedUniverse::UpdateNodesHoveredState()
 {
-	const int32 HoveredChannelID = HoveredChannel.IsValid() ? HoveredChannel->GetChannelID() : -1;
+	const int32 HoveredChannelID = HoveredChannel.IsValid() ? HoveredChannel->GetChannelID() + 1 : -1;
 	for (const TSharedRef<FDMXFixturePatchNode>& Node : PatchedNodes)
 	{
 		Node->SetIsHovered(false);
