DMX - Fix Activity Monitor not updating Widgets when value is reset to 0 in some cases

diff --git a/Engine/Plugins/VirtualProduction/DMX/DMXEngine/Source/DMXEditor/Private/Widgets/Monitors/SDMXActivityInUniverse.cpp b/Engine/Plugins/VirtualProduction/DMX/DMXEngine/Source/DMXEditor/Private/Widgets/Monitors/SDMXActivityInUniverse.cpp
--- a/Engine/Plugins/VirtualProduction/DMX/DMXEngine/Source/DMXEditor/Private/Widgets/Monitors/SDMXActivityInUniverse.cpp
+++ b/Engine/Plugins/VirtualProduction/DMX/DMXEngine/Source/DMXEditor/Private/Widgets/Monitors/SDMXActivityInUniverse.cpp
@@ -155,7 +155,8 @@ void SDMXActivityInUniverse::UpdateChannels()
 			{				
 				// Instantiate the widget
 				ChannelItem->CreateWidget(Buffer[ChannelIndex]);
-			
+				ChannelItem->SetValue(Buffer[ChannelIndex]);
+
 				// Add the item the list source
 				ChannelListSource.Add(ChannelItem);
 				bNumChannelsChanged = true;
