DMX - Remove obsolete ensure condition in auto assign utility

diff --git a/Engine/Plugins/VirtualProduction/DMX/DMXEngine/Source/DMXEditor/Private/DMXFixturePatchAutoAssignUtility.cpp b/Engine/Plugins/VirtualProduction/DMX/DMXEngine/Source/DMXEditor/Private/DMXFixturePatchAutoAssignUtility.cpp
--- a/Engine/Plugins/VirtualProduction/DMX/DMXEngine/Source/DMXEditor/Private/DMXFixturePatchAutoAssignUtility.cpp
+++ b/Engine/Plugins/VirtualProduction/DMX/DMXEngine/Source/DMXEditor/Private/DMXFixturePatchAutoAssignUtility.cpp
@@ -216,11 +216,6 @@ void FDMXFixturePatchAutoAssignUtility::AutoAssignAfterLastAddressInLibrary()
 		return;
 	}
 
-	if (!ensureMsgf(!OtherPatches.IsEmpty() && OtherPatches[0], TEXT("Cannot auto assign. Invalid 'other' patch when trying to auto assign to last address in DMX Library")))
-	{
-		return;
-	}
-
 	const uint64 MinRequiredSize = ComputeMinRequiredSize();
 
 	const uint64 AbsoluteAddressAfterLastPatchInLibrary = [this]()
