DMX - Fix crash with undefined mvr fixture actor spawn parameter

diff --git a/Engine/Plugins/VirtualProduction/DMX/DMXEngine/Source/DMXRuntime/Private/MVR/DMXMVRSceneActor.cpp b/Engine/Plugins/VirtualProduction/DMX/DMXEngine/Source/DMXRuntime/Private/MVR/DMXMVRSceneActor.cpp
--- a/Engine/Plugins/VirtualProduction/DMX/DMXEngine/Source/DMXRuntime/Private/MVR/DMXMVRSceneActor.cpp
+++ b/Engine/Plugins/VirtualProduction/DMX/DMXEngine/Source/DMXRuntime/Private/MVR/DMXMVRSceneActor.cpp
@@ -532,6 +532,7 @@ AActor* ADMXMVRSceneActor::SpawnMVRActor(const TSubclassOf<AActor>&ActorClass, U
 	ActorSpawnParameters.Template = Template;
 	ActorSpawnParameters.SpawnCollisionHandlingOverride = ESpawnActorCollisionHandlingMethod::AlwaysSpawn;
 	ActorSpawnParameters.Name = ActorName;
+	ActorSpawnParameters.NameMode = FActorSpawnParameters::ESpawnActorNameMode::Requested;
 	AActor* NewFixtureActor = World->SpawnActor<AActor>(ActorClass, Transform, ActorSpawnParameters);
 	if (!NewFixtureActor)
 	{
