DMX - MVR Actor now correctly ignores case and trims, when comparing mode names of Fixture Types and GDTFs

diff --git a/Engine/Plugins/VirtualProduction/DMX/DMXEngine/Source/DMXEditor/Private/Factories/DMXLibraryFromMVRImporter.cpp b/Engine/Plugins/VirtualProduction/DMX/DMXEngine/Source/DMXEditor/Private/Factories/DMXLibraryFromMVRImporter.cpp
--- a/Engine/Plugins/VirtualProduction/DMX/DMXEngine/Source/DMXEditor/Private/Factories/DMXLibraryFromMVRImporter.cpp
+++ b/Engine/Plugins/VirtualProduction/DMX/DMXEngine/Source/DMXEditor/Private/Factories/DMXLibraryFromMVRImporter.cpp
@@ -293,7 +293,9 @@ void FDMXLibraryFromMVRImporter::InitializeDMXLibrary(UDMXLibrary* DMXLibrary, c
 
 		int32 ActiveModeIndex = FixtureType->Modes.IndexOfByPredicate([FixtureNode](const FDMXFixtureMode& Mode)
 			{
-				return Mode.ModeName == FixtureNode->GDTFMode;
+				const FString TrimmedModeName = Mode.ModeName.TrimStartAndEnd();
+				const FString TrimmedGDTFModeName = FixtureNode->GDTFMode.TrimStartAndEnd();
+				return TrimmedModeName.Equals(TrimmedGDTFModeName);
 			});
 
 		if (ActiveModeIndex != INDEX_NONE)
