DMX - Improve Fixture Patch editor UX, better 'Add Patch' menu, improve options for auto-assign

diff --git a/Engine/Plugins/VirtualProduction/DMX/DMXEngine/Source/DMXEditor/Private/Commands/DMXEditorCommands.cpp b/Engine/Plugins/VirtualProduction/DMX/DMXEngine/Source/DMXEditor/Private/Commands/DMXEditorCommands.cpp
--- a/Engine/Plugins/VirtualProduction/DMX/DMXEngine/Source/DMXEditor/Private/Commands/DMXEditorCommands.cpp
+++ b/Engine/Plugins/VirtualProduction/DMX/DMXEngine/Source/DMXEditor/Private/Commands/DMXEditorCommands.cpp
@@ -27,9 +27,9 @@ void FDMXEditorCommands::RegisterCommands()
 	UI_COMMAND(ToggleReceiveDMX, "Receive DMX", "Sets whether DMX is received in editor", EUserInterfaceActionType::ToggleButton, FInputChord(EModifierKey::Alt, EKeys::M));
 	UI_COMMAND(ToggleSendDMX, "Send DMX", "Sets whether DMX is sent from editor", EUserInterfaceActionType::ToggleButton, FInputChord(EModifierKey::Alt, EKeys::N));
 
-	UI_COMMAND(AutoAssignToFirstUniverseInSelection, "Auto Assign", "Auto assigns selected patches, doesn't decrement Universe of selection.", EUserInterfaceActionType::Button, FInputChord());
-	UI_COMMAND(AutoAssignToFirstReachableUniverse, "Auto Assign to first reachable Universe", "Auto assigns selected patches to the first Universe reachable by the DMX Library", EUserInterfaceActionType::Button, FInputChord());
-	UI_COMMAND(AutoAssignAfterLastAddressInLibrary, "Auto Assign after last Patch in Library", "Auto assigns selected patches after the last patch in the DMX Library", EUserInterfaceActionType::Button, FInputChord());
+	UI_COMMAND(Align, "Align", "Aligns selected patches", EUserInterfaceActionType::Button, FInputChord());
+	UI_COMMAND(Stack, "Stack", "Stacks selected patches", EUserInterfaceActionType::Button, FInputChord());
+	UI_COMMAND(AutoAssignSelectedUniverse, "Auto-Assign in selected Universe", "Auto-assigns selected patches to first consecutive range of free channels, starting from the selected Universe.", EUserInterfaceActionType::Button, FInputChord());
 }
 
 #undef LOCTEXT_NAMESPACE
diff --git a/Engine/Plugins/VirtualProduction/DMX/DMXEngine/Source/DMXEditor/Private/DMXFixturePatchAutoAssignUtility.cpp b/Engine/Plugins/VirtualProduction/DMX/DMXEngine/Source/DMXEditor/Private/DMXFixturePatchAutoAssignUtility.cpp
deleted file mode 100644
diff --git a/Engine/Plugins/VirtualProduction/DMX/DMXEngine/Source/DMXEditor/Private/DMXFixturePatchAutoAssignUtility.h b/Engine/Plugins/VirtualProduction/DMX/DMXEngine/Source/DMXEditor/Private/DMXFixturePatchAutoAssignUtility.h
deleted file mode 100644
diff --git a/Engine/Plugins/VirtualProduction/DMX/DMXEngine/Source/DMXEditor/Private/FixturePatchAutoAssignUtility.cpp b/Engine/Plugins/VirtualProduction/DMX/DMXEngine/Source/DMXEditor/Private/FixturePatchAutoAssignUtility.cpp
new file mode 100644
--- /dev/null
+++ b/Engine/Plugins/VirtualProduction/DMX/DMXEngine/Source/DMXEditor/Private/FixturePatchAutoAssignUtility.cpp
@@ -0,0 +1,584 @@
+// Copyright Epic Games, Inc. All Rights Reserved.
+
+#include "FixturePatchAutoAssignUtility.h"
+
+#include "Algo/Accumulate.h"
+#include "Algo/MaxElement.h"
+#include "Commands/DMXEditorCommands.h"
+#include "DMXProtocolConstants.h"
+#include "DMXEditor.h"
+#include "DMXFixturePatchSharedData.h"
+#include "Framework/MultiBox/MultiBoxBuilder.h"
+#include "IO/DMXInputPort.h"
+#include "IO/DMXOutputPort.h"
+#include "Library/DMXEntityFixturePatch.h"
+#include "Library/DMXLibrary.h"
+#include "ScopedTransaction.h"
+
+
+#define LOCTEXT_NAMESPACE "FixturePatchAutoAssignUtility"
+
+namespace UE::DMXEditor::AutoAssign::Private
+{
+	FAutoAssignElement::FAutoAssignElement(int64 LowerBoundInclusive, int64 UpperBoundExclusive)
+		: AbsoluteRange(TRange<int64>(LowerBoundInclusive, UpperBoundExclusive))
+	{}
+
+	FAutoAssignElement::FAutoAssignElement(const TRange<int64>& InAbsoluteRange)
+		: AbsoluteRange(InAbsoluteRange)
+	{}
+
+	FAutoAssignElement::FAutoAssignElement(UDMXEntityFixturePatch& FixturePatch)
+		: WeakFixturePatch(&FixturePatch)
+	{
+		const int64 AbsoluteStartingChannel = FixturePatch.GetUniverseID() * DMX_UNIVERSE_SIZE + FixturePatch.GetStartingChannel() - 1;
+		const int64 AbsoluteEndingChannel = AbsoluteStartingChannel + FixturePatch.GetChannelSpan();
+		AbsoluteRange = TRange<int64>(AbsoluteStartingChannel, AbsoluteEndingChannel);
+	}
+
+	int32 FAutoAssignElement::GetUniverse() const
+	{
+		const int64 Universe = AbsoluteRange.GetLowerBoundValue() / DMX_UNIVERSE_SIZE;
+		return FMath::Clamp(Universe, 1, std::numeric_limits<int32>::max());
+	}
+
+	int32 FAutoAssignElement::GetChannel() const
+	{
+		const int64 Channel = AbsoluteRange.GetLowerBoundValue() % DMX_UNIVERSE_SIZE;
+		return FMath::Clamp(Channel, 1, std::numeric_limits<int32>::max());
+	}
+
+	int64 FAutoAssignElement::GetSize() const
+	{
+		return AbsoluteRange.Size<int64>();
+	}
+
+	int64 FAutoAssignElement::GetLowerBoundValue() const
+	{
+		return AbsoluteRange.GetLowerBoundValue();
+	}
+
+	int64 FAutoAssignElement::GetUpperBoundValue() const
+	{
+		return AbsoluteRange.GetUpperBoundValue();
+	}
+
+	void FAutoAssignElement::SetAbsoluteStartingChannel(int64 NewStartingChannel)
+	{
+		const int64 NewEndingAddress = NewStartingChannel + AbsoluteRange.GetUpperBoundValue() - AbsoluteRange.GetLowerBoundValue();
+		AbsoluteRange = TRange<int64>(NewStartingChannel, NewEndingAddress);
+	}
+
+	void FAutoAssignElement::ApplyToPatch()
+	{
+		if (UDMXEntityFixturePatch* Patch = WeakFixturePatch.Get())
+		{
+			const int64 Universe = (GetLowerBoundValue() + 1) / DMX_UNIVERSE_SIZE;
+			const int64 Channel = (GetLowerBoundValue() + 1) % DMX_UNIVERSE_SIZE;
+			if (!ensureMsgf(Universe >= 1 && Channel >= 1, TEXT("Invalid universe and channel resulting from auto assign. Channels not applied for '%s'."), *Patch->Name))
+			{
+				return;
+			}
+
+			const int32 UniverseSafe = FMath::Clamp(Universe, 1, std::numeric_limits<int32>::max());
+			const int32 ChannelSafe = FMath::Clamp(Channel, 1, std::numeric_limits<int32>::max());
+
+			Patch->SetUniverseID(UniverseSafe);
+			Patch->SetStartingChannel(ChannelSafe);
+		}
+	}
+
+	UDMXEntityFixturePatch* FAutoAssignElement::GetFixturePatch() const
+	{
+		return WeakFixturePatch.Get();
+	}
+
+	int32 FAutoAssignElementsUtility::Assign(TArray<TSharedRef<FAutoAssignElement>> ElementsToAssign, int32 AssignToUniverse, int32 AssignToChannel)
+	{
+		if (ElementsToAssign.IsEmpty())
+		{
+			return 1;
+		}
+
+		if (!EnsureValidPatchElements(ElementsToAssign))
+		{
+			return 1;
+		}
+
+		Algo::StableSortBy(ElementsToAssign, &FAutoAssignElement::GetLowerBoundValue);
+
+		const int64 AssignToChannelAbsolute = FMath::Max(AssignToUniverse * DMX_UNIVERSE_SIZE + AssignToChannel - 1, (int64)DMX_UNIVERSE_SIZE);
+		
+		int64 Offset = AssignToChannelAbsolute - ElementsToAssign[0]->GetLowerBoundValue();
+		for (const TSharedRef<FAutoAssignElement>& Element : ElementsToAssign)
+		{
+			const bool bFitsCurrentUniverse = (Element->GetLowerBoundValue() + Offset) / DMX_UNIVERSE_SIZE == (Element->GetUpperBoundValue() - 1 + Offset) / DMX_UNIVERSE_SIZE;
+			if (bFitsCurrentUniverse)
+			{
+				const int64 NewStartingChannel = Element->GetLowerBoundValue() + Offset;
+				Element->SetAbsoluteStartingChannel(NewStartingChannel);
+			}
+			else
+			{
+				// Assign to next universe
+				const int64 OldStartingChannel = Element->GetLowerBoundValue() + Offset;
+				const int64 NewStartingChannel = (Element->GetLowerBoundValue() + Offset + DMX_UNIVERSE_SIZE) / DMX_UNIVERSE_SIZE * DMX_UNIVERSE_SIZE;
+				Element->SetAbsoluteStartingChannel(NewStartingChannel);
+
+				Offset += NewStartingChannel - OldStartingChannel;
+			}
+
+			Element->ApplyToPatch();
+		}
+
+		const int64 FirstPatchedUniverse = ElementsToAssign[0]->GetLowerBoundValue() / DMX_UNIVERSE_SIZE;
+		return FirstPatchedUniverse;
+	}
+
+	void FAutoAssignElementsUtility::Align(TArray<TSharedRef<FAutoAssignElement>> ElementsToAlign)
+	{
+		if (ElementsToAlign.IsEmpty())
+		{
+			return;
+		}
+
+		if (!EnsureValidPatchElements(ElementsToAlign))
+		{
+			return;
+		}
+
+		Algo::StableSortBy(ElementsToAlign, &FAutoAssignElement::GetLowerBoundValue);
+
+		if (ElementsToAlign[0]->GetLowerBoundValue() < DMX_UNIVERSE_SIZE)
+		{
+			ElementsToAlign[0]->SetAbsoluteStartingChannel(DMX_UNIVERSE_SIZE);
+			ElementsToAlign[0]->ApplyToPatch();
+		}
+
+		for (int32 ElementIndex = 1; ElementIndex < ElementsToAlign.Num(); ElementIndex++)
+		{
+			const TSharedRef<FAutoAssignElement>& PreviousElement = ElementsToAlign[ElementIndex - 1];
+			const TSharedRef<FAutoAssignElement>& Element = ElementsToAlign[ElementIndex];
+
+			const bool bFitsCurrentUniverse = (PreviousElement->GetUpperBoundValue() + Element->GetSize() - 1) / DMX_UNIVERSE_SIZE == PreviousElement->GetUpperBoundValue() / DMX_UNIVERSE_SIZE;
+			if (bFitsCurrentUniverse)
+			{
+				Element->SetAbsoluteStartingChannel(PreviousElement->GetUpperBoundValue());
+			}
+			else
+			{
+				// Assign to next universe
+				const int64 NewStartingChannel = (PreviousElement->GetUpperBoundValue() + DMX_UNIVERSE_SIZE) / DMX_UNIVERSE_SIZE * DMX_UNIVERSE_SIZE;
+				Element->SetAbsoluteStartingChannel(NewStartingChannel);
+			}
+
+			Element->ApplyToPatch();
+		}
+	}
+
+	void FAutoAssignElementsUtility::Stack(TArray<TSharedRef<FAutoAssignElement>> ElementsToStack)
+	{
+		if (ElementsToStack.IsEmpty())
+		{
+			return;
+		}
+
+		if (!EnsureValidPatchElements(ElementsToStack))
+		{
+			return;
+		}
+
+		Algo::StableSortBy(ElementsToStack, &FAutoAssignElement::GetLowerBoundValue);
+
+		// If the largest patch would exceed universe size, offset the stack so the largest patch fits into the universe as well
+		const TSharedRef<FAutoAssignElement>* const LargestElement = Algo::MaxElementBy(ElementsToStack, &FAutoAssignElement::GetSize);
+
+		const int64 EndingChannelForLargestExclusive = ElementsToStack[0]->GetLowerBoundValue() + (*LargestElement)->GetSize();
+		const bool bOffsetToLargestPatch = (EndingChannelForLargestExclusive - 1) / DMX_UNIVERSE_SIZE != ElementsToStack[0]->GetLowerBoundValue() / DMX_UNIVERSE_SIZE;
+		
+		const int64 StackOnChannel = [&ElementsToStack, EndingChannelForLargestExclusive, bOffsetToLargestPatch]()
+		{
+			if (bOffsetToLargestPatch)
+			{
+				const int64 Offset = EndingChannelForLargestExclusive % DMX_UNIVERSE_SIZE;
+				return FMath::Max(ElementsToStack[0]->GetLowerBoundValue() - Offset, (int64)DMX_UNIVERSE_SIZE);
+			}
+			else
+			{
+				return FMath::Max(ElementsToStack[0]->GetLowerBoundValue(), (int64)DMX_UNIVERSE_SIZE);
+			}
+		}();
+	
+		for (const TSharedRef<FAutoAssignElement>& Element : ElementsToStack)
+		{
+			Element->SetAbsoluteStartingChannel(StackOnChannel);
+			Element->ApplyToPatch();
+		}
+	}
+
+	void FAutoAssignElementsUtility::SpreadOverUniverses(TArray<TSharedRef<FAutoAssignElement>> ElementsToSpread)
+	{
+		if (ElementsToSpread.IsEmpty())
+		{
+			return;
+		}
+
+		if (!EnsureValidPatchElements(ElementsToSpread))
+		{
+			return;
+		}
+
+		Algo::StableSortBy(ElementsToSpread, &FAutoAssignElement::GetLowerBoundValue);
+
+		// If the largest patch would exceed universe size, offset the patches so the largest patch fits into the universe as well
+		const TSharedRef<FAutoAssignElement>* const LargestElement = Algo::MaxElementBy(ElementsToSpread, &FAutoAssignElement::GetSize);
+
+		const int64 EndingChannelForLargestExclusive = ElementsToSpread[0]->GetLowerBoundValue() + (*LargestElement)->GetSize();
+		const bool bOffsetToLargestPatch = (EndingChannelForLargestExclusive - 1) / DMX_UNIVERSE_SIZE != ElementsToSpread[0]->GetLowerBoundValue() / DMX_UNIVERSE_SIZE;
+
+		int64 SpreadOnChannel = [&ElementsToSpread, EndingChannelForLargestExclusive, bOffsetToLargestPatch]()
+		{
+			if (bOffsetToLargestPatch)
+			{
+				const int64 Offset = EndingChannelForLargestExclusive % DMX_UNIVERSE_SIZE;
+				return FMath::Max(ElementsToSpread[0]->GetLowerBoundValue() - Offset, (int64)DMX_UNIVERSE_SIZE);
+			}
+			else
+			{
+				return FMath::Max(ElementsToSpread[0]->GetLowerBoundValue(), (int64)DMX_UNIVERSE_SIZE);
+			}
+		}();
+
+		for (const TSharedRef<FAutoAssignElement>& Element : ElementsToSpread)
+		{
+			Element->SetAbsoluteStartingChannel(SpreadOnChannel);
+			SpreadOnChannel += DMX_UNIVERSE_SIZE;
+
+			Element->ApplyToPatch();
+		}
+	}
+
+	void FAutoAssignElementsUtility::AutoAssign(const TArray<TSharedRef<FAutoAssignElement>>& FreeElements, TArray<TSharedRef<FAutoAssignElement>> AutoAssignElements)
+	{
+		if (AutoAssignElements.IsEmpty())
+		{
+			return;
+		}
+
+		if (!EnsureValidPatchElements(AutoAssignElements))
+		{
+			return;
+		}
+		Algo::StableSortBy(AutoAssignElements, &FAutoAssignElement::GetLowerBoundValue);
+
+		TArray<int64> StartingChannels;
+		for (const TSharedRef<FAutoAssignElement>& FreeElement : FreeElements)
+		{
+			StartingChannels.Reset();
+			const int64 AssignToChannelAbsolute = FMath::Max(FreeElement->GetLowerBoundValue(), (int64)DMX_UNIVERSE_SIZE);
+
+			int64 Offset = AssignToChannelAbsolute - AutoAssignElements[0]->GetLowerBoundValue();
+			for (const TSharedRef<FAutoAssignElement>& Element : AutoAssignElements)
+			{
+				const bool bFitsCurrentUniverse = (Element->GetLowerBoundValue() + Offset) / DMX_UNIVERSE_SIZE == (Element->GetUpperBoundValue() - 1 + Offset) / DMX_UNIVERSE_SIZE;
+				if (bFitsCurrentUniverse)
+				{
+					const int64 NewStartingChannel = Element->GetLowerBoundValue() + Offset;
+					StartingChannels.Add(NewStartingChannel);
+				}
+				else
+				{
+					// Assign to next universe
+					const int64 OldStartingChannel = Element->GetLowerBoundValue() + Offset;
+					const int64 NewStartingChannel = (Element->GetLowerBoundValue() + Offset + DMX_UNIVERSE_SIZE) / DMX_UNIVERSE_SIZE * DMX_UNIVERSE_SIZE;
+
+					StartingChannels.Add(NewStartingChannel);
+
+					Offset += NewStartingChannel - OldStartingChannel;
+				}
+			}
+
+			const TRange<int64> AssignedRange(StartingChannels[0], StartingChannels.Last() + AutoAssignElements.Last()->GetSize() - 1);
+			if (AssignedRange.Size<int64>() <= FreeElement->GetSize())
+			{
+				for (int64 ElementIndex = 0; ElementIndex < AutoAssignElements.Num(); ElementIndex++)
+				{
+					AutoAssignElements[ElementIndex]->SetAbsoluteStartingChannel(StartingChannels[ElementIndex]);
+					AutoAssignElements[ElementIndex]->ApplyToPatch();
+				}
+
+				return;
+			}
+		}
+	}
+
+	bool FAutoAssignElementsUtility::EnsureValidPatchElements(const TArray<TSharedRef<FAutoAssignElement>>& Elements)
+	{
+		const bool bValidElements = Algo::FindByPredicate(Elements, [](const TSharedRef<FAutoAssignElement>& Element)
+			{
+				// Find an invalid element
+				return Element->GetLowerBoundValue() > Element->GetUpperBoundValue() || Element->GetSize() < 1 || Element->GetSize() > DMX_UNIVERSE_SIZE;
+			}) == nullptr;
+		if (!ensureMsgf(bValidElements, TEXT("Elements to assign are not in valid DMX range. Ignoring call.")))
+		{
+			return false;
+		}
+		return true;
+	}
+}
+
+
+namespace UE::DMXEditor::AutoAssign
+{
+	int32 FAutoAssignUtility::Assign(TArray<UDMXEntityFixturePatch*> FixturePatches, int32 AssignToUniverse, int32 AssignToChannel)
+	{
+		FAutoAssignUtility Instance;
+
+		using namespace Private;
+		const TArray<TSharedRef<FAutoAssignElement>> PatchElements = Instance.CreatePatchElements(FixturePatches);
+		return FAutoAssignElementsUtility::Assign(PatchElements, AssignToUniverse, AssignToChannel);
+	}
+
+	void FAutoAssignUtility::Align(TArray<UDMXEntityFixturePatch*> FixturePatches)
+	{
+		FAutoAssignUtility Instance;
+
+		using namespace Private;
+		const TArray<TSharedRef<FAutoAssignElement>> PatchElements = Instance.CreatePatchElements(FixturePatches);
+		FAutoAssignElementsUtility::Align(PatchElements);
+	}
+
+	void FAutoAssignUtility::Stack(TArray<UDMXEntityFixturePatch*> FixturePatches)
+	{
+		FAutoAssignUtility Instance;
+
+		using namespace Private;
+		const TArray<TSharedRef<FAutoAssignElement>> PatchElements = Instance.CreatePatchElements(FixturePatches);
+		FAutoAssignElementsUtility::Stack(PatchElements);
+	}
+
+	void FAutoAssignUtility::SpreadOverUniverses(TArray<UDMXEntityFixturePatch*> FixturePatches)
+	{
+		FAutoAssignUtility Instance;
+
+		using namespace Private;
+		const TArray<TSharedRef<FAutoAssignElement>> PatchElements = Instance.CreatePatchElements(FixturePatches);
+		FAutoAssignElementsUtility::SpreadOverUniverses(PatchElements);
+	}
+
+	int32 FAutoAssignUtility::AutoAssign(EAutoAssignMode Mode, const TSharedRef<FDMXEditor>& DMXEditor, TArray<UDMXEntityFixturePatch*> FixturePatches, int32 UniverseUnderMouse, int32 ChannelUnderMouse)
+	{
+		FAutoAssignUtility Instance;
+		return Instance.AutoAssignInternal(Mode, DMXEditor, FixturePatches, UniverseUnderMouse, ChannelUnderMouse);
+	}
+
+	int32 FAutoAssignUtility::AutoAssignInternal(EAutoAssignMode Mode, const TSharedRef<FDMXEditor>& DMXEditor, TArray<UDMXEntityFixturePatch*> FixturePatches, int32 UniverseUnderMouse, int32 ChannelUnderMouse)
+	{
+		UDMXLibrary* DMXLibrary = GetDMXLibrary(FixturePatches);
+		if (!DMXLibrary)
+		{
+			return 1;
+		}
+
+		using namespace Private;
+		const int64 AbsoluteStartingChannel = FindAbsoluteStartingChannel(Mode, DMXEditor , *DMXLibrary, FixturePatches, UniverseUnderMouse, ChannelUnderMouse);
+		const TArray<TSharedRef<FAutoAssignElement>> FreeElements = CreateFreeElements(*DMXLibrary, FixturePatches, AbsoluteStartingChannel);
+		const TArray<TSharedRef<FAutoAssignElement>> PatchElements = CreatePatchElements(FixturePatches);
+
+		FAutoAssignElementsUtility::AutoAssign(FreeElements, PatchElements);
+
+		const TSharedRef<FAutoAssignElement>* const FirstPatchElementPtr = Algo::FindByPredicate(PatchElements, [](const TSharedRef<FAutoAssignElement>& Element)
+			{
+				return Element->GetLowerBoundValue();
+			});
+		if (!ensureMsgf(FirstPatchElementPtr, TEXT("Unexpected FAutoAssignElementsUtility::AutoAssign cannot find first patch elements.")))
+		{
+			return 1;
+		}
+
+		return (*FirstPatchElementPtr)->GetLowerBoundValue() / DMX_UNIVERSE_SIZE;
+	}
+
+	int64 FAutoAssignUtility::FindAbsoluteStartingChannel(EAutoAssignMode Mode, const TSharedRef<FDMXEditor>& DMXEditor, const UDMXLibrary& DMXLibrary, const TArray<UDMXEntityFixturePatch*>& FixturePatches, int32 UniverseUnderMouse, int32 ChannelUnderMouse) const
+	{
+		if (!ensureMsgf(!FixturePatches.IsEmpty() && !FixturePatches.Contains(nullptr), TEXT("Unexpected no fixture patches or null entry in Fixture Patches. Expected to be handled during FAutoAssignUtility construction.")))
+		{
+			return DMX_UNIVERSE_SIZE;
+		}
+
+		if (!ensureMsgf(
+			Mode != EAutoAssignMode::ChannelUnderMouse || 
+			(UniverseUnderMouse > 0 && UniverseUnderMouse <= DMX_UNIVERSE_SIZE && ChannelUnderMouse > 0 && ChannelUnderMouse <= DMX_UNIVERSE_SIZE),
+			TEXT("Trying to auto assign to channel under mouse, but provided channels are not valid. Skipping auto assign.")))
+		{
+			return DMX_UNIVERSE_SIZE;
+		}
+
+		if (Mode == EAutoAssignMode::ChannelUnderMouse)
+		{
+			return (int64)UniverseUnderMouse * DMX_UNIVERSE_SIZE + ChannelUnderMouse - 1;
+		}
+		else if (Mode == EAutoAssignMode::SelectedUniverse)
+		{
+			return DMXEditor->GetFixturePatchSharedData()->GetSelectedUniverse() * DMX_UNIVERSE_SIZE;
+		}
+		else if (Mode == EAutoAssignMode::FirstReachableUniverse)
+		{
+			const FDMXInputPortSharedRef* const FirstReachableInputPortPtr = Algo::MinElementBy(DMXLibrary.GetInputPorts(), &FDMXInputPort::GetLocalUniverseStart);
+			const FDMXOutputPortSharedRef* const FirstReachableOutputPortPtr = Algo::MinElementBy(DMXLibrary.GetOutputPorts(), &FDMXOutputPort::GetLocalUniverseStart);
+
+			if (FirstReachableInputPortPtr && FirstReachableOutputPortPtr)
+			{
+				return (int64)FMath::Max((*FirstReachableInputPortPtr)->GetLocalUniverseStart(), (*FirstReachableOutputPortPtr)->GetLocalUniverseStart()) * DMX_UNIVERSE_SIZE;
+			}
+			else if (FirstReachableInputPortPtr)
+			{
+				return  (int64)(*FirstReachableInputPortPtr)->GetLocalUniverseStart() * DMX_UNIVERSE_SIZE;
+			}
+			else if (FirstReachableOutputPortPtr)
+			{
+				return  (int64)(*FirstReachableOutputPortPtr)->GetLocalUniverseStart() * DMX_UNIVERSE_SIZE;
+			}
+			
+			// If there's no valid ports, assign to 1.1
+			return DMX_UNIVERSE_SIZE;
+		}
+		else if (Mode == EAutoAssignMode::AfterLastPatchedUniverse)
+		{
+			const TArray<UDMXEntityFixturePatch*> OtherPatches = FindOtherPatchesInLibrary(DMXLibrary, FixturePatches);
+
+			const UDMXEntityFixturePatch* const* LastOtherFixturePatchPtr = Algo::MaxElementBy(OtherPatches, [](const UDMXEntityFixturePatch* OtherPatch)
+				{
+					return OtherPatch->GetUniverseID() * DMX_UNIVERSE_SIZE + OtherPatch->GetStartingChannel() + OtherPatch->GetChannelSpan();
+				});
+
+			if (LastOtherFixturePatchPtr)
+			{
+				return  (int64)(*LastOtherFixturePatchPtr)->GetUniverseID() * DMX_UNIVERSE_SIZE + (*LastOtherFixturePatchPtr)->GetStartingChannel() + (*LastOtherFixturePatchPtr)->GetChannelSpan();
+			}
+			else
+			{
+				return DMX_UNIVERSE_SIZE;
+			}
+		}
+		else
+		{
+			checkf(0, TEXT("Unhandled enum value in FAutoAssignUtility::FindAbsoluteStartingChannel"));
+			return DMX_UNIVERSE_SIZE;
+		}
+	}	
+	
+	UDMXLibrary* FAutoAssignUtility::GetDMXLibrary(TArray<UDMXEntityFixturePatch*> FixturePatches) const
+	{
+		if (FixturePatches.IsEmpty())
+		{
+			return nullptr;
+		}
+		FixturePatches.Remove(nullptr);
+
+		// Ensure all patches reside in the same library
+		UDMXLibrary* DMXLibrary = FixturePatches[0]->GetParentLibrary();
+		if (!ensureMsgf(DMXLibrary, TEXT("Trying to auto assign fixture patches but the DMX Library of its patches is not valid.")))
+		{
+			return nullptr;
+		}
+		check(DMXLibrary);
+
+		for (const UDMXEntityFixturePatch* FixturePatch : FixturePatches)
+		{
+			if (!ensureMsgf(FixturePatch->GetParentLibrary() == DMXLibrary, TEXT("Trying to auto assign fixture patches, but the patches don't share a common DMX Library. This is not supported")))
+			{
+				return nullptr;
+			}
+		}
+
+		return DMXLibrary;
+	}
+
+	TArray<TSharedRef<Private::FAutoAssignElement>> FAutoAssignUtility::CreateFreeElements(UDMXLibrary& DMXLibrary, const TArray<UDMXEntityFixturePatch*>& FixturePatches, int64 AbsoluteStartingChannel) const
+	{
+		using namespace Private;
+		TArray<UDMXEntityFixturePatch*> OtherPatchesInLibrary = FindOtherPatchesInLibrary(DMXLibrary, FixturePatches);
+		Algo::StableSortBy(OtherPatchesInLibrary, [](const UDMXEntityFixturePatch* FixturePatch)
+			{
+				return FixturePatch->GetUniverseID() * DMX_UNIVERSE_SIZE + FixturePatch->GetStartingChannel();
+			});
+
+		TArray<FAutoAssignElement> OtherPatchElements;
+		for (UDMXEntityFixturePatch* FixturePatch : OtherPatchesInLibrary)
+		{
+			if (IsValid(FixturePatch))
+			{
+				OtherPatchElements.Add(FAutoAssignElement(*FixturePatch));
+			}
+		}
+
+		TArray<TSharedRef<FAutoAssignElement>> FreeElements;
+		int64 StartingChannelOfFreeRange = FMath::Max(AbsoluteStartingChannel, (int64)DMX_UNIVERSE_SIZE);
+		for (const FAutoAssignElement& OtherPatchElement : OtherPatchElements)
+		{
+			if (StartingChannelOfFreeRange < OtherPatchElement.GetUpperBoundValue())
+			{
+				const TRange<int64> FreeRange(StartingChannelOfFreeRange, OtherPatchElement.GetLowerBoundValue() - 1);
+				FreeElements.Add(MakeShared<FAutoAssignElement>(FreeRange));
+
+				StartingChannelOfFreeRange = OtherPatchElement.GetUpperBoundValue();
+			}
+		}
+		const TSharedRef<FAutoAssignElement> LastFreeElement = MakeShared<FAutoAssignElement>(TRange<int64>(StartingChannelOfFreeRange, std::numeric_limits<int64>::max()));
+		FreeElements.Add(LastFreeElement);
+
+		return FreeElements;
+	}
+
+	TArray<TSharedRef<Private::FAutoAssignElement>> FAutoAssignUtility::CreatePatchElements(TArray<UDMXEntityFixturePatch*> FixturePatches) const
+	{
+		using namespace Private;
+		TArray<TSharedRef<FAutoAssignElement>> PatchElements;
+		if (!ensureMsgf(!FixturePatches.IsEmpty() && !FixturePatches.Contains(nullptr), TEXT("Unexpected no fixture patches or null entry in Fixture Patches. Expected to be handled during FAutoAssignUtility construction.")))
+		{
+			return PatchElements;
+		}
+
+		Algo::StableSortBy(FixturePatches, [](const UDMXEntityFixturePatch* FixturePatch)
+			{
+				return FixturePatch->GetUniverseID() * DMX_UNIVERSE_SIZE + FixturePatch->GetStartingChannel();
+			});
+
+		Algo::TransformIf(FixturePatches, PatchElements,
+			[](UDMXEntityFixturePatch* FixturePatch)
+			{
+				return IsValid(FixturePatch);
+			},
+			[](UDMXEntityFixturePatch* FixturePatch)
+			{
+				return MakeShared<FAutoAssignElement>(*FixturePatch);
+			});
+
+		Algo::StableSortBy(PatchElements, &FAutoAssignElement::GetLowerBoundValue);
+
+		return PatchElements;
+	}
+
+	TArray<UDMXEntityFixturePatch*> FAutoAssignUtility::FindOtherPatchesInLibrary(const UDMXLibrary& DMXLibrary, const TArray<UDMXEntityFixturePatch*>& FixturePatches) const
+	{
+		TArray<UDMXEntityFixturePatch*> Result = DMXLibrary.GetEntitiesTypeCast<UDMXEntityFixturePatch>();
+		Result.RemoveAll([FixturePatches](const UDMXEntityFixturePatch* FixturePatch)
+			{
+				return
+					!FixturePatch ||
+					!FixturePatch->GetFixtureType() ||
+					FixturePatches.Contains(FixturePatch);
+			});
+
+
+		Algo::StableSortBy(Result, [](UDMXEntityFixturePatch* FixturePatch)
+			{
+				return ((uint64)FixturePatch->GetUniverseID() * DMX_UNIVERSE_SIZE) + FixturePatch->GetStartingChannel();
+			});
+
+		return Result;
+	}
+}
+
+#undef LOCTEXT_NAMESPACE
diff --git a/Engine/Plugins/VirtualProduction/DMX/DMXEngine/Source/DMXEditor/Private/FixturePatchAutoAssignUtility.h b/Engine/Plugins/VirtualProduction/DMX/DMXEngine/Source/DMXEditor/Private/FixturePatchAutoAssignUtility.h
new file mode 100644
--- /dev/null
+++ b/Engine/Plugins/VirtualProduction/DMX/DMXEngine/Source/DMXEditor/Private/FixturePatchAutoAssignUtility.h
@@ -0,0 +1,129 @@
+// Copyright Epic Games, Inc. All Rights Reserved.
+
+#pragma once
+
+#include "Containers/Array.h"
+#include "HAL/Platform.h"
+#include "Math/Range.h"
+#include "UObject/WeakObjectPtr.h"
+
+class FDMXEditor;
+class FMenuBuilder;
+class UDMXEntityFixturePatch;
+class UDMXLibrary;
+
+
+namespace UE::DMXEditor::AutoAssign::Private
+{
+	/** 
+	 * An element to auto assign.
+	 * Uses absolute DMX space, e.g. Universe 1 Channel 1 == 512.
+	 */
+	struct FAutoAssignElement
+	{
+		FAutoAssignElement(int64 LowerBoundInclusive, int64 UpperBoundExclusive);
+		FAutoAssignElement(const TRange<int64>& InAbsoluteRange);
+		FAutoAssignElement(UDMXEntityFixturePatch& FixturePatch);
+
+		int64 GetLowerBoundValue() const;
+		int64 GetUpperBoundValue() const;
+		int64 GetSize() const;
+		const TRange<int64>& GetRange() const { return AbsoluteRange; }
+
+		void SetAbsoluteStartingChannel(int64 NewStartingAddress);
+
+		void ApplyToPatch();
+
+		UDMXEntityFixturePatch* GetFixturePatch() const;
+
+	private:
+		int32 GetUniverse() const;
+		int32 GetChannel() const;
+
+		/** The absolute range, starting from universe 0 absolute channel 0 */
+		TRange<int64> AbsoluteRange;
+
+		/** The corresponding fixture patch */
+		TWeakObjectPtr<UDMXEntityFixturePatch> WeakFixturePatch;
+	};
+
+	/** 
+	 * Auto-assigns elements to free ranges.
+	 *
+	 * Uses absolute channel values, e.g. Universe 1 Channel 1 == 512 resp. DMX_UNIVERSE_SIZE.
+	 * Applies results to patches on success.
+	 */
+	class FAutoAssignElementsUtility
+	{
+	public:
+		/** Assigns the elements. Returns the universe of the first assigned patch */
+		static int32 Assign(TArray<TSharedRef<FAutoAssignElement>> ElementsToAssign, int32 AssignToUniverse, int32 AssignToChannel);
+
+		/** Aligns the elements */
+		static void Align(TArray<TSharedRef<FAutoAssignElement>> ElementsToAlign);
+
+		/** Stacks the elements */
+		static void Stack(TArray<TSharedRef<FAutoAssignElement>> ElementsToStack);
+
+		/** Spreads the elements over universes */
+		static void SpreadOverUniverses(TArray<TSharedRef<FAutoAssignElement>> ElementsToSpread);
+
+		/** Auto-assigns a set of elements to free elements. Note, free elements need be in the range of { DMX_UNIVERSE_SIZE, std::numeric_limits<int64>::Max() ), ensured. */
+		static void AutoAssign(const TArray<TSharedRef<FAutoAssignElement>>& FreeElements, TArray<TSharedRef<FAutoAssignElement>> AutoAssignElements);
+
+	private:
+		/** Returns true if elments are in valid DMX range. Hits an ensure condition if elements are not valid. */
+		static bool EnsureValidPatchElements(const TArray<TSharedRef<FAutoAssignElement>>& Elements);
+	};
+}
+
+
+namespace UE::DMXEditor::AutoAssign
+{
+	enum class EAutoAssignMode : uint8
+	{
+		ChannelUnderMouse,
+		SelectedUniverse,
+		FirstReachableUniverse,
+		AfterLastPatchedUniverse
+	};
+
+	/** Utility to auto assign fixture patches. */
+	class FAutoAssignUtility
+	{
+	public:
+		/** Aligns FixturePatches */
+		static int32 Assign(TArray<UDMXEntityFixturePatch*> FixturePatches, int32 AssignToUniverse, int32 AssignToChannel);
+
+		/** Aligns FixturePatches */
+		static void Align(TArray<UDMXEntityFixturePatch*> FixturePatches);
+
+		/** Stacks FixturePatches */
+		static void Stack(TArray<UDMXEntityFixturePatch*> FixturePatches);
+
+		/** Auto assigns FixturePatches, one patch per Universe. Returns the Universe the first patch was auto assigned to.*/
+		static void SpreadOverUniverses(TArray<UDMXEntityFixturePatch*> FixturePatches);
+
+		/** Auto assigns FixturePatches. Returns the Universe the first patch was auto assigned to.*/
+		static int32 AutoAssign(EAutoAssignMode Mode, const TSharedRef<FDMXEditor>& DMXEditor, TArray<UDMXEntityFixturePatch*> FixturePatches, int32 UniverseUnderMouse = 0, int32 ChannelUnderMouse = 0);
+
+	private:
+		/** Non-static implementation of auto assign */
+		int32 AutoAssignInternal(EAutoAssignMode Mode, const TSharedRef<FDMXEditor>& DMXEditor, TArray<UDMXEntityFixturePatch*> FixturePatches, int32 UniverseUnderMouse = 0, int32 ChannelUnderMouse = 0);
+
+		/** Finds the absolute starting channel from options */
+		int64 FindAbsoluteStartingChannel(EAutoAssignMode Mode, const TSharedRef<FDMXEditor>& DMXEditor, const UDMXLibrary& DMXLibrary, const TArray<UDMXEntityFixturePatch*>& FixturePatches, int32 UniverseUnderMouse, int32 ChannelUnderMouse) const;
+
+		/** Gets the DMX Library from the patches, or returns nullptr if a valid DMX Library cannot be deduced from patches */
+		UDMXLibrary* GetDMXLibrary(TArray<UDMXEntityFixturePatch*> FixturePatches) const;
+
+		/** Creates elements for ranges that are not patched in the DMX Library */
+		TArray<TSharedRef<Private::FAutoAssignElement>> CreateFreeElements(UDMXLibrary& DMXLibrary, const TArray<UDMXEntityFixturePatch*>& FixturePatches, int64 AbsoluteStartingChannel) const;
+
+		/** Creates elements from patches. Note also creates elements for gaps between the patches to remember offsets */
+		TArray<TSharedRef<Private::FAutoAssignElement>> CreatePatchElements(TArray<UDMXEntityFixturePatch*> FixturePatches) const;
+
+		/** Helper to find other patches in the library */
+		TArray<UDMXEntityFixturePatch*> FindOtherPatchesInLibrary(const UDMXLibrary& DMXLibrary, const TArray<UDMXEntityFixturePatch*>& FixturePatches) const;
+	};
+}
diff --git a/Engine/Plugins/VirtualProduction/DMX/DMXEngine/Source/DMXEditor/Private/Widgets/FixturePatch/DMXAddFixturePatchMenuData.h b/Engine/Plugins/VirtualProduction/DMX/DMXEngine/Source/DMXEditor/Private/Widgets/FixturePatch/DMXAddFixturePatchMenuData.h
new file mode 100644
--- /dev/null
+++ b/Engine/Plugins/VirtualProduction/DMX/DMXEngine/Source/DMXEditor/Private/Widgets/FixturePatch/DMXAddFixturePatchMenuData.h
@@ -0,0 +1,28 @@
+// Copyright Epic Games, Inc. All Rights Reserved.
+
+#pragma once
+
+#include "UObject/Object.h"
+
+#include "DMXAddFixturePatchMenuData.generated.h"
+
+
+class UDMXEntityFixtureType;
+
+
+UCLASS(Config = DMXEditor)
+class UDMXAddFixturePatchMenuData
+	: public UObject
+{
+	GENERATED_BODY()
+
+public:
+	UPROPERTY(Config)
+	TSoftObjectPtr<UDMXEntityFixtureType> SoftFixtureType;
+
+	UPROPERTY(Config)
+	int32 ActiveModeIndex = 0;
+
+	UPROPERTY(Config)
+	int32 NumPatches = 0;
+};
diff --git a/Engine/Plugins/VirtualProduction/DMX/DMXEngine/Source/DMXEditor/Private/Widgets/FixturePatch/SAddFixturePatchMenu.cpp b/Engine/Plugins/VirtualProduction/DMX/DMXEngine/Source/DMXEditor/Private/Widgets/FixturePatch/SAddFixturePatchMenu.cpp
new file mode 100644
--- /dev/null
+++ b/Engine/Plugins/VirtualProduction/DMX/DMXEngine/Source/DMXEditor/Private/Widgets/FixturePatch/SAddFixturePatchMenu.cpp
@@ -0,0 +1,360 @@
+// Copyright Epic Games, Inc. All Rights Reserved.
+
+#include "SAddFixturePatchMenu.h"
+
+#include "Algo/Copy.h"
+#include "Algo/Find.h"
+#include "DMXAddFixturePatchMenuData.h"
+#include "DMXFixturePatchSharedData.h"
+#include "Editor.h"
+#include "Framework/Application/SlateApplication.h"
+#include "FixturePatchAutoAssignUtility.h"
+#include "Library/DMXEntityFixtureType.h"
+#include "ScopedTransaction.h"
+#include "TimerManager.h"
+#include "UObject/Package.h"
+#include "Widgets/Input/SButton.h"
+#include "Widgets/Layout/SGridPanel.h"
+#include "Widgets/SBoxPanel.h"
+#include "Widgets/SDMXEntityDropdownMenu.h"
+#include "Widgets/Text/SInlineEditableTextBlock.h"
+
+
+#define LOCTEXT_NAMESPACE "SAddFixturePatchMenu"
+
+namespace UE::DMXEditor::FixturePatchEditor
+{
+
+	SAddFixturePatchMenu::~SAddFixturePatchMenu()
+	{
+		UDMXAddFixturePatchMenuData* Model = GetMutableDefault<UDMXAddFixturePatchMenuData>();
+		Model->SoftFixtureType = WeakFixtureType.IsValid() ? WeakFixtureType.Get() : Model->SoftFixtureType;
+		Model->ActiveModeIndex = WeakFixtureType.IsValid() ? ActiveModeIndex : Model->ActiveModeIndex;
+		Model->SaveConfig();
+	}
+
+	void SAddFixturePatchMenu::Construct(const FArguments& InArgs, TWeakPtr<FDMXEditor> InWeakDMXEditor)
+	{
+		WeakDMXEditor = InWeakDMXEditor;
+
+		ChildSlot
+		[
+			SNew(SGridPanel)
+			.FillColumn(1, 1.f)
+
+			// Select Fixture Type
+			+ SGridPanel::Slot(0, 0)
+			.Padding(4.f)
+			.VAlign(VAlign_Center)
+			[
+				SNew(STextBlock)
+				.Text(LOCTEXT("SelectFixtureTypeLabel", "Fixture Type"))
+				.Font(FAppStyle::GetFontStyle(TEXT("PropertyWindow.NormalFont")))
+			]
+			+ SGridPanel::Slot(1, 0)
+				.Padding(4.f)
+			.VAlign(VAlign_Center)
+			[
+				SNew(SBox)
+				.MinDesiredWidth(100.f)
+				[
+					SAssignNew(FixtureTypeSelector, SDMXEntityPickerButton<UDMXEntityFixtureType>)
+					.DMXEditor(WeakDMXEditor)
+					.CurrentEntity_Lambda([this] { return WeakFixtureType.Get(); })
+					.OnEntitySelected(this, &SAddFixturePatchMenu::OnFixtureTypeSelected)	
+				]
+			]
+
+			// Select Mode
+			+ SGridPanel::Slot(0, 1)
+			.Padding(4.f)
+			.VAlign(VAlign_Center)
+			[
+				SNew(STextBlock)
+				.IsEnabled(this, &SAddFixturePatchMenu::HasValidFixtureTypeAndMode)
+				.Text(LOCTEXT("SelectModeLabel", "Mode"))
+				.Font(FAppStyle::GetFontStyle(TEXT("PropertyWindow.NormalFont")))
+			]
+			+ SGridPanel::Slot(1, 1)
+			.Padding(4.f)
+			[
+				SAssignNew(ModeComboBox, SComboBox<TSharedPtr<uint32>>)
+				.IsEnabled(this, &SAddFixturePatchMenu::HasValidFixtureTypeAndMode)
+				.OptionsSource(&ModeSources)
+				.OnGenerateWidget(this, &SAddFixturePatchMenu::GenerateModeComboBoxEntry)
+				.OnSelectionChanged(this, &SAddFixturePatchMenu::OnModeSelected)
+				.InitiallySelectedItem(0)
+				[
+					SNew(STextBlock)
+					.MinDesiredWidth(50.0f)
+					.Text(this, &SAddFixturePatchMenu::GetActiveModeText)
+					.Font(FAppStyle::GetFontStyle(TEXT("PropertyWindow.NormalFont")))
+				]
+			]
+
+			// Num Fixture Patches Label
+			+ SGridPanel::Slot(0, 2)
+			.Padding(4.f)
+			.VAlign(VAlign_Center)
+			[
+				SNew(STextBlock)
+				.IsEnabled(this, &SAddFixturePatchMenu::HasValidFixtureTypeAndMode)
+				.Text(LOCTEXT("NumPatchesLabel", "Num Patches"))
+				.Font(FAppStyle::GetFontStyle(TEXT("PropertyWindow.NormalFont")))
+			]
+
+			// Num Fixture Patches editable Text Box
+			+ SGridPanel::Slot(1, 2)
+			.Padding(4.f)
+			[
+				SNew(SBox)
+				.HAlign(HAlign_Left)
+				[
+					SNew(SEditableTextBox)
+					.IsEnabled(this, &SAddFixturePatchMenu::HasValidFixtureTypeAndMode)
+					.MinDesiredWidth(60.f)
+					.SelectAllTextWhenFocused(true)
+					.ClearKeyboardFocusOnCommit(false)
+					.SelectAllTextOnCommit(true)
+					.Text_Lambda([this]()
+						{
+							return FText::FromString(FString::FromInt(NumFixturePatchesToAdd));
+						})
+					.OnVerifyTextChanged_Lambda([](const FText& InNewText, FText& OutErrorMessage)
+						{
+							int32 Value;
+							if (!LexTryParseString<int32>(Value, *InNewText.ToString()) ||
+								Value < 1)
+							{
+								OutErrorMessage = LOCTEXT("BulkAddPatchesBadString", "Numeric values > 0 only");
+								return false;
+							}
+							return true;
+						})
+					.OnTextCommitted_Lambda([this](const FText& Text, ETextCommit::Type CommitType)
+						{
+							int32 Value;
+							if (LexTryParseString<int32>(Value, *Text.ToString()) &&
+								Value > 0)
+							{
+								constexpr int32 MaxNumPatchesToBulkAdd = 32768;
+								NumFixturePatchesToAdd = FMath::Min(Value, MaxNumPatchesToBulkAdd);
+							}
+						})
+				]
+			]
+
+			// Add fixture patches button
+			+ SGridPanel::Slot(1, 3)
+			.Padding(4.f)
+			[
+				SNew(SBox)
+				.HAlign(HAlign_Right)
+				.MinDesiredWidth(120.f)
+				[
+					SNew(SButton)
+					.IsEnabled(this, &SAddFixturePatchMenu::HasValidFixtureTypeAndMode)
+					.ContentPadding(FMargin(4.0f, 4.0f))
+					.ButtonStyle(FAppStyle::Get(), "FlatButton.Success")
+					.ForegroundColor(FLinearColor::White)
+					.OnClicked(this, &SAddFixturePatchMenu::OnAddFixturePatchButtonClicked)
+					[
+						SNew(STextBlock)
+						.Text_Lambda([this]()
+						{
+							return FText::Format(LOCTEXT("AddFixturePatchButtonText", "Add Fixture {0}|plural(one=Patch, other=Patches)"), NumFixturePatchesToAdd);
+						})
+						.Font(FAppStyle::GetFontStyle(TEXT("PropertyWindow.BoldFont")))
+					]
+				]
+			]
+		];
+
+		RequestRefresh();
+	}
+
+	void SAddFixturePatchMenu::ForceRefresh()
+	{
+		RequestRefreshModeComboBoxTimerHandle.Invalidate();
+
+		UDMXLibrary* DMXLibrary = WeakDMXEditor.IsValid() ? WeakDMXEditor.Pin()->GetDMXLibrary() : nullptr;
+		if (!WeakDMXEditor.IsValid())
+		{
+			return;
+		}
+		const TArray<UDMXEntityFixtureType*> FixtureTypes = DMXLibrary->GetEntitiesTypeCast<UDMXEntityFixtureType>();
+
+		// Init 
+		UDMXAddFixturePatchMenuData* MenuData = GetMutableDefault<UDMXAddFixturePatchMenuData>();
+		UDMXEntityFixtureType* const* SelectedFixtureTypePtr = Algo::FindByPredicate(FixtureTypes, [MenuData](const UDMXEntityFixtureType* FixtureType)
+			{
+				return FixtureType == MenuData->SoftFixtureType;
+			});
+		if (!SelectedFixtureTypePtr)
+		{
+			SelectedFixtureTypePtr = FixtureTypes.IsEmpty() ? nullptr : &FixtureTypes[0];
+		}
+		WeakFixtureType = SelectedFixtureTypePtr ? *SelectedFixtureTypePtr : nullptr;
+		ActiveModeIndex = WeakFixtureType.IsValid() && WeakFixtureType->Modes.IsValidIndex(ActiveModeIndex) ? ActiveModeIndex : 0;
+
+		ModeSources.Reset();
+		if (UDMXEntityFixtureType* FixtureType = WeakFixtureType.Get())
+		{
+			for (int32 ModeIndex = 0; ModeIndex < FixtureType->Modes.Num(); ModeIndex++)
+			{
+				ModeSources.Add(MakeShared<uint32>(ModeIndex));
+			}
+		}
+
+		ModeComboBox->RefreshOptions();
+	}
+
+	TSharedRef<SWidget> SAddFixturePatchMenu::GenerateModeComboBoxEntry(const TSharedPtr<uint32> InModeIndex) const
+	{
+		UDMXEntityFixtureType* FixtureType = Cast<UDMXEntityFixtureType>(WeakFixtureType);
+		if (!FixtureType)
+		{
+			return SNullWidget::NullWidget;
+		}
+
+		const TArray<FDMXFixtureMode>& Modes = FixtureType->Modes;
+
+		return
+			SNew(STextBlock)
+			.MinDesiredWidth(50.0f)
+			.Text_Lambda([&Modes, ModeIndex = *InModeIndex, this]()
+				{	
+					const UDMXAddFixturePatchMenuData* MenuData = GetDefault<UDMXAddFixturePatchMenuData>();
+					if (Modes.IsValidIndex(ModeIndex))
+					{
+						return FText::FromString(Modes[ModeIndex].ModeName);
+					}
+					else
+					{
+						return LOCTEXT("NoModeAvailableText", "No Mode available");
+					}
+				})
+			.Font(FAppStyle::GetFontStyle(TEXT("PropertyWindow.NormalFont")));
+	}
+
+	void SAddFixturePatchMenu::RequestRefresh()
+	{ 
+		if (!RequestRefreshModeComboBoxTimerHandle.IsValid())
+		{
+			RequestRefreshModeComboBoxTimerHandle = GEditor->GetTimerManager()->SetTimerForNextTick(FTimerDelegate::CreateSP(this, &SAddFixturePatchMenu::ForceRefresh));
+		}
+	}
+
+	void SAddFixturePatchMenu::OnFixtureTypeSelected(UDMXEntity* InSelectedFixtureType)
+	{
+		WeakFixtureType = Cast<UDMXEntityFixtureType>(InSelectedFixtureType);
+
+		RequestRefresh();
+	}
+
+	void SAddFixturePatchMenu::OnModeSelected(TSharedPtr<uint32> InSelectedMode, ESelectInfo::Type SelectInfo)
+	{
+		ActiveModeIndex = InSelectedMode.IsValid() ? *InSelectedMode : 0;
+	}
+
+	FReply SAddFixturePatchMenu::OnAddFixturePatchButtonClicked()
+	{
+		FSlateApplication::Get().DismissAllMenus();
+
+		UDMXEntityFixtureType* FixtureType = Cast<UDMXEntityFixtureType>(WeakFixtureType);
+		if (!FixtureType || FixtureType->Modes.IsEmpty())
+		{
+			return FReply::Handled();
+		}
+
+		const TSharedPtr<FDMXEditor> DMXEditor = WeakDMXEditor.Pin();
+		if (!DMXEditor.IsValid())
+		{
+			return FReply::Handled();
+		}
+		const TSharedPtr<FDMXFixturePatchSharedData> SharedData = DMXEditor->GetFixturePatchSharedData();
+		if (!SharedData.IsValid())
+		{
+			return FReply::Handled();
+		}
+		UDMXLibrary* DMXLibrary = DMXEditor->GetDMXLibrary();
+		if (!DMXLibrary)
+		{
+			return FReply::Handled();
+		}
+
+		// Ensure valid mode
+		UDMXAddFixturePatchMenuData* MenuData = GetMutableDefault<UDMXAddFixturePatchMenuData>();
+		if (!ensureMsgf(FixtureType->Modes.IsValidIndex(MenuData->ActiveModeIndex), TEXT("Cannot apply mode. Mode index is invalid.")))
+		{
+			MenuData->ActiveModeIndex = 0;
+		}
+
+		// Create a new fixture patches
+		const FScopedTransaction CreateFixturePatchTransaction(LOCTEXT("CreateFixturePatchTransaction", "Create Fixture Patch"));
+		DMXLibrary->PreEditChange(UDMXLibrary::StaticClass()->FindPropertyByName(UDMXLibrary::GetEntitiesPropertyName()));
+
+		TArray<UDMXEntityFixturePatch*> NewFixturePatches;
+		NewFixturePatches.Reserve(NumFixturePatchesToAdd);
+		for (uint32 iNumFixturePatchesAdded = 0; iNumFixturePatchesAdded < NumFixturePatchesToAdd; iNumFixturePatchesAdded++)
+		{
+			FDMXEntityFixturePatchConstructionParams FixturePatchConstructionParams;
+			FixturePatchConstructionParams.FixtureTypeRef = FDMXEntityFixtureTypeRef(FixtureType);
+			FixturePatchConstructionParams.ActiveMode = MenuData->ActiveModeIndex;
+			FixturePatchConstructionParams.UniverseID = SharedData->GetSelectedUniverse();
+			FixturePatchConstructionParams.StartingAddress = 1;
+
+			constexpr bool bMarkLibraryDirty = false;
+			UDMXEntityFixturePatch* NewFixturePatch = UDMXEntityFixturePatch::CreateFixturePatchInLibrary(FixturePatchConstructionParams, FixtureType->Name, bMarkLibraryDirty);
+			NewFixturePatches.Add(NewFixturePatch);
+		}
+
+		// Align and Auto assign to the currently selected universe
+		using namespace UE::DMXEditor::AutoAssign;
+		FAutoAssignUtility::Align(NewFixturePatches);
+		const int32 AssignedToUniverse = FAutoAssignUtility::AutoAssign(EAutoAssignMode::SelectedUniverse, DMXEditor.ToSharedRef(), NewFixturePatches);
+
+		DMXLibrary->PostEditChange();
+
+		// Select universe and new fixture patches
+		SharedData->SelectUniverse(AssignedToUniverse);
+
+		TArray<TWeakObjectPtr<UDMXEntityFixturePatch>> NewWeakFixturePatches;
+		NewWeakFixturePatches.Reserve(NumFixturePatchesToAdd);
+		Algo::Copy(NewFixturePatches, NewWeakFixturePatches);
+		SharedData->SelectFixturePatches(NewWeakFixturePatches);
+
+		return FReply::Handled();
+	}
+
+	FText SAddFixturePatchMenu::GetActiveModeText() const
+	{
+		UDMXEntityFixtureType* FixtureType = Cast<UDMXEntityFixtureType>(WeakFixtureType);
+		if (!FixtureType)
+		{
+			return FText::GetEmpty();
+		}
+
+		const TArray<FDMXFixtureMode>& Modes = FixtureType->Modes;
+		const UDMXAddFixturePatchMenuData* MenuData = GetMutableDefault<UDMXAddFixturePatchMenuData>();
+		if (Modes.Num() > 0 && Modes.IsValidIndex(MenuData->ActiveModeIndex))
+		{
+			return FText::FromString(Modes[MenuData->ActiveModeIndex].ModeName);
+		}
+		else if (Modes.IsEmpty())
+		{
+			return LOCTEXT("NoModeAvailableComboButtonText", "No Modes in Fixture Type");
+		}
+		else
+		{
+			return LOCTEXT("NoFixtureTypeSelectedComboButtonText", "No Fixture Type selected");
+		}
+	}
+
+	bool SAddFixturePatchMenu::HasValidFixtureTypeAndMode() const
+	{
+		return WeakFixtureType.IsValid() && !WeakFixtureType->Modes.IsEmpty();
+	}
+}
+
+#undef LOCTEXT_NAMESPACE
diff --git a/Engine/Plugins/VirtualProduction/DMX/DMXEngine/Source/DMXEditor/Private/Widgets/FixturePatch/SAddFixturePatchMenu.h b/Engine/Plugins/VirtualProduction/DMX/DMXEngine/Source/DMXEditor/Private/Widgets/FixturePatch/SAddFixturePatchMenu.h
new file mode 100644
--- /dev/null
+++ b/Engine/Plugins/VirtualProduction/DMX/DMXEngine/Source/DMXEditor/Private/Widgets/FixturePatch/SAddFixturePatchMenu.h
@@ -0,0 +1,82 @@
+// Copyright Epic Games, Inc. All Rights Reserved.
+
+#pragma once
+
+#include "Engine/EngineTypes.h"
+#include "Widgets/SCompoundWidget.h"
+
+class FDMXEditor;
+class FReply;
+template <typename OptionType> class SComboBox; 
+template <typename EntityType> class SDMXEntityPickerButton;
+class UDMXAddFixturePatchMenuData;
+class UDMXEntity;
+class UDMXEntityFixtureType;
+
+
+namespace UE::DMXEditor::FixturePatchEditor
+{
+	/** Editor for Fixture Patches */
+	class SAddFixturePatchMenu final
+		: public SCompoundWidget
+	{
+	public:
+		SLATE_BEGIN_ARGS(SAddFixturePatchMenu)
+		{}
+		SLATE_END_ARGS()
+
+		virtual ~SAddFixturePatchMenu();
+
+		/** Constructs this widget */
+		void Construct(const FArguments& InArgs, TWeakPtr<FDMXEditor> InWeakDMXEditor);
+
+		/** Refreshes the mode combo box on the next tick */
+		void RequestRefresh();
+
+	private:
+		/** Generates an entry in the mode combo box */
+		TSharedRef<SWidget> GenerateModeComboBoxEntry(const TSharedPtr<uint32> InMode) const;
+
+		/** Refreshes the mode combo box */
+		void ForceRefresh();
+
+		/** Called when a fixture type was selected */
+		void OnFixtureTypeSelected(UDMXEntity* InSelectedFixtureType);
+
+		/** Called when a mode was selected */
+		void OnModeSelected(TSharedPtr<uint32> InSelectedMode, ESelectInfo::Type SelectInfo);
+
+		/** Called when the 'Add' button was clicked */
+		FReply OnAddFixturePatchButtonClicked();
+
+		/** Returns the text of the active mode */
+		FText GetActiveModeText() const;
+
+		/** Returns if a valid fixture type with a mode is selected */
+		bool HasValidFixtureTypeAndMode() const;
+
+		/** Current fixture type */
+		TWeakObjectPtr<UDMXEntityFixtureType> WeakFixtureType;
+
+		/** The number of fixture patches to add */
+		int32 ActiveModeIndex = 0;
+
+		/** The number of fixture patches to add */
+		uint32 NumFixturePatchesToAdd = 1;
+
+		/** Timer handle for the request refresh combo box timer */
+		FTimerHandle RequestRefreshModeComboBoxTimerHandle;
+
+		/** Widget to select the fixture type */
+		TSharedPtr<SDMXEntityPickerButton<UDMXEntityFixtureType>> FixtureTypeSelector;
+
+		/** Sources for the mode combo box */
+		TArray<TSharedPtr<uint32>> ModeSources;
+
+		/** Combo box to selected the mode */
+		TSharedPtr<SComboBox<TSharedPtr<uint32>>> ModeComboBox;
+
+		/** The DMX editor this widget displays */
+		TWeakPtr<FDMXEditor> WeakDMXEditor;
+	};
+}
diff --git a/Engine/Plugins/VirtualProduction/DMX/DMXEngine/Source/DMXEditor/Private/Widgets/FixturePatch/SDMXFixturePatcher.cpp b/Engine/Plugins/VirtualProduction/DMX/DMXEngine/Source/DMXEditor/Private/Widgets/FixturePatch/SDMXFixturePatcher.cpp
--- a/Engine/Plugins/VirtualProduction/DMX/DMXEngine/Source/DMXEditor/Private/Widgets/FixturePatch/SDMXFixturePatcher.cpp
+++ b/Engine/Plugins/VirtualProduction/DMX/DMXEngine/Source/DMXEditor/Private/Widgets/FixturePatch/SDMXFixturePatcher.cpp
@@ -2,6 +2,7 @@
 
 #include "SDMXFixturePatcher.h"
 
+#include "Algo/RemoveIf.h"
 #include "DMXEditor.h"
 #include "DMXEditorSettings.h"
 #include "DMXEditorTabNames.h"
@@ -13,8 +14,6 @@
 #include "IO/DMXPortManager.h"
 #include "Library/DMXEntityFixturePatch.h"
 #include "Library/DMXLibrary.h"
-
-#include "Algo/RemoveIf.h"
 #include "Styling/AppStyle.h"
 #include "ScopedTransaction.h"
 #include "SlateOptMacros.h"
@@ -138,6 +137,14 @@ void SDMXFixturePatcher::Tick(const FGeometry& AllottedGeometry, const double In
 	{
 		SharedData->SelectUniverse(UniverseToSetNextTick);
 		UniverseToSetNextTick = INDEX_NONE;
+
+		for (const TTuple<int32, TSharedPtr<SDMXPatchedUniverse>>& UniverseToWidgetPair : PatchedUniversesByID)
+		{
+			if (UniverseToWidgetPair.Value.IsValid())
+			{
+				UniverseToWidgetPair.Value->RequestRefresh();
+			}
+		}
 	}
 }
 
@@ -170,7 +177,7 @@ TSharedRef<SWidget> SDMXFixturePatcher::CreateToolbar()
 			.Font(FAppStyle::GetFontStyle("PropertyWindow.NormalFont"))
 			.TextStyle(FAppStyle::Get(), "DetailsView.CategoryTextStyle")
 			.IsEnabled(this, &SDMXFixturePatcher::IsUniverseSelectionEnabled)
-			.Text(LOCTEXT("UniverseSelectorLabel", "Universe"))
+			.Text(LOCTEXT("UniverseSelectorLabel", "Local Universe"))
 		);
 
 		ToolbarBuilder.AddWidget(
@@ -200,11 +207,14 @@ TSharedRef<SWidget> SDMXFixturePatcher::CreateToolbar()
 			.OnCheckStateChanged(this, &SDMXFixturePatcher::OnToggleDisplayAllUniverses)
 			[
 				SNew(STextBlock)
+				.Margin(FMargin(4.f, 0.f))
 				.Font(FAppStyle::GetFontStyle("PropertyWindow.NormalFont"))
 				.Text(LOCTEXT("ShowAllPatchedUniversesLabel", "Show all patched Universes"))
 			]
 		);
 
+		ToolbarBuilder.AddSeparator();
+
 		const UDMXEditorSettings* DMXEditorSettings = GetDefault<UDMXEditorSettings>();
 		const ECheckBoxState InitialDMXMonitorEnabledCheckBoxState = DMXEditorSettings->FixturePatcherSettings.bMonitorEnabled ? ECheckBoxState::Checked : ECheckBoxState::Unchecked;
 
@@ -214,6 +224,7 @@ TSharedRef<SWidget> SDMXFixturePatcher::CreateToolbar()
 			.OnCheckStateChanged(this, &SDMXFixturePatcher::OnToggleDMXMonitorEnabled)
 			[
 				SNew(STextBlock)
+				.Margin(FMargin(4.f, 0.f))
 				.Font(FAppStyle::GetFontStyle("PropertyWindow.NormalFont"))
 				.Text(LOCTEXT("EnableInputMonitorLabel", "Monitor DMX Inputs"))
 				.ToolTipText(LOCTEXT("EnableInputMonitorTooltip", "If checked, monitors DMX Input Ports used in this DMX Library"))
@@ -222,12 +233,23 @@ TSharedRef<SWidget> SDMXFixturePatcher::CreateToolbar()
 	}
 	ToolbarBuilder.EndSection();
 
-	ToolbarBuilder.BeginSection("AdvancedSettings");
-	{
-		ToolbarBuilder.AddWidget(
+	const TSharedRef<SWidget> Toolbar =	SNew(SHorizontalBox)
+		
+		// Toolbar menu
+		+ SHorizontalBox::Slot()
+		.HAlign(HAlign_Left)
+		.AutoWidth()
+		[	
+			ToolbarBuilder.MakeWidget()
+		]
+
+		// Settings button
+		+ SHorizontalBox::Slot()
+		.HAlign(HAlign_Right)
+		[
 			SNew(SComboButton)
 			.ComboButtonStyle(FAppStyle::Get(), "SimpleComboButtonWithIcon") // Use the tool bar item style for this button
-			.OnGetMenuContent(this, &SDMXFixturePatcher::CreateDisplaySettingsMenu)
+			.OnGetMenuContent(this, &SDMXFixturePatcher::CreateSettingsMenu)
 			.HasDownArrow(false)
 			.ButtonContent()
 			[
@@ -235,14 +257,12 @@ TSharedRef<SWidget> SDMXFixturePatcher::CreateToolbar()
 				.ColorAndOpacity(FSlateColor::UseForeground())
 				.Image(FAppStyle::Get().GetBrush("Icons.Settings"))
 			]
-		);
-	}
-	ToolbarBuilder.EndSection();
+		];
 
-	return ToolbarBuilder.MakeWidget();
+	return Toolbar;
 }
 
-TSharedRef<SWidget> SDMXFixturePatcher::CreateDisplaySettingsMenu()
+TSharedRef<SWidget> SDMXFixturePatcher::CreateSettingsMenu()
 {
 	constexpr bool bCloseMenuWindowAfterMenuSelection = true;
 	FMenuBuilder MenuBuilder(bCloseMenuWindowAfterMenuSelection, nullptr);
@@ -331,63 +351,57 @@ FReply SDMXFixturePatcher::OnDropOntoChannel(int32 UniverseID, int32 ChannelID,
 			const int32 AbsoluteOffset = DraggedNodesToChannelOffsetPair.Value + PadOffset;
 
 			int32 NewUniverseID = UniverseID;
-			int32 DesiredStartingChannel = ChannelID;
-			if (IsUniverseSelectionEnabled())
-			{			
-				// Stay within the current Universe if only one is displayed
-				if (FixturePatch->GetUniverseID() != UniverseID)
-				{
-					continue;
-				}
+			int64 AbsoluteDesiredStartingChannel = UniverseID * DMX_MAX_ADDRESS + ChannelID - AbsoluteOffset;
 
-				DesiredStartingChannel = ChannelID - AbsoluteOffset % DMX_MAX_ADDRESS;
+			// Allow drag to another universe if multiple universes are displayed
+			int32 DesiredUniverse = AbsoluteDesiredStartingChannel / DMX_MAX_ADDRESS;
+			if (DesiredUniverse < 1)
+			{
+				DesiredUniverse = 1;
+				AbsoluteDesiredStartingChannel = ChannelID - AbsoluteOffset % DMX_MAX_ADDRESS;
+			}
+			else if (DesiredUniverse > DMX_MAX_UNIVERSE)
+			{
+				DesiredUniverse = DMX_MAX_UNIVERSE;
+				AbsoluteDesiredStartingChannel = DMX_MAX_UNIVERSE * DMX_MAX_ADDRESS - ChannelSpan + 1;
 			}
 			else
 			{
-				int64 AbsoluteDesiredStartingChannel = UniverseID * DMX_MAX_ADDRESS + ChannelID - AbsoluteOffset;
-
-				// Allow drag to another universe if multiple universes are displayed
-				int32 DesiredUniverse = AbsoluteDesiredStartingChannel / DMX_MAX_ADDRESS;
-				if (DesiredUniverse < 1)
-				{
-					DesiredUniverse = 1;
-					AbsoluteDesiredStartingChannel = ChannelID - AbsoluteOffset % DMX_MAX_ADDRESS;
-				}
-				else if (DesiredUniverse > DMX_MAX_UNIVERSE)
-				{
-					DesiredUniverse = DMX_MAX_UNIVERSE;
-					AbsoluteDesiredStartingChannel = DMX_MAX_UNIVERSE * DMX_MAX_ADDRESS - ChannelSpan + 1;
-				}
-				else
-				{
-					NewUniverseID = DesiredUniverse;
-				}
+				NewUniverseID = DesiredUniverse;
+			}
 
-				// Remove from previous universe if needed
-				if (NewUniverseID != PatchedNode->GetUniverseID())
+			// Remove from previous universe if needed
+			if (NewUniverseID != PatchedNode->GetUniverseID())
+			{
+				const TSharedPtr<SDMXPatchedUniverse>* OldUniverseWidgetPtr = PatchedUniversesByID.Find(PatchedNode->GetUniverseID());
+				if (OldUniverseWidgetPtr)
 				{
-					const TSharedPtr<SDMXPatchedUniverse>* OldUniverseWidgetPtr = PatchedUniversesByID.Find(PatchedNode->GetUniverseID());
-					if (OldUniverseWidgetPtr)
-					{
-						(*OldUniverseWidgetPtr)->Remove(PatchedNode);
-						(*OldUniverseWidgetPtr)->RequestRefresh();
-					}
+					(*OldUniverseWidgetPtr)->Remove(PatchedNode);
+					(*OldUniverseWidgetPtr)->RequestRefresh();
 				}
-
-				DesiredStartingChannel = AbsoluteDesiredStartingChannel % DMX_MAX_ADDRESS;
 			}
 
+			const int32 DesiredStartingChannel = AbsoluteDesiredStartingChannel % DMX_MAX_ADDRESS;
 			const int32 NewStartingChannel = ClampStartingChannel(DesiredStartingChannel, ChannelSpan);
 			DraggedNodesToChannelOffsetPair.Key->SetAddresses(NewUniverseID, NewStartingChannel, ChannelSpan);
 
+			// Show newly patched universes if many universes are displayed and the user dragged on top of it 
 			const TSharedPtr<SDMXPatchedUniverse>* UniverseWidgetPtr = PatchedUniversesByID.Find(NewUniverseID);
-			if (!UniverseWidgetPtr)
+			if (!UniverseWidgetPtr && !IsUniverseSelectionEnabled())
+			{
+				if (!UniverseWidgetPtr)
+				{
+					AddUniverse(NewUniverseID);
+					UniverseWidgetPtr = &PatchedUniversesByID.FindChecked(NewUniverseID);
+				}
+				check(UniverseWidgetPtr);
+			}
+		
+			// Add the patch, if the universe is visible
+			if (UniverseWidgetPtr)
 			{
-				AddUniverse(NewUniverseID);
-				UniverseWidgetPtr = &PatchedUniversesByID.FindChecked(NewUniverseID);
+				(*UniverseWidgetPtr)->FindOrAdd(DraggedNodesToChannelOffsetPair.Key);
 			}
-			check(UniverseWidgetPtr);
-			(*UniverseWidgetPtr)->FindOrAdd(DraggedNodesToChannelOffsetPair.Key);
 
 			// Pad further patches so they can't be offset more than the currently dropped patch
 			PadOffset = FMath::Min(PadOffset, DesiredStartingChannel - NewStartingChannel);
diff --git a/Engine/Plugins/VirtualProduction/DMX/DMXEngine/Source/DMXEditor/Private/Widgets/FixturePatch/SDMXFixturePatcher.h b/Engine/Plugins/VirtualProduction/DMX/DMXEngine/Source/DMXEditor/Private/Widgets/FixturePatch/SDMXFixturePatcher.h
--- a/Engine/Plugins/VirtualProduction/DMX/DMXEngine/Source/DMXEditor/Private/Widgets/FixturePatch/SDMXFixturePatcher.h
+++ b/Engine/Plugins/VirtualProduction/DMX/DMXEngine/Source/DMXEditor/Private/Widgets/FixturePatch/SDMXFixturePatcher.h
@@ -61,7 +61,7 @@ private:
 	TSharedRef<SWidget> CreateToolbar();
 
 	/** Creates the settings menu */
-	TSharedRef<SWidget> CreateDisplaySettingsMenu();
+	TSharedRef<SWidget> CreateSettingsMenu();
 
 	/** Selects the fixture patch name display mode */
 	void SelectFixturePatchNameDisplayMode(EDMXFixturePatcherNameDisplayMode DisplayMode);
diff --git a/Engine/Plugins/VirtualProduction/DMX/DMXEngine/Source/DMXEditor/Private/Widgets/FixturePatch/SDMXMVRFixtureList.cpp b/Engine/Plugins/VirtualProduction/DMX/DMXEngine/Source/DMXEditor/Private/Widgets/FixturePatch/SDMXMVRFixtureList.cpp
--- a/Engine/Plugins/VirtualProduction/DMX/DMXEngine/Source/DMXEditor/Private/Widgets/FixturePatch/SDMXMVRFixtureList.cpp
+++ b/Engine/Plugins/VirtualProduction/DMX/DMXEngine/Source/DMXEditor/Private/Widgets/FixturePatch/SDMXMVRFixtureList.cpp
@@ -2,34 +2,34 @@
 
 #include "SDMXMVRFixtureList.h"
 
+#include "Algo/Copy.h"
+#include "Algo/MaxElement.h"
+#include "Algo/MinElement.h"
+#include "Algo/Sort.h"
+#include "Commands/DMXEditorCommands.h"
 #include "DMXEditor.h"
 #include "DMXEditorSettings.h"
 #include "DMXEditorUtils.h"
-#include "DMXFixturePatchAutoAssignUtility.h"
 #include "DMXFixturePatchSharedData.h"
 #include "DMXRuntimeUtils.h"
-#include "Commands/DMXEditorCommands.h"
-#include "Framework/Application/SlateApplication.h"
-#include "Library/DMXEntityFixturePatch.h"
-#include "Library/DMXEntityFixtureType.h"
-#include "Library/DMXLibrary.h"
-#include "Widgets/FixturePatch/DMXMVRFixtureListItem.h"
-#include "Widgets/FixturePatch/SDMXMVRFixtureListRow.h"
-#include "Widgets/FixturePatch/SDMXMVRFixtureListToolbar.h"
-
-#include "Factories.h"
-#include "ScopedTransaction.h"
-#include "TimerManager.h"
-#include "UnrealExporter.h"
-#include "Algo/MinElement.h"
-#include "Algo/Sort.h"
 #include "Exporters/Exporter.h"
+#include "Factories.h"
+#include "FixturePatchAutoAssignUtility.h"
 #include "Framework/Application/SlateApplication.h"
 #include "Framework/Commands/GenericCommands.h"
 #include "Framework/Commands/UICommandList.h"
 #include "Framework/MultiBox/MultiBoxBuilder.h"
 #include "HAL/PlatformApplicationMisc.h"
+#include "Library/DMXEntityFixturePatch.h"
+#include "Library/DMXEntityFixtureType.h"
+#include "Library/DMXLibrary.h"
+#include "ScopedTransaction.h"
 #include "Styling/AppStyle.h"
+#include "TimerManager.h"
+#include "UnrealExporter.h"
+#include "Widgets/FixturePatch/DMXMVRFixtureListItem.h"
+#include "Widgets/FixturePatch/SDMXMVRFixtureListRow.h"
+#include "Widgets/FixturePatch/SDMXMVRFixtureListToolbar.h"
 #include "Widgets/SBoxPanel.h"
 #include "Widgets/Views/SHeaderRow.h"
 #include "Widgets/Views/SListView.h"
@@ -90,41 +90,19 @@ namespace UE::DMX::SDMXMVRFixtureList::Private
 
 		// Find the first free Addresses
 		TArray<UDMXEntityFixturePatch*> FixturePatches = DMXLibrary->GetEntitiesTypeCast<UDMXEntityFixturePatch>();
-		if (FixturePatches.Num() == 0)
+		if (FixturePatches.IsEmpty())
 		{
 			return nullptr;
 		}
-
-		FixturePatches.Sort([](const UDMXEntityFixturePatch& FixturePatchA, const UDMXEntityFixturePatch& FixturePatchB)
-			{
-				const bool bUniverseIsSmaller = FixturePatchA.GetUniverseID() < FixturePatchB.GetUniverseID();
-				const bool bUniverseIsEqual = FixturePatchA.GetUniverseID() == FixturePatchB.GetUniverseID();
-				const bool bAddressIsSmallerOrEqual = FixturePatchA.GetStartingChannel() <= FixturePatchB.GetStartingChannel();
-
-				return bUniverseIsSmaller || (bUniverseIsEqual && bAddressIsSmallerOrEqual);
-			});
-
-		int32 Address = FixturePatches.Last()->GetStartingChannel() + FixturePatches.Last()->GetChannelSpan();
-		int32 Universe = -1;
-		if (Address > DMX_MAX_ADDRESS)
-		{
-			Address = 1;
-			Universe = FixturePatches.Last()->GetUniverseID() + 1;
-		}
-		else
-		{
-			Universe = FixturePatches.Last()->GetUniverseID();
-		}
-
 		UDMXEntityFixturePatch** FixturePatchToDuplicatePtr = Algo::FindByPredicate(FixturePatches, [MVRFixtureUUID](const UDMXEntityFixturePatch* FixturePatch)
 			{
 				return FixturePatch->GetMVRFixtureUUID() == MVRFixtureUUID;
 			});
-
 		if (!ensureMsgf(FixturePatchToDuplicatePtr, TEXT("Trying to duplicate fixture patch, but source fixture patch cannot be found")))
 		{
 			return nullptr;
 		}
+
 		UDMXEntityFixturePatch* FixturePatchToDuplicate = *FixturePatchToDuplicatePtr;
 		check(FixturePatchToDuplicate);
 
@@ -143,18 +121,11 @@ namespace UE::DMX::SDMXMVRFixtureList::Private
 		}
 
 		// Duplicate the Fixture Patch
-		const int32 ChannelSpan = FixturePatchToDuplicate->GetChannelSpan();
-		if (Address + ChannelSpan - 1 > DMX_MAX_ADDRESS)
-		{
-			Address = 1;
-			Universe++;
-		}
-
 		FDMXEntityFixturePatchConstructionParams ConstructionParams;
 		ConstructionParams.FixtureTypeRef = FixturePatchToDuplicate->GetFixtureType();
 		ConstructionParams.ActiveMode = FixturePatchToDuplicate->GetActiveModeIndex();
-		ConstructionParams.UniverseID = Universe;
-		ConstructionParams.StartingAddress = Address;
+		ConstructionParams.UniverseID = FixturePatchToDuplicate->GetUniverseID();
+		ConstructionParams.StartingAddress = FixturePatchToDuplicate->GetStartingChannel();
 
 		constexpr bool bMarkLibraryDirty = false;
 		UDMXEntityFixturePatch* NewFixturePatch = UDMXEntityFixturePatch::CreateFixturePatchInLibrary(ConstructionParams, FixturePatchToDuplicate->Name, bMarkLibraryDirty);
@@ -162,8 +133,6 @@ namespace UE::DMX::SDMXMVRFixtureList::Private
 		// Use the same color as the duplicated patch
 		NewFixturePatch->EditorColor = FixturePatchToDuplicate->EditorColor;
 
-		Address += ChannelSpan;
-
 		DMXLibrary->PostEditChange();
 
 		return NewFixturePatch;
@@ -813,6 +782,8 @@ void SDMXMVRFixtureList::OnSelectionChanged(TSharedPtr<FDMXMVRFixtureListItem> I
 	{
 		return;
 	}
+
+	LastSelectedItem = InItem;
  
 	const TArray<TSharedPtr<FDMXMVRFixtureListItem>> SelectedItems = ListView->GetSelectedItems();
 	TArray<TWeakObjectPtr<UDMXEntityFixturePatch>> FixturePatchesToSelect;
@@ -941,8 +912,13 @@ void SDMXMVRFixtureList::AdoptSelectionFromFixturePatchSharedData()
 	}
 }
 
-void SDMXMVRFixtureList::AutoAssignFixturePatches(EDMXFixturePatchAutoAssignMode AutoAssignMode)
+void SDMXMVRFixtureList::AutoAssignFixturePatches(UE::DMXEditor::AutoAssign::EAutoAssignMode Mode)
 {
+	if (!WeakDMXEditor.IsValid())
+	{
+		return;
+	}
+
 	TArray<UDMXEntityFixturePatch*> FixturePatchesToAutoAssign;
 	const TArray<TWeakObjectPtr<UDMXEntityFixturePatch>> SelectedFixturePatches = FixturePatchSharedData->GetSelectedFixturePatches();
 	for (TWeakObjectPtr<UDMXEntityFixturePatch> FixturePatch : SelectedFixturePatches)
@@ -958,7 +934,8 @@ void SDMXMVRFixtureList::AutoAssignFixturePatches(EDMXFixturePatchAutoAssignMode
 		return;
 	}
 
-	FDMXFixturePatchAutoAssignUtility::AutoAssign(FixturePatchesToAutoAssign, AutoAssignMode);
+	using namespace UE::DMXEditor::AutoAssign;
+	FAutoAssignUtility::AutoAssign(Mode, WeakDMXEditor.Pin().ToSharedRef(), FixturePatchesToAutoAssign, 1, 1);
 	FixturePatchSharedData->SelectUniverse(FixturePatchesToAutoAssign[0]->GetUniverseID());
 
 	RequestListRefresh();
@@ -1131,7 +1108,7 @@ void SDMXMVRFixtureList::SortByColumnID(const EColumnSortPriority::Type SortPrio
 
 	if (ColumnId == FDMXMVRFixtureListCollumnIDs::FixturePatchName)
 	{
-		Algo::Sort(ListSource, [bAscending](const TSharedPtr<FDMXMVRFixtureListItem>& ItemA, const TSharedPtr<FDMXMVRFixtureListItem>& ItemB)
+		Algo::StableSort(ListSource, [bAscending](const TSharedPtr<FDMXMVRFixtureListItem>& ItemA, const TSharedPtr<FDMXMVRFixtureListItem>& ItemB)
 			{
 				const FString FixturePatchNameA = ItemA->GetFixturePatchName();
 				const FString FixturePatchNameB = ItemB->GetFixturePatchName();
@@ -1142,7 +1119,7 @@ void SDMXMVRFixtureList::SortByColumnID(const EColumnSortPriority::Type SortPrio
 	}
 	else if (ColumnId == FDMXMVRFixtureListCollumnIDs::FixtureID)
 	{
-		Algo::Sort(ListSource, [bAscending](const TSharedPtr<FDMXMVRFixtureListItem>& ItemA, const TSharedPtr<FDMXMVRFixtureListItem>& ItemB)
+		Algo::StableSort(ListSource, [bAscending](const TSharedPtr<FDMXMVRFixtureListItem>& ItemA, const TSharedPtr<FDMXMVRFixtureListItem>& ItemB)
 			{
 				bool bIsGreater = [ItemA, ItemB]()
 				{
@@ -1171,7 +1148,7 @@ void SDMXMVRFixtureList::SortByColumnID(const EColumnSortPriority::Type SortPrio
 	}
 	else if (ColumnId == FDMXMVRFixtureListCollumnIDs::FixtureType)
 	{
-		Algo::Sort(ListSource, [bAscending](const TSharedPtr<FDMXMVRFixtureListItem>& ItemA, const TSharedPtr<FDMXMVRFixtureListItem>& ItemB)
+		Algo::StableSort(ListSource, [bAscending](const TSharedPtr<FDMXMVRFixtureListItem>& ItemA, const TSharedPtr<FDMXMVRFixtureListItem>& ItemB)
 			{
 				const FString FixtureTypeA = ItemA->GetFixtureType()->Name;
 				const FString FixtureTypeB = ItemB->GetFixtureType()->Name;
@@ -1182,7 +1159,7 @@ void SDMXMVRFixtureList::SortByColumnID(const EColumnSortPriority::Type SortPrio
 	}
 	else if (ColumnId == FDMXMVRFixtureListCollumnIDs::Mode)
 	{
-		Algo::Sort(ListSource, [bAscending](const TSharedPtr<FDMXMVRFixtureListItem>& ItemA, const TSharedPtr<FDMXMVRFixtureListItem>& ItemB)
+		Algo::StableSort(ListSource, [bAscending](const TSharedPtr<FDMXMVRFixtureListItem>& ItemA, const TSharedPtr<FDMXMVRFixtureListItem>& ItemB)
 			{
 				const bool bIsGreater = ItemA->GetModeIndex() >= ItemB->GetModeIndex();
 				return bAscending ? !bIsGreater : bIsGreater;
@@ -1190,7 +1167,7 @@ void SDMXMVRFixtureList::SortByColumnID(const EColumnSortPriority::Type SortPrio
 	}
 	else if (ColumnId == FDMXMVRFixtureListCollumnIDs::Patch)
 	{
-		Algo::Sort(ListSource, [bAscending](const TSharedPtr<FDMXMVRFixtureListItem>& ItemA, const TSharedPtr<FDMXMVRFixtureListItem>& ItemB)
+		Algo::StableSort(ListSource, [bAscending](const TSharedPtr<FDMXMVRFixtureListItem>& ItemA, const TSharedPtr<FDMXMVRFixtureListItem>& ItemB)
 			{
 				const UDMXEntityFixturePatch* FixturePatchA = ItemA->GetFixturePatch();
 				const UDMXEntityFixturePatch* FixturePatchB = ItemB->GetFixturePatch();
@@ -1227,20 +1204,6 @@ TSharedPtr<SWidget> SDMXMVRFixtureList::OnContextMenuOpening()
 
 	const bool bCloseWindowAfterMenuSelection = true;
 	FMenuBuilder MenuBuilder(bCloseWindowAfterMenuSelection, CommandList);
-
-	// Auto Assign Section
-	if (!FixturePatchSharedData->GetSelectedFixturePatches().IsEmpty())
-	{
-		MenuBuilder.BeginSection("AutoAssignSection", LOCTEXT("AutoAssignSection", "Auto-Assign"));
-		{
-			MenuBuilder.AddMenuEntry(FDMXEditorCommands::Get().AutoAssignToFirstUniverseInSelection);
-			MenuBuilder.AddMenuEntry(FDMXEditorCommands::Get().AutoAssignToFirstReachableUniverse);
-			MenuBuilder.AddMenuEntry(FDMXEditorCommands::Get().AutoAssignAfterLastAddressInLibrary);
-		}
-		MenuBuilder.EndSection();
-	}
-
-	// Basic Operations Section
 	MenuBuilder.BeginSection("BasicOperationsSection", LOCTEXT("BasicOperationsSection", "Basic Operations"));
 	{
 		MenuBuilder.AddMenuEntry(FGenericCommands::Get().Cut);
@@ -1251,6 +1214,13 @@ TSharedPtr<SWidget> SDMXMVRFixtureList::OnContextMenuOpening()
 	}
 	MenuBuilder.EndSection();
 
+	using namespace UE::DMXEditor::AutoAssign;
+	MenuBuilder.BeginSection("AutoAssignSection", LOCTEXT("AutoAssignActionsSection", "Auto-Assign"));
+	{
+		MenuBuilder.AddMenuEntry(FDMXEditorCommands::Get().AutoAssignSelectedUniverse);
+	}
+	MenuBuilder.EndSection();
+
 	return MenuBuilder.MakeWidget();
 }
 
@@ -1293,23 +1263,12 @@ void SDMXMVRFixtureList::RegisterCommands()
 			EUIActionRepeatMode::RepeatEnabled
 		)
 	);
-	CommandList->MapAction
-	(
-		FDMXEditorCommands::Get().AutoAssignToFirstUniverseInSelection,
-		FExecuteAction::CreateSP(this, &SDMXMVRFixtureList::AutoAssignFixturePatches, EDMXFixturePatchAutoAssignMode::FirstUniverseInSelection)
-	);
-
-	CommandList->MapAction
-	(
-		FDMXEditorCommands::Get().AutoAssignToFirstReachableUniverse,
-		FExecuteAction::CreateSP(this, &SDMXMVRFixtureList::AutoAssignFixturePatches, EDMXFixturePatchAutoAssignMode::FirstReachableUniverse),
-		FCanExecuteAction::CreateSP(this, &SDMXMVRFixtureList::DoesDMXLibraryHaveReachableUniverses)
-	);
 
+	using namespace UE::DMXEditor::AutoAssign;
 	CommandList->MapAction
 	(
-		FDMXEditorCommands::Get().AutoAssignAfterLastAddressInLibrary,
-		FExecuteAction::CreateSP(this, &SDMXMVRFixtureList::AutoAssignFixturePatches, EDMXFixturePatchAutoAssignMode::LastAddressInLibrary)
+		FDMXEditorCommands::Get().AutoAssignSelectedUniverse,
+		FExecuteAction::CreateSP(this, &SDMXMVRFixtureList::AutoAssignFixturePatches, EAutoAssignMode::SelectedUniverse)
 	);
 }
 
@@ -1402,31 +1361,40 @@ bool SDMXMVRFixtureList::CanDuplicateItems() const
 
 void SDMXMVRFixtureList::OnDuplicateItems()
 {
-	UDMXLibrary* DMXLibrary = WeakDMXEditor.IsValid() ? WeakDMXEditor.Pin()->GetDMXLibrary() : nullptr;
-	if (!DMXLibrary)
+	const TSharedPtr<FDMXEditor> DMXEditor = WeakDMXEditor.Pin();
+	UDMXLibrary* DMXLibrary = DMXEditor.IsValid() ? DMXEditor->GetDMXLibrary() : nullptr;
+	if (!DMXEditor.IsValid() || !DMXLibrary)
 	{
 		return;
 	}
 
 	TGuardValue<bool>(bChangingDMXLibrary, true);
+
 	const FText TransactionText = LOCTEXT("DuplicateFixturePatchesTransaction", "Duplicate Fixture Patches");
 	const FScopedTransaction PasteTransaction(TransactionText);
 
-	TArray<TWeakObjectPtr<UDMXEntityFixturePatch>> SelectedFixturePatches = FixturePatchSharedData->GetSelectedFixturePatches();
-	
-	// Sort in order of the list so duplicate happens in order of the list
-	Algo::SortBy(SelectedFixturePatches, [this](TWeakObjectPtr<UDMXEntityFixturePatch> FixturePatch)
+	TArray<TWeakObjectPtr<UDMXEntityFixturePatch>> SelectedWeakFixturePatches = FixturePatchSharedData->GetSelectedFixturePatches();
+	SelectedWeakFixturePatches.Remove(nullptr);
+	if (SelectedWeakFixturePatches.IsEmpty())
+	{
+		return;
+	}
+
+	TArray<UDMXEntityFixturePatch*> SelectedFixturePatches;
+	Algo::TransformIf(SelectedWeakFixturePatches, SelectedFixturePatches, 
+		[](const TWeakObjectPtr<UDMXEntityFixturePatch>& WeakFixturePatch) { return WeakFixturePatch.IsValid(); },
+		[](const TWeakObjectPtr<UDMXEntityFixturePatch>& WeakFixturePatch) { return WeakFixturePatch.Get(); });
+
+	// Duplicate in order of patch
+	Algo::StableSortBy(SelectedFixturePatches, [this](UDMXEntityFixturePatch* FixturePatch)
 		{
-			return ListSource.IndexOfByPredicate([FixturePatch](const TSharedPtr<FDMXMVRFixtureListItem>& Item)
-				{
-					return Item->GetFixturePatch() == FixturePatch;
-				});
+			return FixturePatch->GetUniverseID() * DMX_UNIVERSE_SIZE + FixturePatch->GetStartingChannel();
 		});
 
-	TArray<TWeakObjectPtr<UDMXEntityFixturePatch>> NewFixturePatches;
-	for (const TWeakObjectPtr<UDMXEntityFixturePatch> FixturePatch : SelectedFixturePatches)
+	TArray<UDMXEntityFixturePatch*> NewFixturePatches;
+	for (const UDMXEntityFixturePatch* FixturePatch : SelectedFixturePatches)
 	{
-		if(!FixturePatch.IsValid() || !FixturePatch->GetParentLibrary())
+		if(!FixturePatch->GetParentLibrary())
 		{
 			continue;
 		}
@@ -1435,7 +1403,16 @@ void SDMXMVRFixtureList::OnDuplicateItems()
 		UDMXEntityFixturePatch* NewFixturePatch = DuplicatePatchByMVRFixtureUUID(DMXLibrary, FixturePatch->GetMVRFixtureUUID());
 		NewFixturePatches.Add(NewFixturePatch);
 	}
-	FixturePatchSharedData->SelectFixturePatches(NewFixturePatches);
+
+	using namespace UE::DMXEditor::AutoAssign;
+	const int32 AssignedToUniverse = FAutoAssignUtility::AutoAssign(EAutoAssignMode::AfterLastPatchedUniverse, DMXEditor.ToSharedRef(), NewFixturePatches);
+
+	DMXLibrary->PostEditChange();
+
+	TArray<TWeakObjectPtr<UDMXEntityFixturePatch>> NewWeakFixturePatches;
+	Algo::Copy(NewFixturePatches, NewWeakFixturePatches);
+	FixturePatchSharedData->SelectFixturePatches(NewWeakFixturePatches);
+	FixturePatchSharedData->SelectUniverse(AssignedToUniverse);
 
 	RequestListRefresh();
 }
diff --git a/Engine/Plugins/VirtualProduction/DMX/DMXEngine/Source/DMXEditor/Private/Widgets/FixturePatch/SDMXMVRFixtureList.h b/Engine/Plugins/VirtualProduction/DMX/DMXEngine/Source/DMXEditor/Private/Widgets/FixturePatch/SDMXMVRFixtureList.h
--- a/Engine/Plugins/VirtualProduction/DMX/DMXEngine/Source/DMXEditor/Private/Widgets/FixturePatch/SDMXMVRFixtureList.h
+++ b/Engine/Plugins/VirtualProduction/DMX/DMXEngine/Source/DMXEditor/Private/Widgets/FixturePatch/SDMXMVRFixtureList.h
@@ -8,26 +8,24 @@
 #include "Widgets/SCompoundWidget.h"
 #include "Widgets/Views/SHeaderRow.h"
 
-enum class EDMXFixturePatchAutoAssignMode : uint8;
 class FDMXEditor;
 class FDMXMVRFixtureListItem;
 class FDMXFixturePatchSharedData;
+class FUICommandList;
+class ITableRow;
+class SBorder;
 class SDMXMVRFixtureListRow;
 class SDMXMVRFixtureListToolbar;
+class SHeaderRow;
+template <typename ItemType> class SListView;
+class STableViewBase;
 class UDMXEntity;
 class UDMXEntityFixturePatch;
 class UDMXEntityFixtureType;
 class UDMXLibrary;
 class UDMXMVRFixtureNode;
 class UDMXMVRGeneralSceneDescription;
-
-class FUICommandList;
-class ITableRow;
-class SBorder;
-class SHeaderRow;
-template <typename ItemType> class SListView;
-class STableViewBase;
-
+namespace UE::DMXEditor::AutoAssign { enum class EAutoAssignMode : uint8; }
 
 /** Collumn IDs in the Fixture Patch List */
 struct FDMXMVRFixtureListCollumnIDs
@@ -115,7 +113,7 @@ private:
 	void AdoptSelectionFromFixturePatchSharedData();
 
 	/** Auto assigns selected Fixture Patches */
-	void AutoAssignFixturePatches(EDMXFixturePatchAutoAssignMode AutoAssignMode);
+	void AutoAssignFixturePatches(UE::DMXEditor::AutoAssign::EAutoAssignMode Mode);
 
 	/** Returns true if the DMX Library has reachable Universes */
 	bool DoesDMXLibraryHaveReachableUniverses() const;
@@ -156,6 +154,9 @@ private:
 	/** List source when filtered by search */
 	TArray<TSharedPtr<FDMXMVRFixtureListItem>> FilteredListSource;
 
+	/** The item that was last selected in the list */
+	TSharedPtr<FDMXMVRFixtureListItem> LastSelectedItem;
+
 	/** The Search Bar for the List */
 	TSharedPtr<SDMXMVRFixtureListToolbar> Toolbar;
 
diff --git a/Engine/Plugins/VirtualProduction/DMX/DMXEngine/Source/DMXEditor/Private/Widgets/FixturePatch/SDMXMVRFixtureListRow.cpp b/Engine/Plugins/VirtualProduction/DMX/DMXEngine/Source/DMXEditor/Private/Widgets/FixturePatch/SDMXMVRFixtureListRow.cpp
--- a/Engine/Plugins/VirtualProduction/DMX/DMXEngine/Source/DMXEditor/Private/Widgets/FixturePatch/SDMXMVRFixtureListRow.cpp
+++ b/Engine/Plugins/VirtualProduction/DMX/DMXEngine/Source/DMXEditor/Private/Widgets/FixturePatch/SDMXMVRFixtureListRow.cpp
@@ -698,10 +698,17 @@ TSharedRef<SWidget> SDMXMVRFixtureListRow::GeneratePatchWidget()
 			[
 				SAssignNew(PatchTextBlock, SInlineEditableTextBlock)
 				.Text_Lambda([this]()
-					{
+					{						
 						const int32 UniverseID = Item->GetUniverse();
 						const int32 StartingAddress = Item->GetAddress();
-						return FText::Format(LOCTEXT("AddressesText", "{0}.{1}"), UniverseID, StartingAddress);
+						if (UniverseID > 0 && StartingAddress > 0)
+						{
+							return FText::Format(LOCTEXT("AddressesText", "{0}.{1}"), UniverseID, StartingAddress);
+						}
+						else
+						{
+							return LOCTEXT("AddressesText", "Not patched");
+						}
 					})
 				.Font(FAppStyle::GetFontStyle(TEXT("PropertyWindow.NormalFont")))
 				.OnTextCommitted(this, &SDMXMVRFixtureListRow::OnPatchCommitted)
diff --git a/Engine/Plugins/VirtualProduction/DMX/DMXEngine/Source/DMXEditor/Private/Widgets/FixturePatch/SDMXMVRFixtureListToolbar.cpp b/Engine/Plugins/VirtualProduction/DMX/DMXEngine/Source/DMXEditor/Private/Widgets/FixturePatch/SDMXMVRFixtureListToolbar.cpp
--- a/Engine/Plugins/VirtualProduction/DMX/DMXEngine/Source/DMXEditor/Private/Widgets/FixturePatch/SDMXMVRFixtureListToolbar.cpp
+++ b/Engine/Plugins/VirtualProduction/DMX/DMXEngine/Source/DMXEditor/Private/Widgets/FixturePatch/SDMXMVRFixtureListToolbar.cpp
@@ -2,20 +2,20 @@
 
 #include "SDMXMVRFixtureListToolbar.h"
 
+#include "Algo/MaxElement.h"
+#include "Commands/DMXEditorCommands.h"
 #include "DMXEditor.h"
 #include "DMXEditorUtils.h"
 #include "DMXFixturePatchSharedData.h"
 #include "DMXMVRFixtureListItem.h"
-#include "DMXFixturePatchAutoAssignUtility.h"
+#include "EditorStyleSet.h"
+#include "FixturePatchAutoAssignUtility.h"
 #include "Library/DMXEntityFixturePatch.h"
 #include "Library/DMXEntityFixtureType.h"
-#include "Commands/DMXEditorCommands.h"
-#include "Widgets/SDMXEntityDropdownMenu.h"
-
-#include "EditorStyleSet.h"
 #include "ScopedTransaction.h"
-#include "Algo/MaxElement.h"
+#include "SAddFixturePatchMenu.h"
 #include "Widgets/SBoxPanel.h"
+#include "Widgets/SDMXEntityDropdownMenu.h"
 #include "Widgets/Input/SCheckBox.h"
 #include "Widgets/Input/SEditableTextBox.h"
 #include "Widgets/Input/SSearchBox.h"
@@ -198,167 +198,30 @@ TSharedRef<SWidget> SDMXMVRFixtureListToolbar::GenerateFixtureTypeDropdownMenu()
 			]
 			.MenuContent()
 			[
-
-				SNew(SVerticalBox)
-				
-				// Bulk Add Fixture Patches 
-				+ SVerticalBox::Slot()
-				.AutoHeight()
-				[
-					SNew(SHorizontalBox)
-					
-					+SHorizontalBox::Slot()
-					.HAlign(HAlign_Left)
-					.VAlign(VAlign_Center)
-					.Padding(8.f, 2.f, 4.f, 2.f)
-					[
-						SNew(STextBlock)
-						.Text(LOCTEXT("BulkAddPatchesLabel", "Quantity"))
-					]
-
-					+ SHorizontalBox::Slot()
-					.HAlign(HAlign_Fill)
-					.VAlign(VAlign_Center)
-					.Padding(4.f, 2.f, 4.f, 2.f)
-					[
-						SNew(SEditableTextBox)
-						.SelectAllTextWhenFocused(true)
-						.ClearKeyboardFocusOnCommit(false)
-						.SelectAllTextOnCommit(true)
-						.Text_Lambda([this]()
-							{
-								return FText::FromString(FString::FromInt(NumFixturePatchesToAdd));
-							})
-						.OnVerifyTextChanged_Lambda([](const FText& InNewText, FText& OutErrorMessage)
-							{
-								int32 Value;
-								if (!LexTryParseString<int32>(Value, *InNewText.ToString()) ||
-									Value < 1)
-								{
-									OutErrorMessage = LOCTEXT("BulkAddPatchesBadString", "Needs a numeric value > 0");
-									return false;
-								}
-								return true;
-							})
-						.OnTextCommitted_Lambda([this](const FText& Text, ETextCommit::Type CommitType)
-							{
-								int32 Value;
-								if (LexTryParseString<int32>(Value, *Text.ToString()) &&
-									Value > 0)
-								{
-									constexpr int32 MaxNumPatchesToBulkAdd = 512;
-									NumFixturePatchesToAdd = FMath::Min(Value, MaxNumPatchesToBulkAdd);
-								}
-							})
-					]
-				]
-
-				// Fixture Type Selection
-				+ SVerticalBox::Slot()
+				SNew(SBox)
+				.Padding(4.f)
+				.MinDesiredWidth(460.f)
 				[
-					SAssignNew(FixtureTypeDropdownMenu, SDMXEntityDropdownMenu<UDMXEntityFixtureType>)
-					.DMXEditor(WeakDMXEditor)
-					.OnEntitySelected(this, &SDMXMVRFixtureListToolbar::OnAddNewMVRFixtureClicked)
+					SAssignNew(AddFixturePatchMenu, UE::DMXEditor::FixturePatchEditor::SAddFixturePatchMenu, WeakDMXEditor)
 				]
 			]
 			.IsFocusable(true)
-			.ContentPadding(FMargin(5.0f, 1.0f))
+			.ContentPadding(FMargin(4.0f, 4.0f))
 			.ComboButtonStyle(FAppStyle::Get(), "ToolbarComboButton")
 			.ButtonStyle(FAppStyle::Get(), "FlatButton.Success")
 			.ForegroundColor(FLinearColor::White)
 			.ToolTipText(AddButtonToolTip)
-			.OnComboBoxOpened(FOnComboBoxOpened::CreateLambda([this]() 
-				{ 
-					FixtureTypeDropdownMenu->RefreshEntitiesList();
-				}));
-
-	FixtureTypeDropdownMenu->SetComboButton(AddComboButton);
+			.OnComboBoxOpened_Lambda([this]()
+				{
+					if (AddFixturePatchMenu.IsValid())
+					{
+						AddFixturePatchMenu->RequestRefresh();
+					}
+				});
 
 	return AddComboButton;
 }
 
-void SDMXMVRFixtureListToolbar::OnAddNewMVRFixtureClicked(UDMXEntity* InSelectedFixtureType)
-{
-	if (!ensureMsgf(InSelectedFixtureType, TEXT("Trying to add Fixture Patches, but selected fixture type is invalid.")))
-	{
-		return;
-	}
-
-	TSharedPtr<FDMXEditor> DMXEditor = WeakDMXEditor.Pin();
-	if (!DMXEditor.IsValid())
-	{
-		return;
-	}
-	UDMXLibrary* DMXLibrary = DMXEditor->GetDMXLibrary();
-	if (!DMXLibrary)
-	{
-		return;
-	}
-	UDMXEntityFixtureType* FixtureType = CastChecked<UDMXEntityFixtureType>(InSelectedFixtureType);
-
-	// Find a Universe and Address
-	TArray<UDMXEntityFixturePatch*> FixturePatches = DMXLibrary->GetEntitiesTypeCast<UDMXEntityFixturePatch>();
-	UDMXEntityFixturePatch* const* LastFixturePatchPtr = Algo::MaxElementBy(FixturePatches, [](const UDMXEntityFixturePatch* FixturePatch)
-		{
-			return (int64)FixturePatch->GetUniverseID() * DMX_MAX_ADDRESS + FixturePatch->GetStartingChannel();
-		});
-	int32 Universe = 1;
-	int32 Address = 1;
-	int32 ChannelSpan = 0;
-	if (LastFixturePatchPtr && !FixtureType->Modes.IsEmpty())
-	{
-		ChannelSpan = FixtureType->Modes[0].ChannelSpan;
-		const UDMXEntityFixturePatch& LastFixturePatch = **LastFixturePatchPtr;
-		if (LastFixturePatch.GetEndingChannel() + FixtureType->Modes[0].ChannelSpan > DMX_MAX_ADDRESS)
-		{
-			Universe = LastFixturePatch.GetUniverseID() + 1;
-			Address = 1;
-		}
-		else
-		{
-			Universe = LastFixturePatch.GetUniverseID();
-			Address = LastFixturePatch.GetEndingChannel() + 1;
-		}
-	}
-
-	// Create a new fixture patches
-	const FScopedTransaction CreateFixturePatchTransaction(LOCTEXT("CreateFixturePatchTransaction", "Create Fixture Patch"));
-	TArray<UDMXEntityFixturePatch*> NewFixturePatches;
-	TArray<TWeakObjectPtr<UDMXEntityFixturePatch>> NewWeakFixturePatches;
-	DMXLibrary->PreEditChange(UDMXLibrary::StaticClass()->FindPropertyByName(UDMXLibrary::GetEntitiesPropertyName()));
-	for (int32 iNumFixturePatchesAdded = 0; iNumFixturePatchesAdded < NumFixturePatchesToAdd; iNumFixturePatchesAdded++)
-	{
-		FDMXEntityFixturePatchConstructionParams FixturePatchConstructionParams;
-		FixturePatchConstructionParams.FixtureTypeRef = FDMXEntityFixtureTypeRef(FixtureType);
-		FixturePatchConstructionParams.ActiveMode = 0;
-		FixturePatchConstructionParams.UniverseID = Universe;
-		FixturePatchConstructionParams.StartingAddress = Address;
-
-		constexpr bool bMarkLibraryDirty = false;
-		UDMXEntityFixturePatch* NewFixturePatch = UDMXEntityFixturePatch::CreateFixturePatchInLibrary(FixturePatchConstructionParams, FixtureType->Name, bMarkLibraryDirty);
-		NewFixturePatches.Add(NewFixturePatch);
-		NewWeakFixturePatches.Add(NewFixturePatch);	
-
-		// Increment Universe and Address in steps by one. This is enough for auto assign while keeping order of the named patches
-		if (Address + ChannelSpan > DMX_MAX_ADDRESS)
-		{
-			Universe++;
-			Address = 1; 
-		}
-		else
-		{
-			Address++;
-		}
-	}
-
-	// Auto assign
-	FDMXFixturePatchAutoAssignUtility::AutoAssign(NewFixturePatches, EDMXFixturePatchAutoAssignMode::LastAddressInLibrary);
-
-	DMXLibrary->PostEditChange();
-
-	DMXEditor->GetFixturePatchSharedData()->SelectFixturePatches(NewWeakFixturePatches);
-}
-
 void SDMXMVRFixtureListToolbar::OnSearchTextChanged(const FText& SearchText)
 {
 	SearchString = SearchText.ToString();
diff --git a/Engine/Plugins/VirtualProduction/DMX/DMXEngine/Source/DMXEditor/Private/Widgets/FixturePatch/SDMXMVRFixtureListToolbar.h b/Engine/Plugins/VirtualProduction/DMX/DMXEngine/Source/DMXEditor/Private/Widgets/FixturePatch/SDMXMVRFixtureListToolbar.h
--- a/Engine/Plugins/VirtualProduction/DMX/DMXEngine/Source/DMXEditor/Private/Widgets/FixturePatch/SDMXMVRFixtureListToolbar.h
+++ b/Engine/Plugins/VirtualProduction/DMX/DMXEngine/Source/DMXEditor/Private/Widgets/FixturePatch/SDMXMVRFixtureListToolbar.h
@@ -5,13 +5,13 @@
 #include "CoreMinimal.h"
 #include "Widgets/SCompoundWidget.h"
 
+enum class ECheckBoxState : uint8;
 class FDMXEditor;
 class FDMXMVRFixtureListItem;
 template<typename TEntityType> class SDMXEntityDropdownMenu;
 class UDMXEntity;
 class UDMXEntityFixtureType;
-
-enum class ECheckBoxState : uint8;
+namespace UE::DMXEditor::FixturePatchEditor { class SAddFixturePatchMenu; }
 
 
 /** Search bar for the MVR Fixture List */
@@ -37,9 +37,6 @@ private:
 	/** Generates the 'Add MVR Fixture' Dropdown Menu */
 	TSharedRef<SWidget> GenerateFixtureTypeDropdownMenu();
 
-	/** Called when a new Fixture Patch was selected in the FixtureTypeDropdownMenu */
-	void OnAddNewMVRFixtureClicked(UDMXEntity* InSelectedFixtureType);
-
 	/** Called when the Search Text changed */
 	void OnSearchTextChanged(const FText& SearchText);
 
@@ -52,11 +49,8 @@ private:
 	/** True if only conflicts are searched for */
 	bool bShowConfictsOnly = false;
 
-	/** The number of fixture patches to add when one is selected in the entity dropdown */
-	int32 NumFixturePatchesToAdd = 1;
-
-	/** Dropdown menu of Fixture Types */
-	TSharedPtr<SDMXEntityDropdownMenu<UDMXEntityFixtureType>> FixtureTypeDropdownMenu;
+	/** Menu displayed to add fixture patches */
+	TSharedPtr<UE::DMXEditor::FixturePatchEditor::SAddFixturePatchMenu> AddFixturePatchMenu;
 
 	/** The DMX Editor that owns this widget */
 	TWeakPtr<FDMXEditor> WeakDMXEditor;
diff --git a/Engine/Plugins/VirtualProduction/DMX/DMXEngine/Source/DMXEditor/Private/Widgets/FixturePatch/SDMXPatchedUniverse.cpp b/Engine/Plugins/VirtualProduction/DMX/DMXEngine/Source/DMXEditor/Private/Widgets/FixturePatch/SDMXPatchedUniverse.cpp
--- a/Engine/Plugins/VirtualProduction/DMX/DMXEngine/Source/DMXEditor/Private/Widgets/FixturePatch/SDMXPatchedUniverse.cpp
+++ b/Engine/Plugins/VirtualProduction/DMX/DMXEngine/Source/DMXEditor/Private/Widgets/FixturePatch/SDMXPatchedUniverse.cpp
@@ -2,34 +2,65 @@
 
 #include "SDMXPatchedUniverse.h"
 
+#include "Algo/ForEach.h"
+#include "Algo/RemoveIf.h"
+#include "Commands/DMXEditorCommands.h"
 #include "DMXEditor.h"
-#include "DMXFixturePatchAutoAssignUtility.h"
 #include "DMXFixturePatchEditorDefinitions.h"
 #include "DMXFixturePatchNode.h"
 #include "DMXFixturePatchSharedData.h"
-#include "SDMXChannelConnector.h"
-#include "SDMXFixturePatchFragment.h"
-#include "Commands/DMXEditorCommands.h"
 #include "DragDrop/DMXEntityFixturePatchDragDropOp.h"
+#include "FixturePatchAutoAssignUtility.h"
+#include "Framework/Commands/GenericCommands.h"
+#include "Framework/Commands/UICommandList.h"
+#include "Framework/MultiBox/MultiBoxBuilder.h"
 #include "Framework/Application/SlateApplication.h"
 #include "IO/DMXInputPort.h"
 #include "IO/DMXOutputPort.h"
 #include "IO/DMXPortManager.h"
 #include "Library/DMXEntityFixturePatch.h"
 #include "Library/DMXLibrary.h"
-
-#include "Algo/RemoveIf.h"
-#include "Framework/Commands/GenericCommands.h"
-#include "Framework/Commands/UICommandList.h"
-#include "Framework/MultiBox/MultiBoxBuilder.h"
-#include "Styling/AppStyle.h"
+#include "ScopedTransaction.h"
+#include "SDMXChannelConnector.h"
+#include "SDMXFixturePatchFragment.h"
 #include "SlateOptMacros.h"
+#include "Styling/AppStyle.h"
+#include "Widgets/Input/SComboBox.h"
 #include "Widgets/Layout/SBorder.h"
 #include "Widgets/Layout/SGridPanel.h"
+#include "Widgets/SCompoundWidget.h"
 
 
 #define LOCTEXT_NAMESPACE "SDMXPatchedUniverse"
 
+namespace UE::DMXEditor::FixturePatchEditor::Private
+{
+	/** Helper to raise pre and post edit change events on an array of fixture patches */
+	class FScopedEditChangeFixturePatches
+	{
+	public:
+		FScopedEditChangeFixturePatches(TArray<UDMXEntityFixturePatch*> InFixturePatches)
+			: FixturePatches(InFixturePatches)
+		{
+			for (UDMXEntityFixturePatch* Patch : FixturePatches)
+			{
+				Patch->PreEditChange(nullptr);
+			}
+		}
+
+		~FScopedEditChangeFixturePatches()
+		{
+			for (UDMXEntityFixturePatch* Patch : FixturePatches)
+			{
+				Patch->PostEditChange();
+			}
+		}
+
+	private:
+		TArray<UDMXEntityFixturePatch*> FixturePatches;
+	};
+}
+
 BEGIN_SLATE_FUNCTION_BUILD_OPTIMIZATION
 void SDMXPatchedUniverse::Construct(const FArguments& InArgs)
 {
@@ -39,12 +70,12 @@ void SDMXPatchedUniverse::Construct(const FArguments& InArgs)
 
 	UniverseID = InArgs._UniverseID;
 
-	DMXEditorPtr = InArgs._DMXEditor;
+	WeakDMXEditor = InArgs._DMXEditor;
 	OnDragEnterChannel = InArgs._OnDragEnterChannel;
 	OnDragLeaveChannel = InArgs._OnDragLeaveChannel;
 	OnDropOntoChannel = InArgs._OnDropOntoChannel;
 
-	const TSharedPtr<FDMXEditor> DMXEditor = DMXEditorPtr.Pin();
+	const TSharedPtr<FDMXEditor> DMXEditor = WeakDMXEditor.Pin();
 	if (!DMXEditor.IsValid())
 	{
 		return;
@@ -148,7 +179,7 @@ void SDMXPatchedUniverse::CreateChannelConnectors()
 		int32 Channel = ChannelIndex + FDMXChannelGridSpecs::ChannelIDOffset;
 
 		TSharedRef<SDMXChannelConnector> ChannelPatchWidget =
-			SNew(SDMXChannelConnector, DMXEditorPtr.Pin())
+			SNew(SDMXChannelConnector, WeakDMXEditor.Pin())
 			.ChannelID(Channel)
 			.Value(0)
 			.OnHovered(this, &SDMXPatchedUniverse::HandleOnChannelHovered, Channel)
@@ -394,7 +425,7 @@ void SDMXPatchedUniverse::SetUniverseIDInternal(int32 NewUniverseID)
 			TSharedPtr<FDMXFixturePatchNode> Node = FindPatchNode(FixturePatch);
 			if (!Node.IsValid())
 			{
-				Node = FDMXFixturePatchNode::Create(DMXEditorPtr, FixturePatch);
+				Node = FDMXFixturePatchNode::Create(WeakDMXEditor, FixturePatch);
 			}
 
 			FindOrAdd(Node.ToSharedRef());
@@ -514,12 +545,17 @@ void SDMXPatchedUniverse::ResetMonitor()
 
 void SDMXPatchedUniverse::HandleOnChannelHovered(int32 Channel)
 {
-	HoveredChannel = ChannelConnectors[Channel];
+	if (ChannelConnectors.IsValidIndex(Channel - 1))
+	{
+		HoveredChannel = ChannelConnectors[Channel - 1];
+		LastHoveredChannel = Channel;
+	}
 }
 
 void SDMXPatchedUniverse::HandleOnChannelUnhovered(int32 Channel)
 {
-	if (HoveredChannel == ChannelConnectors[Channel])
+	if (ChannelConnectors.IsValidIndex(Channel - 1) && 
+		HoveredChannel == ChannelConnectors[Channel - 1])
 	{
 		HoveredChannel.Reset();
 	}
@@ -535,7 +571,7 @@ FReply SDMXPatchedUniverse::HandleOnMouseButtonDownOnChannel(uint32 Channel, con
 			[](TWeakObjectPtr<UDMXEntityFixturePatch> WeakFixturePatch) { return WeakFixturePatch.IsValid(); },
 			[](TWeakObjectPtr<UDMXEntityFixturePatch> WeakFixturePatch) { return WeakFixturePatch.Get(); }
 		);
-		Algo::SortBy(Selection, [](UDMXEntityFixturePatch* FixturePatch)
+		Algo::StableSortBy(Selection, [](UDMXEntityFixturePatch* FixturePatch)
 			{
 				return FixturePatch->GetUniverseID() * DMX_MAX_ADDRESS + FixturePatch->GetStartingChannel();
 			});
@@ -616,6 +652,9 @@ FReply SDMXPatchedUniverse::HandleOnMouseButtonDownOnChannel(uint32 Channel, con
 		}
 	}
 
+	// Select this universe
+	SharedData->SelectUniverse(UniverseID);
+
 	return FReply::Unhandled();
 }
 
@@ -670,7 +709,7 @@ FReply SDMXPatchedUniverse::HandleOnDragDetectedOnChannel(uint32 Channel, const
 			return !FixturePatch.IsValid();
 		});
 	SelectedFixturePatches.SetNum(NumValidPatches);
-	Algo::SortBy(SelectedFixturePatches, [](TWeakObjectPtr<UDMXEntityFixturePatch> FixturePatch)
+	Algo::StableSortBy(SelectedFixturePatches, [](TWeakObjectPtr<UDMXEntityFixturePatch> FixturePatch)
 		{
 			return (int64)FixturePatch->GetUniverseID() * DMX_MAX_ADDRESS + FixturePatch->GetStartingChannel();
 		}
@@ -853,29 +892,107 @@ void SDMXPatchedUniverse::OnSelectionChanged()
 	}
 }
 
-void SDMXPatchedUniverse::AutoAssignFixturePatches(EDMXFixturePatchAutoAssignMode AutoAssignMode)
+void SDMXPatchedUniverse::AutoAssignFixturePatches(UE::DMXEditor::AutoAssign::EAutoAssignMode AutoAssignMode)
 {
+	const TSharedPtr<FDMXEditor> DMXEditor = WeakDMXEditor.Pin();
 	TArray<UDMXEntityFixturePatch*> FixturePatchesToAutoAssign;
-	const TArray<TWeakObjectPtr<UDMXEntityFixturePatch>> SelectedFixturePatches = SharedData->GetSelectedFixturePatches();
-	for (TWeakObjectPtr<UDMXEntityFixturePatch> FixturePatch : SelectedFixturePatches)
+	if (!DMXEditor.IsValid() || !GetSelectedFixturePatches(FixturePatchesToAutoAssign))
 	{
-		if (FixturePatch.IsValid())
-		{
-			FixturePatchesToAutoAssign.Add(FixturePatch.Get());
-		}
+		return;
+	}
+
+	const FText TransactionText = FText::Format(LOCTEXT("AutoAssignTransaction", "Auto-Assign Fixture {0}|plural(one=Patch, other=Patches)"), FixturePatchesToAutoAssign.Num());
+	const FScopedTransaction AutoAssignTransaction(TransactionText);
+
+	UE::DMXEditor::FixturePatchEditor::Private::FScopedEditChangeFixturePatches ScopedEditChange(FixturePatchesToAutoAssign);
+	const int32 AssignedToUniverse = UE::DMXEditor::AutoAssign::FAutoAssignUtility::AutoAssign(AutoAssignMode, DMXEditor.ToSharedRef(), FixturePatchesToAutoAssign, UniverseID, LastHoveredChannel);
+
+	SharedData->SelectUniverse(AssignedToUniverse);
+
+	RequestRefresh();
+}
+
+void SDMXPatchedUniverse::AssignFixturePatches()
+{
+	TArray<UDMXEntityFixturePatch*> FixturePatchesToAssign;
+	if (!GetSelectedFixturePatches(FixturePatchesToAssign))
+	{
+		return;
 	}
 
-	if (FixturePatchesToAutoAssign.IsEmpty())
+	const FText TransactionText = FText::Format(LOCTEXT("AssignTransaction", "Assign Fixture {0}|plural(one=Patch, other=Patches)"), FixturePatchesToAssign.Num());
+	const FScopedTransaction AssignTransaction(TransactionText);
+
+	UE::DMXEditor::FixturePatchEditor::Private::FScopedEditChangeFixturePatches ScopedEditChange(FixturePatchesToAssign);
+	const int32 FirstPatchedUniverse = UE::DMXEditor::AutoAssign::FAutoAssignUtility::Assign(FixturePatchesToAssign, UniverseID, LastHoveredChannel);
+
+	SharedData->SelectUniverse(FirstPatchedUniverse);
+	
+	RequestRefresh();
+}
+
+void SDMXPatchedUniverse::AlignFixturePatches()
+{
+	TArray<UDMXEntityFixturePatch*> FixturePatchesToAlign;
+	if (!GetSelectedFixturePatches(FixturePatchesToAlign))
 	{
 		return;
 	}
 
-	FDMXFixturePatchAutoAssignUtility::AutoAssign(FixturePatchesToAutoAssign, AutoAssignMode);
-	SharedData->SelectUniverse(FixturePatchesToAutoAssign[0]->GetUniverseID());
+	const FText TransactionText = FText::Format(LOCTEXT("AlignTransaction", "Align Fixture {0}|plural(one=Patch, other=Patches)"), FixturePatchesToAlign.Num());
+	const FScopedTransaction AlignTransaction(TransactionText);
+	
+	UE::DMXEditor::FixturePatchEditor::Private::FScopedEditChangeFixturePatches ScopedEditChange(FixturePatchesToAlign);
+	UE::DMXEditor::AutoAssign::FAutoAssignUtility::Align(FixturePatchesToAlign);
 
 	RequestRefresh();
 }
 
+void SDMXPatchedUniverse::StackFixturePatches()
+{
+	TArray<UDMXEntityFixturePatch*> FixturePatchesToStack;
+	if (!GetSelectedFixturePatches(FixturePatchesToStack))
+	{
+		return;
+	}
+
+	const FText TransactionText = FText::Format(LOCTEXT("StackTransaction", "Stack Fixture {0}|plural(one=Patch, other=Patches)"), FixturePatchesToStack.Num());
+	const FScopedTransaction StackTransaction(TransactionText);
+
+	UE::DMXEditor::FixturePatchEditor::Private::FScopedEditChangeFixturePatches ScopedEditChange(FixturePatchesToStack);
+	UE::DMXEditor::AutoAssign::FAutoAssignUtility::Stack(FixturePatchesToStack);
+
+	RequestRefresh();
+}
+
+void SDMXPatchedUniverse::SpreadFixturePatchesOverUniverses()
+{
+	TArray<UDMXEntityFixturePatch*> FixturePatchesToSpread;
+	if (!GetSelectedFixturePatches(FixturePatchesToSpread))
+	{
+		return;
+	}
+
+	const FText TransactionText = FText::Format(LOCTEXT("SpreadOverUniverseTransaction", "Spread Fixture {0}|plural(one=Patch, other=Patches) over Universes"), FixturePatchesToSpread.Num());
+	const FScopedTransaction SpreadOverUniverseTransaction(TransactionText);
+
+	UE::DMXEditor::FixturePatchEditor::Private::FScopedEditChangeFixturePatches ScopedEditChange(FixturePatchesToSpread);
+	UE::DMXEditor::AutoAssign::FAutoAssignUtility::SpreadOverUniverses(FixturePatchesToSpread);
+
+	RequestRefresh();
+}
+
+bool SDMXPatchedUniverse::GetSelectedFixturePatches(TArray<UDMXEntityFixturePatch*>& OutFixturePatchArray) const
+{
+	OutFixturePatchArray.Reset();
+
+	Algo::TransformIf(SharedData->GetSelectedFixturePatches(), OutFixturePatchArray,
+		[](const TWeakObjectPtr<UDMXEntityFixturePatch>& WeakFixturePatch) { return WeakFixturePatch.IsValid(); },
+		[](const TWeakObjectPtr<UDMXEntityFixturePatch>& WeakFixturePatch) { return WeakFixturePatch.Get(); });
+
+	return !OutFixturePatchArray.IsEmpty();
+}
+
 bool SDMXPatchedUniverse::DoesDMXLibraryHaveReachableUniverses() const
 {
 	if (UDMXLibrary* DMXLibrary = GetDMXLibrary())
@@ -944,7 +1061,7 @@ TArray<UDMXEntityFixturePatch*> SDMXPatchedUniverse::GetFixturePatchesOnChannel(
 
 UDMXLibrary* SDMXPatchedUniverse::GetDMXLibrary() const
 {
-	if (TSharedPtr<FDMXEditor> DMXEditor = DMXEditorPtr.Pin())
+	if (TSharedPtr<FDMXEditor> DMXEditor = WeakDMXEditor.Pin())
 	{
 		return DMXEditor->GetDMXLibrary();
 	}
@@ -961,23 +1078,21 @@ void SDMXPatchedUniverse::RegisterCommands()
 	// listen to common editor shortcuts for copy/paste etc
 	CommandList = MakeShared<FUICommandList>();
 
+	using namespace UE::DMXEditor::AutoAssign;
 	CommandList->MapAction
 	(
-		FDMXEditorCommands::Get().AutoAssignToFirstUniverseInSelection,
-		FExecuteAction::CreateSP(this, &SDMXPatchedUniverse::AutoAssignFixturePatches, EDMXFixturePatchAutoAssignMode::FirstUniverseInSelection)
+		FDMXEditorCommands::Get().AutoAssignSelectedUniverse,
+		FExecuteAction::CreateSP(this, &SDMXPatchedUniverse::AutoAssignFixturePatches, EAutoAssignMode::SelectedUniverse)
 	);
-
 	CommandList->MapAction
 	(
-		FDMXEditorCommands::Get().AutoAssignToFirstReachableUniverse,
-		FExecuteAction::CreateSP(this, &SDMXPatchedUniverse::AutoAssignFixturePatches, EDMXFixturePatchAutoAssignMode::FirstReachableUniverse),
-		FCanExecuteAction::CreateSP(this, &SDMXPatchedUniverse::DoesDMXLibraryHaveReachableUniverses)
+		FDMXEditorCommands::Get().Align,
+		FExecuteAction::CreateSP(this, &SDMXPatchedUniverse::AlignFixturePatches)
 	);
-
 	CommandList->MapAction
 	(
-		FDMXEditorCommands::Get().AutoAssignAfterLastAddressInLibrary,
-		FExecuteAction::CreateSP(this, &SDMXPatchedUniverse::AutoAssignFixturePatches, EDMXFixturePatchAutoAssignMode::LastAddressInLibrary)
+		FDMXEditorCommands::Get().Stack,
+		FExecuteAction::CreateSP(this, &SDMXPatchedUniverse::StackFixturePatches)
 	);
 }
 
@@ -995,11 +1110,62 @@ TSharedRef<SWidget> SDMXPatchedUniverse::CreateContextMenu(int32 Channel)
 	// Auto Assign Section
 	if (!SharedData->GetSelectedFixturePatches().IsEmpty())
 	{
-		MenuBuilder.BeginSection("AutoAssignSection", LOCTEXT("AutoAssignSection", "Auto-Assign"));
+		using namespace UE::DMXEditor::AutoAssign;
+
+		// Auto assign section
+		MenuBuilder.BeginSection("AutoAssignSection", LOCTEXT("AutoAssignActionsSection", "Auto-Assign"));
 		{
-			MenuBuilder.AddMenuEntry(FDMXEditorCommands::Get().AutoAssignToFirstUniverseInSelection);
-			MenuBuilder.AddMenuEntry(FDMXEditorCommands::Get().AutoAssignToFirstReachableUniverse);
-			MenuBuilder.AddMenuEntry(FDMXEditorCommands::Get().AutoAssignAfterLastAddressInLibrary);
+			MenuBuilder.AddMenuEntry(FDMXEditorCommands::Get().AutoAssignSelectedUniverse);
+
+			const FText AutoAssignUnderMouseAssignLabel = FText::Format(LOCTEXT("AutoAssignUnderMouseAssignLabel", "Auto-Assign at {0}.{1}"), FText::FromString(FString::FromInt(UniverseID)), FText::FromString(FString::FromInt(LastHoveredChannel)));
+			MenuBuilder.AddMenuEntry
+			(
+				AutoAssignUnderMouseAssignLabel,
+				LOCTEXT("AssignTooltip", "Auto-assigns selected patches to first consecutive range of free channels, starting from the current mouse position."),
+				FSlateIcon(),
+				FUIAction
+				(
+					FExecuteAction::CreateSP(this, &SDMXPatchedUniverse::AutoAssignFixturePatches, EAutoAssignMode::ChannelUnderMouse),
+					FCanExecuteAction::CreateLambda([this, SharedThis = AsShared()] { return !SharedData->GetSelectedFixturePatches().IsEmpty(); })
+				),
+				NAME_None,
+				EUserInterfaceActionType::Button
+			);
+		}
+		MenuBuilder.EndSection();
+
+		// Assign section
+		MenuBuilder.BeginSection("AssignSection", LOCTEXT("AssignSection", "Assign"));
+		{
+			const FText AssignLabel = FText::Format(LOCTEXT("AssignMenuLabel", "Assign to {0}.{1}"), FText::FromString(FString::FromInt(UniverseID)), FText::FromString(FString::FromInt(LastHoveredChannel)));
+			MenuBuilder.AddMenuEntry
+			(
+				AssignLabel,
+				LOCTEXT("AssignTooltip", "Assigns selected patches to the hovered channel"),
+				FSlateIcon(),
+				FUIAction
+				(
+					FExecuteAction::CreateSP(this, &SDMXPatchedUniverse::AssignFixturePatches),
+					FCanExecuteAction::CreateLambda([this, SharedThis = AsShared()] { return !SharedData->GetSelectedFixturePatches().IsEmpty(); })
+				),
+				NAME_None,
+				EUserInterfaceActionType::Button
+			);
+			MenuBuilder.AddMenuEntry(FDMXEditorCommands::Get().Align);
+			MenuBuilder.AddMenuEntry(FDMXEditorCommands::Get().Stack);
+			MenuBuilder.AddMenuEntry
+			(
+				LOCTEXT("SpreadOverUniversesLabel", "Spread over Universes"),
+				LOCTEXT("SpreadOverUniversesTooltip", "Assigns each patch to its own Universe"),
+				FSlateIcon(),
+				FUIAction
+				(
+					FExecuteAction::CreateSP(this, &SDMXPatchedUniverse::SpreadFixturePatchesOverUniverses),
+					FCanExecuteAction::CreateLambda([this, SharedThis = AsShared()] { return SharedData->GetSelectedFixturePatches().Num() > 1; })
+				),
+				NAME_None,
+				EUserInterfaceActionType::Button
+			);
 		}
 		MenuBuilder.EndSection();
 	}
diff --git a/Engine/Plugins/VirtualProduction/DMX/DMXEngine/Source/DMXEditor/Private/Widgets/FixturePatch/SDMXPatchedUniverse.h b/Engine/Plugins/VirtualProduction/DMX/DMXEngine/Source/DMXEditor/Private/Widgets/FixturePatch/SDMXPatchedUniverse.h
--- a/Engine/Plugins/VirtualProduction/DMX/DMXEngine/Source/DMXEditor/Private/Widgets/FixturePatch/SDMXPatchedUniverse.h
+++ b/Engine/Plugins/VirtualProduction/DMX/DMXEngine/Source/DMXEditor/Private/Widgets/FixturePatch/SDMXPatchedUniverse.h
@@ -12,21 +12,19 @@
 #include "Widgets/DeclarativeSyntaxSupport.h"
 #include "Widgets/SCompoundWidget.h"
 
-enum class EDMXFixturePatchAutoAssignMode : uint8;
 class FDMXEditor;
 class FDMXFixturePatchNode;
 class FDMXFixturePatchSharedData;
-class SDMXChannelConnector;
-class SDMXFixturePatchFragment;
-class UDMXLibrary;
-class UDMXEntityFixturePatch;
-
 class FDragDropEvent;
 struct FTimerHandle;
 class FUICommandList;
 class SBorder;
+class SDMXChannelConnector;
+class SDMXFixturePatchFragment;
 class SGridPanel;
-
+class UDMXLibrary;
+class UDMXEntityFixturePatch;
+namespace UE::DMXEditor::AutoAssign { enum class EAutoAssignMode : uint8; }
 
 enum class EDMXPatchedUniverseReachability
 {
@@ -178,7 +176,22 @@ private:
 	void OnSelectionChanged();
 
 	/** Auto assigns selected Fixture Patches */
-	void AutoAssignFixturePatches(EDMXFixturePatchAutoAssignMode AutoAssignMode);
+	void AutoAssignFixturePatches(UE::DMXEditor::AutoAssign::EAutoAssignMode AutoAssignMode);
+
+	/** Assign selected Fixture Patches */
+	void AssignFixturePatches();
+
+	/** Aligns selected Fixture Patches */
+	void AlignFixturePatches();
+
+	/** Stacks selected Fixture Patches */
+	void StackFixturePatches();
+
+	/** Spreads selected Fixture Patches over Universes */
+	void SpreadFixturePatchesOverUniverses();
+
+	/** Gets selected fixture patches. Returns false if no patches are selected */
+	bool GetSelectedFixturePatches(TArray<UDMXEntityFixturePatch*>& OutFixturePatchArray) const;
 
 	/** Returns true if the DMX Library has reachable Universes */
 	bool DoesDMXLibraryHaveReachableUniverses() const;
@@ -198,6 +211,9 @@ private:
 	/** The hovered channel, or unset if no channel is hovered */
 	TSharedPtr<SDMXChannelConnector> HoveredChannel;
 
+	/** The channel that was last hovered */
+	int32 LastHoveredChannel = 1;
+
 	/** If true, monitors inputs */
 	bool bMonitorInputs = false;
 
@@ -238,7 +254,7 @@ private:
 	TSharedPtr<FDMXFixturePatchSharedData> SharedData;
 
 	/** The owning editor */
-	TWeakPtr<FDMXEditor> DMXEditorPtr;
+	TWeakPtr<FDMXEditor> WeakDMXEditor;
 
 	private:
 	///////////////////////////////////////////////////
diff --git a/Engine/Plugins/VirtualProduction/DMX/DMXEngine/Source/DMXEditor/Private/Widgets/FixtureType/SDMXFixtureTypeModesEditorModeRow.cpp b/Engine/Plugins/VirtualProduction/DMX/DMXEngine/Source/DMXEditor/Private/Widgets/FixtureType/SDMXFixtureTypeModesEditorModeRow.cpp
--- a/Engine/Plugins/VirtualProduction/DMX/DMXEngine/Source/DMXEditor/Private/Widgets/FixtureType/SDMXFixtureTypeModesEditorModeRow.cpp
+++ b/Engine/Plugins/VirtualProduction/DMX/DMXEngine/Source/DMXEditor/Private/Widgets/FixtureType/SDMXFixtureTypeModesEditorModeRow.cpp
@@ -80,6 +80,11 @@ bool SDMXFixtureTypeModesEditorModeRow::OnVerifyModeNameChanged(const FText& InN
 
 void SDMXFixtureTypeModesEditorModeRow::OnModeNameCommitted(const FText& InNewText, ETextCommit::Type InTextCommit)
 {
+	if (InNewText.IsEmpty())
+	{
+		return;
+	}
+
 	FText UniqueModeName;
 	ModeItem->SetModeName(InNewText, UniqueModeName);
 
diff --git a/Engine/Plugins/VirtualProduction/DMX/DMXEngine/Source/DMXEditor/Public/Commands/DMXEditorCommands.h b/Engine/Plugins/VirtualProduction/DMX/DMXEngine/Source/DMXEditor/Public/Commands/DMXEditorCommands.h
--- a/Engine/Plugins/VirtualProduction/DMX/DMXEngine/Source/DMXEditor/Public/Commands/DMXEditorCommands.h
+++ b/Engine/Plugins/VirtualProduction/DMX/DMXEngine/Source/DMXEditor/Public/Commands/DMXEditorCommands.h
@@ -37,7 +37,7 @@ public:
 	TSharedPtr<FUICommandInfo> ToggleSendDMX;
 
 	// Auto assign
-	TSharedPtr<FUICommandInfo> AutoAssignToFirstUniverseInSelection;
-	TSharedPtr<FUICommandInfo> AutoAssignToFirstReachableUniverse;
-	TSharedPtr<FUICommandInfo> AutoAssignAfterLastAddressInLibrary;
+	TSharedPtr<FUICommandInfo> Align;
+	TSharedPtr<FUICommandInfo> Stack;
+	TSharedPtr<FUICommandInfo> AutoAssignSelectedUniverse;
 };
diff --git a/Engine/Plugins/VirtualProduction/DMX/DMXEngine/Source/DMXRuntime/Public/Library/DMXLibrary.h b/Engine/Plugins/VirtualProduction/DMX/DMXEngine/Source/DMXRuntime/Public/Library/DMXLibrary.h
--- a/Engine/Plugins/VirtualProduction/DMX/DMXEngine/Source/DMXRuntime/Public/Library/DMXLibrary.h
+++ b/Engine/Plugins/VirtualProduction/DMX/DMXEngine/Source/DMXRuntime/Public/Library/DMXLibrary.h
@@ -41,7 +41,7 @@ DECLARE_MULTICAST_DELEGATE_TwoParams(FDMXOnEntityArrayChangedDelegate, class UDM
 /** DEPRECATED 5.0 */
 DECLARE_MULTICAST_DELEGATE_OneParam(FOnEntitiesUpdated_DEPRECATED, class UDMXLibrary*);
 
-UCLASS(BlueprintType, Blueprintable, AutoExpandCategories = DMX)
+UCLASS(BlueprintType, Blueprintable, Config = DMXEditor, AutoExpandCategories = DMX)
 class DMXRUNTIME_API UDMXLibrary
 	: public UDMXObjectBase
 {
@@ -244,7 +244,7 @@ protected:
 	void UpgradeFromControllersToPorts();
 #endif // WITH_EDITOR
 
-	/** Input ports of the Library */
+	/** Ports of the Library */
 	UPROPERTY(EditAnywhere, BlueprintReadOnly, NonTransactional, Category = "DMX", Meta = (ShowOnlyInnerProperties))
 	FDMXLibraryPortReferences PortReferences;
 
