DMX - Fix CIS Build Engine Localization

diff --git a/Engine/Plugins/VirtualProduction/DMX/DMXEngine/Source/DMXEditor/Private/Widgets/FixturePatch/SDMXMVRFixtureListRow.cpp b/Engine/Plugins/VirtualProduction/DMX/DMXEngine/Source/DMXEditor/Private/Widgets/FixturePatch/SDMXMVRFixtureListRow.cpp
--- a/Engine/Plugins/VirtualProduction/DMX/DMXEngine/Source/DMXEditor/Private/Widgets/FixturePatch/SDMXMVRFixtureListRow.cpp
+++ b/Engine/Plugins/VirtualProduction/DMX/DMXEngine/Source/DMXEditor/Private/Widgets/FixturePatch/SDMXMVRFixtureListRow.cpp
@@ -707,7 +707,7 @@ TSharedRef<SWidget> SDMXMVRFixtureListRow::GeneratePatchWidget()
 						}
 						else
 						{
-							return LOCTEXT("AddressesText", "Not patched");
+							return LOCTEXT("NotPatchedText", "Not patched");
 						}
 					})
 				.Font(FAppStyle::GetFontStyle(TEXT("PropertyWindow.NormalFont")))
diff --git a/Engine/Plugins/VirtualProduction/DMX/DMXEngine/Source/DMXEditor/Private/Widgets/FixturePatch/SDMXPatchedUniverse.cpp b/Engine/Plugins/VirtualProduction/DMX/DMXEngine/Source/DMXEditor/Private/Widgets/FixturePatch/SDMXPatchedUniverse.cpp
--- a/Engine/Plugins/VirtualProduction/DMX/DMXEngine/Source/DMXEditor/Private/Widgets/FixturePatch/SDMXPatchedUniverse.cpp
+++ b/Engine/Plugins/VirtualProduction/DMX/DMXEngine/Source/DMXEditor/Private/Widgets/FixturePatch/SDMXPatchedUniverse.cpp
@@ -1121,7 +1121,7 @@ TSharedRef<SWidget> SDMXPatchedUniverse::CreateContextMenu(int32 Channel)
 			MenuBuilder.AddMenuEntry
 			(
 				AutoAssignUnderMouseAssignLabel,
-				LOCTEXT("AssignTooltip", "Auto-assigns selected patches to first consecutive range of free channels, starting from the current mouse position."),
+				LOCTEXT("AutoAssignTooltip", "Auto-assigns selected patches to first consecutive range of free channels, starting from the current mouse position."),
 				FSlateIcon(),
 				FUIAction
 				(
