DMX - Add missing address box to the new 'add  patch' menu, fix various small issues with new ways of assigning patches

diff --git a/Engine/Plugins/VirtualProduction/DMX/DMXEngine/Source/DMXEditor/Private/FixturePatchAutoAssignUtility.cpp b/Engine/Plugins/VirtualProduction/DMX/DMXEngine/Source/DMXEditor/Private/FixturePatchAutoAssignUtility.cpp
--- a/Engine/Plugins/VirtualProduction/DMX/DMXEngine/Source/DMXEditor/Private/FixturePatchAutoAssignUtility.cpp
+++ b/Engine/Plugins/VirtualProduction/DMX/DMXEngine/Source/DMXEditor/Private/FixturePatchAutoAssignUtility.cpp
@@ -74,8 +74,8 @@ namespace UE::DMXEditor::AutoAssign::Private
 		if (UDMXEntityFixturePatch* Patch = WeakFixturePatch.Get())
 		{
 			const int64 Universe = (GetLowerBoundValue() + 1) / DMX_UNIVERSE_SIZE;
-			const int64 Channel = (GetLowerBoundValue() + 1) % DMX_UNIVERSE_SIZE;
-			if (!ensureMsgf(Universe >= 1 && Channel >= 1, TEXT("Invalid universe and channel resulting from auto assign. Channels not applied for '%s'."), *Patch->Name))
+			const int64 Channel =  GetLowerBoundValue() + 1 - Universe * DMX_UNIVERSE_SIZE;
+			if (!ensureMsgf(Universe >= 1 && Channel >= 1 && Channel <= DMX_UNIVERSE_SIZE, TEXT("Invalid universe and channel resulting from auto assign. Channels not applied for '%s'."), *Patch->Name))
 			{
 				return;
 			}
@@ -366,13 +366,13 @@ namespace UE::DMXEditor::AutoAssign
 		FAutoAssignElementsUtility::SpreadOverUniverses(PatchElements);
 	}
 
-	int32 FAutoAssignUtility::AutoAssign(EAutoAssignMode Mode, const TSharedRef<FDMXEditor>& DMXEditor, TArray<UDMXEntityFixturePatch*> FixturePatches, int32 UniverseUnderMouse, int32 ChannelUnderMouse)
+	int32 FAutoAssignUtility::AutoAssign(EAutoAssignMode Mode, const TSharedRef<FDMXEditor>& DMXEditor, TArray<UDMXEntityFixturePatch*> FixturePatches, int32 UserDefinedUniverse, int32 UserDefinedChannel)
 	{
 		FAutoAssignUtility Instance;
-		return Instance.AutoAssignInternal(Mode, DMXEditor, FixturePatches, UniverseUnderMouse, ChannelUnderMouse);
+		return Instance.AutoAssignInternal(Mode, DMXEditor, FixturePatches, UserDefinedUniverse, UserDefinedChannel);
 	}
 
-	int32 FAutoAssignUtility::AutoAssignInternal(EAutoAssignMode Mode, const TSharedRef<FDMXEditor>& DMXEditor, TArray<UDMXEntityFixturePatch*> FixturePatches, int32 UniverseUnderMouse, int32 ChannelUnderMouse)
+	int32 FAutoAssignUtility::AutoAssignInternal(EAutoAssignMode Mode, const TSharedRef<FDMXEditor>& DMXEditor, TArray<UDMXEntityFixturePatch*> FixturePatches, int32 UserDefinedUniverse, int32 UserDefinedChannel)
 	{
 		UDMXLibrary* DMXLibrary = GetDMXLibrary(FixturePatches);
 		if (!DMXLibrary)
@@ -381,7 +381,7 @@ namespace UE::DMXEditor::AutoAssign
 		}
 
 		using namespace Private;
-		const int64 AbsoluteStartingChannel = FindAbsoluteStartingChannel(Mode, DMXEditor , *DMXLibrary, FixturePatches, UniverseUnderMouse, ChannelUnderMouse);
+		const int64 AbsoluteStartingChannel = FindAbsoluteStartingChannel(Mode, DMXEditor , *DMXLibrary, FixturePatches, UserDefinedUniverse, UserDefinedChannel);
 		const TArray<TSharedRef<FAutoAssignElement>> FreeElements = CreateFreeElements(*DMXLibrary, FixturePatches, AbsoluteStartingChannel);
 		const TArray<TSharedRef<FAutoAssignElement>> PatchElements = CreatePatchElements(FixturePatches);
 
@@ -399,7 +399,7 @@ namespace UE::DMXEditor::AutoAssign
 		return (*FirstPatchElementPtr)->GetLowerBoundValue() / DMX_UNIVERSE_SIZE;
 	}
 
-	int64 FAutoAssignUtility::FindAbsoluteStartingChannel(EAutoAssignMode Mode, const TSharedRef<FDMXEditor>& DMXEditor, const UDMXLibrary& DMXLibrary, const TArray<UDMXEntityFixturePatch*>& FixturePatches, int32 UniverseUnderMouse, int32 ChannelUnderMouse) const
+	int64 FAutoAssignUtility::FindAbsoluteStartingChannel(EAutoAssignMode Mode, const TSharedRef<FDMXEditor>& DMXEditor, const UDMXLibrary& DMXLibrary, const TArray<UDMXEntityFixturePatch*>& FixturePatches, int32 UniverseUnderMouse, int32 UserDefinedChannel) const
 	{
 		if (!ensureMsgf(!FixturePatches.IsEmpty() && !FixturePatches.Contains(nullptr), TEXT("Unexpected no fixture patches or null entry in Fixture Patches. Expected to be handled during FAutoAssignUtility construction.")))
 		{
@@ -407,16 +407,16 @@ namespace UE::DMXEditor::AutoAssign
 		}
 
 		if (!ensureMsgf(
-			Mode != EAutoAssignMode::ChannelUnderMouse || 
-			(UniverseUnderMouse > 0 && UniverseUnderMouse <= DMX_UNIVERSE_SIZE && ChannelUnderMouse > 0 && ChannelUnderMouse <= DMX_UNIVERSE_SIZE),
+			Mode != EAutoAssignMode::UserDefinedChannel || 
+			(UniverseUnderMouse > 0 && UniverseUnderMouse <= DMX_UNIVERSE_SIZE && UserDefinedChannel > 0 && UserDefinedChannel <= DMX_UNIVERSE_SIZE),
 			TEXT("Trying to auto assign to channel under mouse, but provided channels are not valid. Skipping auto assign.")))
 		{
 			return DMX_UNIVERSE_SIZE;
 		}
 
-		if (Mode == EAutoAssignMode::ChannelUnderMouse)
+		if (Mode == EAutoAssignMode::UserDefinedChannel)
 		{
-			return (int64)UniverseUnderMouse * DMX_UNIVERSE_SIZE + ChannelUnderMouse - 1;
+			return (int64)UniverseUnderMouse * DMX_UNIVERSE_SIZE + UserDefinedChannel - 1;
 		}
 		else if (Mode == EAutoAssignMode::SelectedUniverse)
 		{
@@ -454,7 +454,7 @@ namespace UE::DMXEditor::AutoAssign
 
 			if (LastOtherFixturePatchPtr)
 			{
-				return  (int64)(*LastOtherFixturePatchPtr)->GetUniverseID() * DMX_UNIVERSE_SIZE + (*LastOtherFixturePatchPtr)->GetStartingChannel() + (*LastOtherFixturePatchPtr)->GetChannelSpan();
+				return  (int64)(*LastOtherFixturePatchPtr)->GetUniverseID() * DMX_UNIVERSE_SIZE + (*LastOtherFixturePatchPtr)->GetStartingChannel() + (*LastOtherFixturePatchPtr)->GetChannelSpan() - 1;
 			}
 			else
 			{
diff --git a/Engine/Plugins/VirtualProduction/DMX/DMXEngine/Source/DMXEditor/Private/FixturePatchAutoAssignUtility.h b/Engine/Plugins/VirtualProduction/DMX/DMXEngine/Source/DMXEditor/Private/FixturePatchAutoAssignUtility.h
--- a/Engine/Plugins/VirtualProduction/DMX/DMXEngine/Source/DMXEditor/Private/FixturePatchAutoAssignUtility.h
+++ b/Engine/Plugins/VirtualProduction/DMX/DMXEngine/Source/DMXEditor/Private/FixturePatchAutoAssignUtility.h
@@ -82,7 +82,7 @@ namespace UE::DMXEditor::AutoAssign
 {
 	enum class EAutoAssignMode : uint8
 	{
-		ChannelUnderMouse,
+		UserDefinedChannel,
 		SelectedUniverse,
 		FirstReachableUniverse,
 		AfterLastPatchedUniverse
@@ -105,14 +105,14 @@ namespace UE::DMXEditor::AutoAssign
 		static void SpreadOverUniverses(TArray<UDMXEntityFixturePatch*> FixturePatches);
 
 		/** Auto assigns FixturePatches. Returns the Universe the first patch was auto assigned to.*/
-		static int32 AutoAssign(EAutoAssignMode Mode, const TSharedRef<FDMXEditor>& DMXEditor, TArray<UDMXEntityFixturePatch*> FixturePatches, int32 UniverseUnderMouse = 0, int32 ChannelUnderMouse = 0);
+		static int32 AutoAssign(EAutoAssignMode Mode, const TSharedRef<FDMXEditor>& DMXEditor, TArray<UDMXEntityFixturePatch*> FixturePatches, int32 UniverseUnderMouse = 0, int32 UserDefinedChannel = 0);
 
 	private:
 		/** Non-static implementation of auto assign */
-		int32 AutoAssignInternal(EAutoAssignMode Mode, const TSharedRef<FDMXEditor>& DMXEditor, TArray<UDMXEntityFixturePatch*> FixturePatches, int32 UniverseUnderMouse = 0, int32 ChannelUnderMouse = 0);
+		int32 AutoAssignInternal(EAutoAssignMode Mode, const TSharedRef<FDMXEditor>& DMXEditor, TArray<UDMXEntityFixturePatch*> FixturePatches, int32 UniverseUnderMouse = 0, int32 UserDefinedChannel = 0);
 
 		/** Finds the absolute starting channel from options */
-		int64 FindAbsoluteStartingChannel(EAutoAssignMode Mode, const TSharedRef<FDMXEditor>& DMXEditor, const UDMXLibrary& DMXLibrary, const TArray<UDMXEntityFixturePatch*>& FixturePatches, int32 UniverseUnderMouse, int32 ChannelUnderMouse) const;
+		int64 FindAbsoluteStartingChannel(EAutoAssignMode Mode, const TSharedRef<FDMXEditor>& DMXEditor, const UDMXLibrary& DMXLibrary, const TArray<UDMXEntityFixturePatch*>& FixturePatches, int32 UniverseUnderMouse, int32 UserDefinedChannel) const;
 
 		/** Gets the DMX Library from the patches, or returns nullptr if a valid DMX Library cannot be deduced from patches */
 		UDMXLibrary* GetDMXLibrary(TArray<UDMXEntityFixturePatch*> FixturePatches) const;
diff --git a/Engine/Plugins/VirtualProduction/DMX/DMXEngine/Source/DMXEditor/Private/Widgets/FixturePatch/SAddFixturePatchMenu.cpp b/Engine/Plugins/VirtualProduction/DMX/DMXEngine/Source/DMXEditor/Private/Widgets/FixturePatch/SAddFixturePatchMenu.cpp
--- a/Engine/Plugins/VirtualProduction/DMX/DMXEngine/Source/DMXEditor/Private/Widgets/FixturePatch/SAddFixturePatchMenu.cpp
+++ b/Engine/Plugins/VirtualProduction/DMX/DMXEngine/Source/DMXEditor/Private/Widgets/FixturePatch/SAddFixturePatchMenu.cpp
@@ -5,11 +5,13 @@
 #include "Algo/Copy.h"
 #include "Algo/Find.h"
 #include "DMXAddFixturePatchMenuData.h"
+#include "DMXEditor.h"
 #include "DMXFixturePatchSharedData.h"
 #include "Editor.h"
 #include "Framework/Application/SlateApplication.h"
 #include "FixturePatchAutoAssignUtility.h"
 #include "Library/DMXEntityFixtureType.h"
+#include "Library/DMXLibrary.h"
 #include "ScopedTransaction.h"
 #include "TimerManager.h"
 #include "UObject/Package.h"
@@ -17,196 +19,258 @@
 #include "Widgets/Layout/SGridPanel.h"
 #include "Widgets/SBoxPanel.h"
 #include "Widgets/SDMXEntityDropdownMenu.h"
-#include "Widgets/Text/SInlineEditableTextBlock.h"
+#include "Widgets/Input/SEditableTextBox.h"
 
 
 #define LOCTEXT_NAMESPACE "SAddFixturePatchMenu"
 
 namespace UE::DMXEditor::FixturePatchEditor
 {
-
 	SAddFixturePatchMenu::~SAddFixturePatchMenu()
 	{
-		UDMXAddFixturePatchMenuData* Model = GetMutableDefault<UDMXAddFixturePatchMenuData>();
-		Model->SoftFixtureType = WeakFixtureType.IsValid() ? WeakFixtureType.Get() : Model->SoftFixtureType;
-		Model->ActiveModeIndex = WeakFixtureType.IsValid() ? ActiveModeIndex : Model->ActiveModeIndex;
-		Model->SaveConfig();
+		SaveConfig();
 	}
 
 	void SAddFixturePatchMenu::Construct(const FArguments& InArgs, TWeakPtr<FDMXEditor> InWeakDMXEditor)
 	{
 		WeakDMXEditor = InWeakDMXEditor;
 
+		SharedData = WeakDMXEditor.IsValid() ? WeakDMXEditor.Pin()->GetFixturePatchSharedData() : nullptr;
+		UDMXLibrary* DMXLibrary = WeakDMXEditor.IsValid() ? WeakDMXEditor.Pin()->GetDMXLibrary() : nullptr;
+		if (!SharedData.IsValid() || !DMXLibrary)
+		{
+			return;
+		}
+
+		const TArray<UDMXEntityFixtureType*> FixtureTypes = DMXLibrary->GetEntitiesTypeCast<UDMXEntityFixtureType>();
+		UDMXAddFixturePatchMenuData* MenuData = GetMutableDefault<UDMXAddFixturePatchMenuData>();
+		UDMXEntityFixtureType* const* SelectedFixtureTypePtr = Algo::FindByPredicate(FixtureTypes, [MenuData](const UDMXEntityFixtureType* FixtureType)
+			{
+				return FixtureType == MenuData->SoftFixtureType;
+			});
+		if (!SelectedFixtureTypePtr)
+		{
+			SelectedFixtureTypePtr = FixtureTypes.IsEmpty() ? nullptr : &FixtureTypes[0];
+		}
+		WeakFixtureType = SelectedFixtureTypePtr ? *SelectedFixtureTypePtr : nullptr;
+		ActiveModeIndex = WeakFixtureType.IsValid() && WeakFixtureType->Modes.IsValidIndex(ActiveModeIndex) ? ActiveModeIndex : 0;
+
 		ChildSlot
 		[
-			SNew(SGridPanel)
-			.FillColumn(1, 1.f)
-
-			// Select Fixture Type
-			+ SGridPanel::Slot(0, 0)
-			.Padding(4.f)
-			.VAlign(VAlign_Center)
+			SNew(SBorder)
+			.BorderImage(FAppStyle::GetBrush("ToolPanel.GroupBorder"))
 			[
-				SNew(STextBlock)
-				.Text(LOCTEXT("SelectFixtureTypeLabel", "Fixture Type"))
-				.Font(FAppStyle::GetFontStyle(TEXT("PropertyWindow.NormalFont")))
-			]
-			+ SGridPanel::Slot(1, 0)
+				SNew(SGridPanel)
+				.FillColumn(1, 1.f)
+
+				// Select Fixture Type
+				+ SGridPanel::Slot(0, 0)
 				.Padding(4.f)
-			.VAlign(VAlign_Center)
-			[
-				SNew(SBox)
-				.MinDesiredWidth(100.f)
+				.VAlign(VAlign_Center)
 				[
-					SAssignNew(FixtureTypeSelector, SDMXEntityPickerButton<UDMXEntityFixtureType>)
-					.DMXEditor(WeakDMXEditor)
-					.CurrentEntity_Lambda([this] { return WeakFixtureType.Get(); })
-					.OnEntitySelected(this, &SAddFixturePatchMenu::OnFixtureTypeSelected)	
+					SNew(STextBlock)
+					.Text(LOCTEXT("SelectFixtureTypeLabel", "Fixture Type"))
+					.Font(FAppStyle::GetFontStyle(TEXT("PropertyWindow.NormalFont")))
 				]
-			]
 
-			// Select Mode
-			+ SGridPanel::Slot(0, 1)
-			.Padding(4.f)
-			.VAlign(VAlign_Center)
-			[
-				SNew(STextBlock)
-				.IsEnabled(this, &SAddFixturePatchMenu::HasValidFixtureTypeAndMode)
-				.Text(LOCTEXT("SelectModeLabel", "Mode"))
-				.Font(FAppStyle::GetFontStyle(TEXT("PropertyWindow.NormalFont")))
-			]
-			+ SGridPanel::Slot(1, 1)
-			.Padding(4.f)
-			[
-				SAssignNew(ModeComboBox, SComboBox<TSharedPtr<uint32>>)
-				.IsEnabled(this, &SAddFixturePatchMenu::HasValidFixtureTypeAndMode)
-				.OptionsSource(&ModeSources)
-				.OnGenerateWidget(this, &SAddFixturePatchMenu::GenerateModeComboBoxEntry)
-				.OnSelectionChanged(this, &SAddFixturePatchMenu::OnModeSelected)
-				.InitiallySelectedItem(0)
+				+ SGridPanel::Slot(1, 0)
+				.Padding(4.f)
+				.VAlign(VAlign_Center)
+				[
+					MakeFixtureTypeSelectWidget()
+				]
+
+				// Select Mode
+				+ SGridPanel::Slot(0, 1)
+				.Padding(4.f)
+				.VAlign(VAlign_Center)
 				[
 					SNew(STextBlock)
-					.MinDesiredWidth(50.0f)
-					.Text(this, &SAddFixturePatchMenu::GetActiveModeText)
+					.IsEnabled(this, &SAddFixturePatchMenu::HasValidFixtureTypeAndMode)
+					.Text(LOCTEXT("SelectModeLabel", "Mode"))
 					.Font(FAppStyle::GetFontStyle(TEXT("PropertyWindow.NormalFont")))
 				]
-			]
-
-			// Num Fixture Patches Label
-			+ SGridPanel::Slot(0, 2)
-			.Padding(4.f)
-			.VAlign(VAlign_Center)
-			[
-				SNew(STextBlock)
-				.IsEnabled(this, &SAddFixturePatchMenu::HasValidFixtureTypeAndMode)
-				.Text(LOCTEXT("NumPatchesLabel", "Num Patches"))
-				.Font(FAppStyle::GetFontStyle(TEXT("PropertyWindow.NormalFont")))
-			]
+				+ SGridPanel::Slot(1, 1)
+				.Padding(4.f)
+				[
+					MakeModeSelectWidget()
+				]
 
-			// Num Fixture Patches editable Text Box
-			+ SGridPanel::Slot(1, 2)
-			.Padding(4.f)
-			[
-				SNew(SBox)
-				.HAlign(HAlign_Left)
+				// Universe label 
+				+ SGridPanel::Slot(0, 2)
+				.Padding(4.f)
+				.VAlign(VAlign_Center)
 				[
-					SNew(SEditableTextBox)
+					SNew(STextBlock)
 					.IsEnabled(this, &SAddFixturePatchMenu::HasValidFixtureTypeAndMode)
-					.MinDesiredWidth(60.f)
-					.SelectAllTextWhenFocused(true)
-					.ClearKeyboardFocusOnCommit(false)
-					.SelectAllTextOnCommit(true)
-					.Text_Lambda([this]()
-						{
-							return FText::FromString(FString::FromInt(NumFixturePatchesToAdd));
-						})
-					.OnVerifyTextChanged_Lambda([](const FText& InNewText, FText& OutErrorMessage)
-						{
-							int32 Value;
-							if (!LexTryParseString<int32>(Value, *InNewText.ToString()) ||
-								Value < 1)
-							{
-								OutErrorMessage = LOCTEXT("BulkAddPatchesBadString", "Numeric values > 0 only");
-								return false;
-							}
-							return true;
-						})
-					.OnTextCommitted_Lambda([this](const FText& Text, ETextCommit::Type CommitType)
-						{
-							int32 Value;
-							if (LexTryParseString<int32>(Value, *Text.ToString()) &&
-								Value > 0)
-							{
-								constexpr int32 MaxNumPatchesToBulkAdd = 32768;
-								NumFixturePatchesToAdd = FMath::Min(Value, MaxNumPatchesToBulkAdd);
-							}
-						})
+					.Text(LOCTEXT("UniverseDotChannelLabel", "Universe.Channel"))
+					.Font(FAppStyle::GetFontStyle(TEXT("PropertyWindow.NormalFont")))
 				]
-			]
 
-			// Add fixture patches button
-			+ SGridPanel::Slot(1, 3)
-			.Padding(4.f)
-			[
-				SNew(SBox)
-				.HAlign(HAlign_Right)
-				.MinDesiredWidth(120.f)
+				// Select universe and channel 
+				+ SGridPanel::Slot(1, 2)
+				.Padding(4.f)
+				.VAlign(VAlign_Center)
+				[
+					MakeUniverseChannelSelectWidget()
+				]
+
+				// Num Fixture Patches Label
+				+ SGridPanel::Slot(0, 3)
+				.Padding(4.f)
+				.VAlign(VAlign_Center)
 				[
-					SNew(SButton)
+					SNew(STextBlock)
 					.IsEnabled(this, &SAddFixturePatchMenu::HasValidFixtureTypeAndMode)
-					.ContentPadding(FMargin(4.0f, 4.0f))
-					.ButtonStyle(FAppStyle::Get(), "FlatButton.Success")
-					.ForegroundColor(FLinearColor::White)
-					.OnClicked(this, &SAddFixturePatchMenu::OnAddFixturePatchButtonClicked)
-					[
-						SNew(STextBlock)
-						.Text_Lambda([this]()
-						{
-							return FText::Format(LOCTEXT("AddFixturePatchButtonText", "Add Fixture {0}|plural(one=Patch, other=Patches)"), NumFixturePatchesToAdd);
-						})
-						.Font(FAppStyle::GetFontStyle(TEXT("PropertyWindow.BoldFont")))
-					]
+					.Text(LOCTEXT("NumPatchesLabel", "Num Patches"))
+					.Font(FAppStyle::GetFontStyle(TEXT("PropertyWindow.NormalFont")))
+				]
+
+				// Num Fixture Patches editable Text Box
+				+ SGridPanel::Slot(1, 3)
+				.Padding(4.f)
+				[
+					MakeNumFixturePatchesEditableTextBox()
+				]
+
+				// Add fixture patches button
+				+ SGridPanel::Slot(1, 4)
+				.Padding(4.f)
+				[
+					MakeAddFixturePatchesButton()
 				]
 			]
 		];
 
 		RequestRefresh();
+
+		UDMXLibrary::GetOnEntitiesAdded().AddSP(this, &SAddFixturePatchMenu::OnEntityAddedOrRemoved);
+		UDMXLibrary::GetOnEntitiesRemoved().AddSP(this, &SAddFixturePatchMenu::OnEntityAddedOrRemoved);
 	}
 
-	void SAddFixturePatchMenu::ForceRefresh()
+	void SAddFixturePatchMenu::RequestRefresh()
 	{
 		RequestRefreshModeComboBoxTimerHandle.Invalidate();
 
-		UDMXLibrary* DMXLibrary = WeakDMXEditor.IsValid() ? WeakDMXEditor.Pin()->GetDMXLibrary() : nullptr;
-		if (!DMXLibrary)
+		if (!RequestRefreshModeComboBoxTimerHandle.IsValid())
 		{
-			return;
+			RequestRefreshModeComboBoxTimerHandle = GEditor->GetTimerManager()->SetTimerForNextTick(FTimerDelegate::CreateSP(this, &SAddFixturePatchMenu::ForceRefresh));
 		}
-		const TArray<UDMXEntityFixtureType*> FixtureTypes = DMXLibrary->GetEntitiesTypeCast<UDMXEntityFixtureType>();
+	}
 
-		// Init 
-		UDMXAddFixturePatchMenuData* MenuData = GetMutableDefault<UDMXAddFixturePatchMenuData>();
-		UDMXEntityFixtureType* const* SelectedFixtureTypePtr = Algo::FindByPredicate(FixtureTypes, [MenuData](const UDMXEntityFixtureType* FixtureType)
-			{
-				return FixtureType == MenuData->SoftFixtureType;
-			});
-		if (!SelectedFixtureTypePtr)
-		{
-			SelectedFixtureTypePtr = FixtureTypes.IsEmpty() ? nullptr : &FixtureTypes[0];
-		}
-		WeakFixtureType = SelectedFixtureTypePtr ? *SelectedFixtureTypePtr : nullptr;
-		ActiveModeIndex = WeakFixtureType.IsValid() && WeakFixtureType->Modes.IsValidIndex(ActiveModeIndex) ? ActiveModeIndex : 0;
+	void SAddFixturePatchMenu::SaveConfig()
+	{
+		UDMXAddFixturePatchMenuData* Data = GetMutableDefault<UDMXAddFixturePatchMenuData>();
+		Data->SoftFixtureType = WeakFixtureType.Get();
+		Data->ActiveModeIndex = ActiveModeIndex;
+		Data->SaveConfig();
+	}
 
-		ModeSources.Reset();
-		if (UDMXEntityFixtureType* FixtureType = WeakFixtureType.Get())
-		{
-			for (int32 ModeIndex = 0; ModeIndex < FixtureType->Modes.Num(); ModeIndex++)
-			{
-				ModeSources.Add(MakeShared<uint32>(ModeIndex));
-			}
-		}
+	TSharedRef<SWidget> SAddFixturePatchMenu::MakeFixtureTypeSelectWidget()
+	{
+		return
+			SNew(SBox)
+			.MinDesiredWidth(100.f)
+			[
+				SAssignNew(FixtureTypeSelector, SDMXEntityPickerButton<UDMXEntityFixtureType>)
+				.DMXEditor(WeakDMXEditor)
+				.CurrentEntity_Lambda([this] { return WeakFixtureType.Get(); })
+				.OnEntitySelected(this, &SAddFixturePatchMenu::OnFixtureTypeSelected)
+			];
+	}
 
-		ModeComboBox->RefreshOptions();
+	TSharedRef<SWidget> SAddFixturePatchMenu::MakeModeSelectWidget()
+	{
+		return
+			SAssignNew(ModeComboBox, SComboBox<TSharedPtr<uint32>>)
+			.IsEnabled(this, &SAddFixturePatchMenu::HasValidFixtureTypeAndMode)
+			.OptionsSource(&ModeSources)
+			.OnGenerateWidget(this, &SAddFixturePatchMenu::GenerateModeComboBoxEntry)
+			.OnSelectionChanged(this, &SAddFixturePatchMenu::OnModeSelected)
+			.InitiallySelectedItem(0)
+			[
+				SNew(STextBlock)
+				.MinDesiredWidth(50.0f)
+				.Text(this, &SAddFixturePatchMenu::GetActiveModeText)
+				.Font(FAppStyle::GetFontStyle(TEXT("PropertyWindow.NormalFont")))
+			];
+	}
+
+	TSharedRef<SWidget> SAddFixturePatchMenu::MakeUniverseChannelSelectWidget()
+	{
+		return
+			SNew(SBox)
+			.HAlign(HAlign_Left)
+			[
+				SNew(SEditableTextBox)
+				.IsEnabled(this, &SAddFixturePatchMenu::HasValidFixtureTypeAndMode)
+				.MinDesiredWidth(60.f)
+				.Font(FAppStyle::GetFontStyle(TEXT("PropertyWindow.NormalFont")))
+				.SelectAllTextWhenFocused(true)
+				.ClearKeyboardFocusOnCommit(true)
+				.Text(this, &SAddFixturePatchMenu::GetUniverseChannelText)
+				.OnVerifyTextChanged(this, &SAddFixturePatchMenu::OnVerifyUniverseChannelText)
+				.OnTextCommitted(this, &SAddFixturePatchMenu::OnUniverseChannelTextCommitted)
+			];
+	}
+
+	TSharedRef<SWidget> SAddFixturePatchMenu::MakeNumFixturePatchesEditableTextBox()
+	{
+		return
+			SNew(SBox)
+			.HAlign(HAlign_Left)
+			[
+				SNew(SEditableTextBox)
+				.IsEnabled(this, &SAddFixturePatchMenu::HasValidFixtureTypeAndMode)
+				.MinDesiredWidth(60.f)
+				.Font(FAppStyle::GetFontStyle(TEXT("PropertyWindow.NormalFont")))
+				.SelectAllTextWhenFocused(true)
+				.ClearKeyboardFocusOnCommit(true)
+				.Text_Lambda([this]()
+					{
+						return FText::FromString(FString::FromInt(NumFixturePatchesToAdd));
+					})
+				.OnTextCommitted_Lambda([this](const FText& Text, ETextCommit::Type CommitType)
+					{
+						int32 Value;
+						if (LexTryParseString<int32>(Value, *Text.ToString()) &&
+							Value > 0)
+						{
+							constexpr int32 MaxNumFixturePatchesToAdd = 16384;
+							NumFixturePatchesToAdd = FMath::Clamp(Value, 1, MaxNumFixturePatchesToAdd);
+						}
+
+						// Add fixture patches if enter is pressed
+						if (CommitType == ETextCommit::OnEnter)
+						{
+							OnAddFixturePatchButtonClicked();
+						}
+					})
+			];
+	}
+
+	TSharedRef<SWidget> SAddFixturePatchMenu::MakeAddFixturePatchesButton()
+	{
+		return
+			SNew(SBox)
+			.HAlign(HAlign_Right)
+			.MinDesiredWidth(120.f)
+			[
+				SNew(SButton)
+				.IsEnabled(this, &SAddFixturePatchMenu::HasValidFixtureTypeAndMode)
+				.ContentPadding(FMargin(4.0f, 4.0f))
+				.ButtonStyle(FAppStyle::Get(), "FlatButton.Success")
+				.ForegroundColor(FLinearColor::White)
+				.OnClicked(this, &SAddFixturePatchMenu::OnAddFixturePatchButtonClicked)
+				[
+					SNew(STextBlock)
+					.Text_Lambda([this]()
+						{
+							return FText::Format(LOCTEXT("AddFixturePatchButtonText", "Add Fixture {0}|plural(one=Patch, other=Patches)"), NumFixturePatchesToAdd);
+						})
+					.Font(FAppStyle::GetFontStyle(TEXT("PropertyWindow.BoldFont")))
+				]
+			];
 	}
 
 	TSharedRef<SWidget> SAddFixturePatchMenu::GenerateModeComboBoxEntry(const TSharedPtr<uint32> InModeIndex) const
@@ -223,7 +287,7 @@ namespace UE::DMXEditor::FixturePatchEditor
 			SNew(STextBlock)
 			.MinDesiredWidth(50.0f)
 			.Text_Lambda([&Modes, ModeIndex = *InModeIndex, this]()
-				{	
+				{
 					const UDMXAddFixturePatchMenuData* MenuData = GetDefault<UDMXAddFixturePatchMenuData>();
 					if (Modes.IsValidIndex(ModeIndex))
 					{
@@ -237,12 +301,35 @@ namespace UE::DMXEditor::FixturePatchEditor
 			.Font(FAppStyle::GetFontStyle(TEXT("PropertyWindow.NormalFont")));
 	}
 
-	void SAddFixturePatchMenu::RequestRefresh()
-	{ 
-		if (!RequestRefreshModeComboBoxTimerHandle.IsValid())
+	void SAddFixturePatchMenu::ForceRefresh()
+	{
+		RequestRefreshModeComboBoxTimerHandle.Invalidate();
+
+
+		if (!WeakFixtureType.IsValid() && WeakDMXEditor.IsValid())
 		{
-			RequestRefreshModeComboBoxTimerHandle = GEditor->GetTimerManager()->SetTimerForNextTick(FTimerDelegate::CreateSP(this, &SAddFixturePatchMenu::ForceRefresh));
+			if (UDMXLibrary* DMXLibrary = WeakDMXEditor.Pin()->GetDMXLibrary())
+			{
+				const TArray<UDMXEntityFixtureType*> FixtureTypes = DMXLibrary->GetEntitiesTypeCast<UDMXEntityFixtureType>();
+				WeakFixtureType = FixtureTypes.IsEmpty() ? nullptr : FixtureTypes[0];
+			}
+		}
+
+		ModeSources.Reset();
+		if (UDMXEntityFixtureType* FixtureType = WeakFixtureType.Get())
+		{
+			for (int32 ModeIndex = 0; ModeIndex < FixtureType->Modes.Num(); ModeIndex++)
+			{
+				ModeSources.Add(MakeShared<uint32>(ModeIndex));
+			}
 		}
+
+		ModeComboBox->RefreshOptions();
+	}
+
+	void SAddFixturePatchMenu::OnEntityAddedOrRemoved(UDMXLibrary* DMXLibrary, TArray<UDMXEntity*> Entities)
+	{
+		RequestRefresh();
 	}
 
 	void SAddFixturePatchMenu::OnFixtureTypeSelected(UDMXEntity* InSelectedFixtureType)
@@ -257,6 +344,73 @@ namespace UE::DMXEditor::FixturePatchEditor
 		ActiveModeIndex = InSelectedMode.IsValid() ? *InSelectedMode : 0;
 	}
 
+	FText SAddFixturePatchMenu::GetUniverseChannelText() const
+	{
+		if (!Universe.IsSet() && !Channel.IsSet())
+		{
+			const int32 SelectedUniverse = SharedData->GetSelectedUniverse();
+			return FText::FromString(FString::FromInt(SelectedUniverse) + TEXT(".") + FString::FromInt(1));
+		}
+		else
+		{
+			const FString UnivereString = Universe.IsSet() ? FString::FromInt(Universe.GetValue()) : TEXT("");
+			const FString ChannelString = Channel.IsSet() ? FString::FromInt(Channel.GetValue()) : TEXT("");
+			const FString Separator = Universe.IsSet() && Channel.IsSet() ? TEXT(".") : TEXT("");
+
+			return FText::FromString(UnivereString + Separator + ChannelString);
+		}
+	}
+
+	bool SAddFixturePatchMenu::OnVerifyUniverseChannelText(const FText& InNewText, FText& OutErrorMessage)
+	{
+		const FString String = InNewText.ToString();
+		int32 IndexOfDot = INDEX_NONE;
+		InNewText.ToString().FindChar('.', IndexOfDot);
+		if (IndexOfDot == INDEX_NONE && !String.IsNumeric())
+		{
+			OutErrorMessage = LOCTEXT("InvalidUniverseInfo", "Universe, for example '3' or Universe.Channel, for example '4.5'");
+			return false;
+		}
+		return true;
+	}
+
+	void SAddFixturePatchMenu::OnUniverseChannelTextCommitted(const FText& Text, ETextCommit::Type CommitType)
+	{
+		static const TCHAR* Delimiter[] = { TEXT("."), TEXT(","), TEXT(":"), TEXT(" ") };
+		TArray<FString> Substrings;
+		constexpr bool bCullEmpty = true;
+		Text.ToString().ParseIntoArray(Substrings, Delimiter, 4, bCullEmpty);
+
+		int32 NewUniverse;
+		if (Substrings.IsValidIndex(0) &&
+			LexTryParseString(NewUniverse, *Substrings[0]))
+		{
+			Universe = NewUniverse;
+			SharedData->SelectUniverse(NewUniverse);
+		}
+		else
+		{
+			Universe.Reset();
+		}
+
+		int32 NewChannel;
+		if (Substrings.IsValidIndex(1) &&
+			LexTryParseString(NewChannel, *Substrings[1]))
+		{
+			Channel = NewChannel;
+		}
+		else
+		{
+			Channel.Reset();
+		}
+
+		// Add fixture patches if enter is pressed
+		if (CommitType == ETextCommit::OnEnter)
+		{
+			OnAddFixturePatchButtonClicked();
+		}
+	}
+
 	FReply SAddFixturePatchMenu::OnAddFixturePatchButtonClicked()
 	{
 		FSlateApplication::Get().DismissAllMenus();
@@ -272,11 +426,7 @@ namespace UE::DMXEditor::FixturePatchEditor
 		{
 			return FReply::Handled();
 		}
-		const TSharedPtr<FDMXFixturePatchSharedData> SharedData = DMXEditor->GetFixturePatchSharedData();
-		if (!SharedData.IsValid())
-		{
-			return FReply::Handled();
-		}
+
 		UDMXLibrary* DMXLibrary = DMXEditor->GetDMXLibrary();
 		if (!DMXLibrary)
 		{
@@ -284,40 +434,40 @@ namespace UE::DMXEditor::FixturePatchEditor
 		}
 
 		// Ensure valid mode
-		UDMXAddFixturePatchMenuData* MenuData = GetMutableDefault<UDMXAddFixturePatchMenuData>();
-		if (!ensureMsgf(FixtureType->Modes.IsValidIndex(MenuData->ActiveModeIndex), TEXT("Cannot apply mode. Mode index is invalid.")))
+		if (!ensureMsgf(FixtureType->Modes.IsValidIndex(ActiveModeIndex), TEXT("Cannot apply mode. Mode index is invalid.")))
 		{
-			MenuData->ActiveModeIndex = 0;
+			return FReply::Handled();
 		}
 
-		// Create a new fixture patches
+		// Create new fixture patches
 		const FScopedTransaction CreateFixturePatchTransaction(LOCTEXT("CreateFixturePatchTransaction", "Create Fixture Patch"));
 		DMXLibrary->PreEditChange(UDMXLibrary::StaticClass()->FindPropertyByName(UDMXLibrary::GetEntitiesPropertyName()));
 
+		const int32 PatchToUniverse = Universe.IsSet() ? Universe.GetValue() : SharedData->GetSelectedUniverse();
+
 		TArray<UDMXEntityFixturePatch*> NewFixturePatches;
 		NewFixturePatches.Reserve(NumFixturePatchesToAdd);
 		for (uint32 iNumFixturePatchesAdded = 0; iNumFixturePatchesAdded < NumFixturePatchesToAdd; iNumFixturePatchesAdded++)
 		{
 			FDMXEntityFixturePatchConstructionParams FixturePatchConstructionParams;
 			FixturePatchConstructionParams.FixtureTypeRef = FDMXEntityFixtureTypeRef(FixtureType);
-			FixturePatchConstructionParams.ActiveMode = MenuData->ActiveModeIndex;
-			FixturePatchConstructionParams.UniverseID = SharedData->GetSelectedUniverse();
-			FixturePatchConstructionParams.StartingAddress = 1;
+			FixturePatchConstructionParams.ActiveMode = ActiveModeIndex;
+			FixturePatchConstructionParams.UniverseID = PatchToUniverse;
+			FixturePatchConstructionParams.StartingAddress = Channel.IsSet() ? Channel.GetValue() : 1;
 
 			constexpr bool bMarkLibraryDirty = false;
 			UDMXEntityFixturePatch* NewFixturePatch = UDMXEntityFixturePatch::CreateFixturePatchInLibrary(FixturePatchConstructionParams, FixtureType->Name, bMarkLibraryDirty);
 			NewFixturePatches.Add(NewFixturePatch);
 		}
 
-		// Align and Auto assign to the currently selected universe
+		// Align
 		using namespace UE::DMXEditor::AutoAssign;
 		FAutoAssignUtility::Align(NewFixturePatches);
-		const int32 AssignedToUniverse = FAutoAssignUtility::AutoAssign(EAutoAssignMode::SelectedUniverse, DMXEditor.ToSharedRef(), NewFixturePatches);
 
 		DMXLibrary->PostEditChange();
 
 		// Select universe and new fixture patches
-		SharedData->SelectUniverse(AssignedToUniverse);
+		SharedData->SelectUniverse(PatchToUniverse);
 
 		TArray<TWeakObjectPtr<UDMXEntityFixturePatch>> NewWeakFixturePatches;
 		NewWeakFixturePatches.Reserve(NumFixturePatchesToAdd);
diff --git a/Engine/Plugins/VirtualProduction/DMX/DMXEngine/Source/DMXEditor/Private/Widgets/FixturePatch/SAddFixturePatchMenu.h b/Engine/Plugins/VirtualProduction/DMX/DMXEngine/Source/DMXEditor/Private/Widgets/FixturePatch/SAddFixturePatchMenu.h
--- a/Engine/Plugins/VirtualProduction/DMX/DMXEngine/Source/DMXEditor/Private/Widgets/FixturePatch/SAddFixturePatchMenu.h
+++ b/Engine/Plugins/VirtualProduction/DMX/DMXEngine/Source/DMXEditor/Private/Widgets/FixturePatch/SAddFixturePatchMenu.h
@@ -6,12 +6,15 @@
 #include "Widgets/SCompoundWidget.h"
 
 class FDMXEditor;
+class FDMXFixturePatchSharedData;
 class FReply;
-template <typename OptionType> class SComboBox; 
+template <typename OptionType> class SComboBox;
 template <typename EntityType> class SDMXEntityPickerButton;
+class SEditableTextBox;
 class UDMXAddFixturePatchMenuData;
 class UDMXEntity;
 class UDMXEntityFixtureType;
+class UDMXLibrary;
 
 
 namespace UE::DMXEditor::FixturePatchEditor
@@ -34,18 +37,48 @@ namespace UE::DMXEditor::FixturePatchEditor
 		void RequestRefresh();
 
 	private:
+		/** Saves state to config */
+		void SaveConfig();
+
+		/** Creates a widget to select a fixture type */
+		TSharedRef<SWidget> MakeFixtureTypeSelectWidget();
+
+		/** Creates a widget to select a mode */
+		TSharedRef<SWidget> MakeModeSelectWidget();
+
+		/** Creates an widget to select the Universe and Channel */
+		TSharedRef<SWidget> MakeUniverseChannelSelectWidget();
+
+		/** Creates an editable text box to specify the number of patches */
+		TSharedRef<SWidget> MakeNumFixturePatchesEditableTextBox();
+
+		/** Creates a button that adds the fixture patches when clicked */
+		TSharedRef<SWidget> MakeAddFixturePatchesButton();
+
 		/** Generates an entry in the mode combo box */
 		TSharedRef<SWidget> GenerateModeComboBoxEntry(const TSharedPtr<uint32> InMode) const;
 
 		/** Refreshes the mode combo box */
 		void ForceRefresh();
 
+		/** Called when an entity was added or removed from the DMX Library */
+		void OnEntityAddedOrRemoved(UDMXLibrary* DMXLibrary, TArray<UDMXEntity*> Entities);
+
 		/** Called when a fixture type was selected */
 		void OnFixtureTypeSelected(UDMXEntity* InSelectedFixtureType);
 
 		/** Called when a mode was selected */
 		void OnModeSelected(TSharedPtr<uint32> InSelectedMode, ESelectInfo::Type SelectInfo);
 
+		/** Returns the current universe channel text */
+		FText GetUniverseChannelText() const;
+
+		/** Called to verify the universe channel text */
+		bool OnVerifyUniverseChannelText(const FText& InNewText, FText& OutErrorMessage);
+
+		/** Called when the universe channel text was committed */
+		void OnUniverseChannelTextCommitted(const FText& Text, ETextCommit::Type CommitType);
+
 		/** Called when the 'Add' button was clicked */
 		FReply OnAddFixturePatchButtonClicked();
 
@@ -61,6 +94,12 @@ namespace UE::DMXEditor::FixturePatchEditor
 		/** The number of fixture patches to add */
 		int32 ActiveModeIndex = 0;
 
+		/** Universe where the patches will be added */
+		TOptional<int32> Universe;
+
+		/** Channel where the patches will be added */
+		TOptional<int32> Channel; 
+
 		/** The number of fixture patches to add */
 		uint32 NumFixturePatchesToAdd = 1;
 
@@ -70,12 +109,18 @@ namespace UE::DMXEditor::FixturePatchEditor
 		/** Widget to select the fixture type */
 		TSharedPtr<SDMXEntityPickerButton<UDMXEntityFixtureType>> FixtureTypeSelector;
 
+		/** Text boxt to set the channel where the patch will be spawned*/
+		TSharedPtr<SEditableTextBox> ChannelSelectWidget;
+
 		/** Sources for the mode combo box */
 		TArray<TSharedPtr<uint32>> ModeSources;
 
 		/** Combo box to selected the mode */
 		TSharedPtr<SComboBox<TSharedPtr<uint32>>> ModeComboBox;
 
+		/** Shared data for fixture patches */
+		TSharedPtr<FDMXFixturePatchSharedData> SharedData;
+
 		/** The DMX editor this widget displays */
 		TWeakPtr<FDMXEditor> WeakDMXEditor;
 	};
diff --git a/Engine/Plugins/VirtualProduction/DMX/DMXEngine/Source/DMXEditor/Private/Widgets/FixturePatch/SDMXFixturePatcher.cpp b/Engine/Plugins/VirtualProduction/DMX/DMXEngine/Source/DMXEditor/Private/Widgets/FixturePatch/SDMXFixturePatcher.cpp
--- a/Engine/Plugins/VirtualProduction/DMX/DMXEngine/Source/DMXEditor/Private/Widgets/FixturePatch/SDMXFixturePatcher.cpp
+++ b/Engine/Plugins/VirtualProduction/DMX/DMXEngine/Source/DMXEditor/Private/Widgets/FixturePatch/SDMXFixturePatcher.cpp
@@ -351,57 +351,63 @@ FReply SDMXFixturePatcher::OnDropOntoChannel(int32 UniverseID, int32 ChannelID,
 			const int32 AbsoluteOffset = DraggedNodesToChannelOffsetPair.Value + PadOffset;
 
 			int32 NewUniverseID = UniverseID;
-			int64 AbsoluteDesiredStartingChannel = UniverseID * DMX_MAX_ADDRESS + ChannelID - AbsoluteOffset;
-
-			// Allow drag to another universe if multiple universes are displayed
-			int32 DesiredUniverse = AbsoluteDesiredStartingChannel / DMX_MAX_ADDRESS;
-			if (DesiredUniverse < 1)
-			{
-				DesiredUniverse = 1;
-				AbsoluteDesiredStartingChannel = ChannelID - AbsoluteOffset % DMX_MAX_ADDRESS;
-			}
-			else if (DesiredUniverse > DMX_MAX_UNIVERSE)
+			int32 DesiredStartingChannel = ChannelID;
+			if (IsUniverseSelectionEnabled())
 			{
-				DesiredUniverse = DMX_MAX_UNIVERSE;
-				AbsoluteDesiredStartingChannel = DMX_MAX_UNIVERSE * DMX_MAX_ADDRESS - ChannelSpan + 1;
+				// Stay within the current Universe if only one universe is displayed
+				if (FixturePatch->GetUniverseID() != UniverseID)
+				{
+					continue;
+				}
+
+				DesiredStartingChannel = ChannelID - AbsoluteOffset % DMX_MAX_ADDRESS;
 			}
 			else
 			{
-				NewUniverseID = DesiredUniverse;
-			}
+				int64 AbsoluteDesiredStartingChannel = UniverseID * DMX_MAX_ADDRESS + ChannelID - AbsoluteOffset;
 
-			// Remove from previous universe if needed
-			if (NewUniverseID != PatchedNode->GetUniverseID())
-			{
-				const TSharedPtr<SDMXPatchedUniverse>* OldUniverseWidgetPtr = PatchedUniversesByID.Find(PatchedNode->GetUniverseID());
-				if (OldUniverseWidgetPtr)
+				// Allow drag to another universe if multiple universes are displayed
+				int32 DesiredUniverse = AbsoluteDesiredStartingChannel / DMX_MAX_ADDRESS;
+				if (DesiredUniverse < 1)
 				{
-					(*OldUniverseWidgetPtr)->Remove(PatchedNode);
-					(*OldUniverseWidgetPtr)->RequestRefresh();
+					DesiredUniverse = 1;
+					AbsoluteDesiredStartingChannel = ChannelID - AbsoluteOffset % DMX_MAX_ADDRESS;
 				}
+				else if (DesiredUniverse > DMX_MAX_UNIVERSE)
+				{
+					DesiredUniverse = DMX_MAX_UNIVERSE;
+					AbsoluteDesiredStartingChannel = DMX_MAX_UNIVERSE * DMX_MAX_ADDRESS - ChannelSpan + 1;
+				}
+				else
+				{
+					NewUniverseID = DesiredUniverse;
+				}
+
+				// Remove from previous universe if needed
+				if (NewUniverseID != PatchedNode->GetUniverseID())
+				{
+					const TSharedPtr<SDMXPatchedUniverse>* OldUniverseWidgetPtr = PatchedUniversesByID.Find(PatchedNode->GetUniverseID());
+					if (OldUniverseWidgetPtr)
+					{
+						(*OldUniverseWidgetPtr)->Remove(PatchedNode);
+						(*OldUniverseWidgetPtr)->RequestRefresh();
+					}
+				}
+
+				DesiredStartingChannel = AbsoluteDesiredStartingChannel % DMX_MAX_ADDRESS;
 			}
 
-			const int32 DesiredStartingChannel = AbsoluteDesiredStartingChannel % DMX_MAX_ADDRESS;
 			const int32 NewStartingChannel = ClampStartingChannel(DesiredStartingChannel, ChannelSpan);
 			DraggedNodesToChannelOffsetPair.Key->SetAddresses(NewUniverseID, NewStartingChannel, ChannelSpan);
 
-			// Show newly patched universes if many universes are displayed and the user dragged on top of it 
 			const TSharedPtr<SDMXPatchedUniverse>* UniverseWidgetPtr = PatchedUniversesByID.Find(NewUniverseID);
-			if (!UniverseWidgetPtr && !IsUniverseSelectionEnabled())
-			{
-				if (!UniverseWidgetPtr)
-				{
-					AddUniverse(NewUniverseID);
-					UniverseWidgetPtr = &PatchedUniversesByID.FindChecked(NewUniverseID);
-				}
-				check(UniverseWidgetPtr);
-			}
-		
-			// Add the patch, if the universe is visible
-			if (UniverseWidgetPtr)
+			if (!UniverseWidgetPtr)
 			{
-				(*UniverseWidgetPtr)->FindOrAdd(DraggedNodesToChannelOffsetPair.Key);
+				AddUniverse(NewUniverseID);
+				UniverseWidgetPtr = &PatchedUniversesByID.FindChecked(NewUniverseID);
 			}
+			check(UniverseWidgetPtr);
+			(*UniverseWidgetPtr)->FindOrAdd(DraggedNodesToChannelOffsetPair.Key);
 
 			// Pad further patches so they can't be offset more than the currently dropped patch
 			PadOffset = FMath::Min(PadOffset, DesiredStartingChannel - NewStartingChannel);
diff --git a/Engine/Plugins/VirtualProduction/DMX/DMXEngine/Source/DMXEditor/Private/Widgets/FixturePatch/SDMXMVRFixtureList.cpp b/Engine/Plugins/VirtualProduction/DMX/DMXEngine/Source/DMXEditor/Private/Widgets/FixturePatch/SDMXMVRFixtureList.cpp
--- a/Engine/Plugins/VirtualProduction/DMX/DMXEngine/Source/DMXEditor/Private/Widgets/FixturePatch/SDMXMVRFixtureList.cpp
+++ b/Engine/Plugins/VirtualProduction/DMX/DMXEngine/Source/DMXEditor/Private/Widgets/FixturePatch/SDMXMVRFixtureList.cpp
@@ -541,7 +541,7 @@ SDMXMVRFixtureList::SDMXMVRFixtureList()
 
 SDMXMVRFixtureList::~SDMXMVRFixtureList()
 {
-	SaveHeaderRowSettings();
+	SaveSettings();
 }
 
 void SDMXMVRFixtureList::PostUndo(bool bSuccess)
@@ -667,7 +667,7 @@ void SDMXMVRFixtureList::RefreshList()
 {	
 	RequestListRefreshTimerHandle.Invalidate();
 
-	SaveHeaderRowSettings();
+	SaveSettings();
 
 	// Clear cached data
 	Rows.Reset();
@@ -1019,7 +1019,7 @@ TSharedRef<SHeaderRow> SDMXMVRFixtureList::GenerateHeaderRow()
 	);
 
 	// Restore user settings
-	RestoresHeaderRowSettings();
+	RestoreSettings();
 
 	return HeaderRow.ToSharedRef();
 }
@@ -1030,7 +1030,7 @@ void SDMXMVRFixtureList::SetKeyboardFocus()
 	FSlateApplication::Get().SetKeyboardFocus(AsShared());
 }
 
-void SDMXMVRFixtureList::SaveHeaderRowSettings()
+void SDMXMVRFixtureList::SaveSettings()
 {
 	UDMXEditorSettings* EditorSettings = GetMutableDefault<UDMXEditorSettings>();
 	if (HeaderRow.IsValid() && EditorSettings)
@@ -1055,11 +1055,13 @@ void SDMXMVRFixtureList::SaveHeaderRowSettings()
 			}
 		}
 
+		EditorSettings->MVRFixtureListSettings.SortByCollumnID = SortedByColumnID;
+
 		EditorSettings->SaveConfig();
 	}
 }
 
-void SDMXMVRFixtureList::RestoresHeaderRowSettings()
+void SDMXMVRFixtureList::RestoreSettings()
 {
 	if (const UDMXEditorSettings* EditorSettings = GetDefault<UDMXEditorSettings>())
 	{
@@ -1086,6 +1088,8 @@ void SDMXMVRFixtureList::RestoresHeaderRowSettings()
 		{
 			HeaderRow->SetColumnWidth(FDMXMVRFixtureListCollumnIDs::Patch, PatchColumnWidth);
 		}
+
+		SortedByColumnID = EditorSettings->MVRFixtureListSettings.SortByCollumnID;
 	}
 }
 
@@ -1326,8 +1330,9 @@ void SDMXMVRFixtureList::OnPasteItems()
 {
 	using namespace UE::DMX::SDMXMVRFixtureList::Private;
 
-	UDMXLibrary* DMXLibrary = WeakDMXEditor.IsValid() ? WeakDMXEditor.Pin()->GetDMXLibrary() : nullptr;
-	if (!DMXLibrary)
+	const TSharedPtr<FDMXEditor> DMXEditor = WeakDMXEditor.Pin();
+	UDMXLibrary* DMXLibrary = DMXEditor.IsValid() ? DMXEditor->GetDMXLibrary() : nullptr;
+	if (!DMXEditor.IsValid() || !DMXLibrary)
 	{
 		return;
 	}
@@ -1348,8 +1353,13 @@ void SDMXMVRFixtureList::OnPasteItems()
 			WeakPastedFixturePatches.Add(FixturePatch);
 		}
 
-		FixturePatchSharedData->SelectFixturePatches(WeakPastedFixturePatches);
+		// Assign
+		using namespace UE::DMXEditor::AutoAssign;
+		int32 AssignedToUniverse = FAutoAssignUtility::AutoAssign(EAutoAssignMode::SelectedUniverse, DMXEditor.ToSharedRef(), PastedFixturePatches);
 
+		FixturePatchSharedData->SelectUniverse(AssignedToUniverse);
+		FixturePatchSharedData->SelectFixturePatches(WeakPastedFixturePatches);
+	
 		RequestListRefresh();
 	}
 }
@@ -1372,6 +1382,7 @@ void SDMXMVRFixtureList::OnDuplicateItems()
 
 	const FText TransactionText = LOCTEXT("DuplicateFixturePatchesTransaction", "Duplicate Fixture Patches");
 	const FScopedTransaction PasteTransaction(TransactionText);
+	DMXLibrary->PreEditChange(nullptr);
 
 	TArray<TWeakObjectPtr<UDMXEntityFixturePatch>> SelectedWeakFixturePatches = FixturePatchSharedData->GetSelectedFixturePatches();
 	SelectedWeakFixturePatches.Remove(nullptr);
@@ -1433,7 +1444,6 @@ void SDMXMVRFixtureList::OnDeleteItems()
 		return;
 	}
 	
-	// Its safe to assume all patches are in the same Library - A Multi-Library Editor wouldn't make sense.
 	UDMXLibrary* DMXLibrary = SelectedItems[0]->GetDMXLibrary();
 	if (!DMXLibrary)
 	{
diff --git a/Engine/Plugins/VirtualProduction/DMX/DMXEngine/Source/DMXEditor/Private/Widgets/FixturePatch/SDMXMVRFixtureList.h b/Engine/Plugins/VirtualProduction/DMX/DMXEngine/Source/DMXEditor/Private/Widgets/FixturePatch/SDMXMVRFixtureList.h
--- a/Engine/Plugins/VirtualProduction/DMX/DMXEngine/Source/DMXEditor/Private/Widgets/FixturePatch/SDMXMVRFixtureList.h
+++ b/Engine/Plugins/VirtualProduction/DMX/DMXEngine/Source/DMXEditor/Private/Widgets/FixturePatch/SDMXMVRFixtureList.h
@@ -125,10 +125,10 @@ private:
 	TSharedRef<SHeaderRow> GenerateHeaderRow();
 
 	/** Saves how the user customized the header row in DMX Editor Settings */
-	void SaveHeaderRowSettings();
+	void SaveSettings();
 
 	/** Restores how the user customized the header row in DMX Editor Settings */
-	void RestoresHeaderRowSettings();
+	void RestoreSettings();
 
 	/** Returns the column sort mode for the list*/
 	EColumnSortMode::Type GetColumnSortMode(const FName ColumnId) const;
diff --git a/Engine/Plugins/VirtualProduction/DMX/DMXEngine/Source/DMXEditor/Private/Widgets/FixturePatch/SDMXPatchedUniverse.cpp b/Engine/Plugins/VirtualProduction/DMX/DMXEngine/Source/DMXEditor/Private/Widgets/FixturePatch/SDMXPatchedUniverse.cpp
--- a/Engine/Plugins/VirtualProduction/DMX/DMXEngine/Source/DMXEditor/Private/Widgets/FixturePatch/SDMXPatchedUniverse.cpp
+++ b/Engine/Plugins/VirtualProduction/DMX/DMXEngine/Source/DMXEditor/Private/Widgets/FixturePatch/SDMXPatchedUniverse.cpp
@@ -584,7 +584,7 @@ FReply SDMXPatchedUniverse::HandleOnMouseButtonDownOnChannel(uint32 Channel, con
 			!Selection.IsEmpty())
 		{
 			ShiftSelectAnchorPatch = ShiftSelectAnchorPatch.IsValid() ? ShiftSelectAnchorPatch : Selection[0];
-			
+
 			UDMXLibrary* DMXLibrary = ShiftSelectAnchorPatch->GetParentLibrary();
 			if (!DMXLibrary)
 			{
@@ -600,18 +600,18 @@ FReply SDMXPatchedUniverse::HandleOnMouseButtonDownOnChannel(uint32 Channel, con
 			TArray<TWeakObjectPtr<UDMXEntityFixturePatch>> ShiftSelection;
 			Algo::TransformIf(FixturePatches, ShiftSelection,
 				[&SelectedRange, Channel, this](UDMXEntityFixturePatch* FixturePatch)
-				{			
+				{
 					const int64 PatchAbsoluteStartingChannel = FixturePatch->GetUniverseID() * DMX_MAX_ADDRESS + FixturePatch->GetStartingChannel();
 					const int64 PatchAbsoluteEndingChannel = FixturePatch->GetUniverseID() * DMX_MAX_ADDRESS + FixturePatch->GetEndingChannel();
 					const TRange<int64> PatchRange(PatchAbsoluteStartingChannel, PatchAbsoluteEndingChannel + 1);
-					return 
+					return
 						SelectedRange.Overlaps(PatchRange);
 				},
 				[](UDMXEntityFixturePatch* FixturePatch)
 				{
 					return FixturePatch;
-				}
-			);
+				});
+
 			SharedData->SelectFixturePatches(ShiftSelection);
 
 			return FReply::Handled();
@@ -619,9 +619,9 @@ FReply SDMXPatchedUniverse::HandleOnMouseButtonDownOnChannel(uint32 Channel, con
 		ShiftSelectAnchorPatch.Reset();
 
 		// Ctrl-Select
-		if (FSlateApplication::Get().GetModifierKeys().IsControlDown() && 
+		if (FSlateApplication::Get().GetModifierKeys().IsControlDown() &&
 			ClickedPatch)
-		{		
+		{
 			if (const bool bUnselect = Selection.Num() > 1 && Selection.Contains(ClickedPatch))
 			{
 				TArray<TWeakObjectPtr<UDMXEntityFixturePatch>> WeakSelection = SharedData->GetSelectedFixturePatches();
@@ -1006,21 +1006,21 @@ UDMXEntityFixturePatch* SDMXPatchedUniverse::GetTopmostFixturePatchOnChannel(uin
 {
 	TArray<TSharedPtr<FDMXFixturePatchNode>> Nodes(PatchedNodes);
 	Nodes.RemoveAll([Channel](const TSharedPtr<FDMXFixturePatchNode>& Node)
+		{
+			if (UDMXEntityFixturePatch* FixturePatch = Node->GetFixturePatch().Get())
 			{
-				if (UDMXEntityFixturePatch* FixturePatch = Node->GetFixturePatch().Get())
-				{
-					const uint32 PatchStartingChannel = FixturePatch->GetStartingChannel();
-					const TRange<uint32> FixturePatchRange(PatchStartingChannel, PatchStartingChannel + FixturePatch->GetChannelSpan());
+				const uint32 PatchStartingChannel = FixturePatch->GetStartingChannel();
+				const TRange<uint32> FixturePatchRange(PatchStartingChannel, PatchStartingChannel + FixturePatch->GetChannelSpan());
 
-					return !FixturePatchRange.Contains(Channel);
-				}
-				return true;
-			});
+				return !FixturePatchRange.Contains(Channel);
+			}
+			return true;
+		});
 
 	Nodes.Sort([](const TSharedPtr<FDMXFixturePatchNode>& NodeA, const TSharedPtr<FDMXFixturePatchNode>& NodeB)
-			{
-				return NodeA->GetZOrder() > NodeB->GetZOrder();
-			});
+		{
+			return NodeA->GetZOrder() > NodeB->GetZOrder();
+		});
 
 	if (Nodes.IsEmpty())
 	{
@@ -1107,7 +1107,6 @@ TSharedRef<SWidget> SDMXPatchedUniverse::CreateContextMenu(int32 Channel)
 	constexpr bool bCloseWindowAfterMenuSelection = true;
 	FMenuBuilder MenuBuilder(bCloseWindowAfterMenuSelection, CommandList);
 
-	// Auto Assign Section
 	if (!SharedData->GetSelectedFixturePatches().IsEmpty())
 	{
 		using namespace UE::DMXEditor::AutoAssign;
@@ -1125,7 +1124,7 @@ TSharedRef<SWidget> SDMXPatchedUniverse::CreateContextMenu(int32 Channel)
 				FSlateIcon(),
 				FUIAction
 				(
-					FExecuteAction::CreateSP(this, &SDMXPatchedUniverse::AutoAssignFixturePatches, EAutoAssignMode::ChannelUnderMouse),
+					FExecuteAction::CreateSP(this, &SDMXPatchedUniverse::AutoAssignFixturePatches, EAutoAssignMode::UserDefinedChannel),
 					FCanExecuteAction::CreateLambda([this, SharedThis = AsShared()] { return !SharedData->GetSelectedFixturePatches().IsEmpty(); })
 				),
 				NAME_None,
diff --git a/Engine/Plugins/VirtualProduction/DMX/DMXEngine/Source/DMXEditor/Public/DMXEditorSettings.h b/Engine/Plugins/VirtualProduction/DMX/DMXEngine/Source/DMXEditor/Public/DMXEditorSettings.h
--- a/Engine/Plugins/VirtualProduction/DMX/DMXEngine/Source/DMXEditor/Public/DMXEditorSettings.h
+++ b/Engine/Plugins/VirtualProduction/DMX/DMXEngine/Source/DMXEditor/Public/DMXEditorSettings.h
@@ -5,6 +5,7 @@
 #include "CoreMinimal.h"
 #include "Containers/UnrealString.h"
 #include "Engine/DeveloperSettings.h"
+#include "Widgets/Views/SHeaderRow.h"
 
 #include "DMXEditorSettings.generated.h"
 
@@ -33,6 +34,9 @@ struct FDMXMVRFixtureListSettings
 	/** With of the Patch column */
 	UPROPERTY()
 	float PatchColumnWidth = 0.f;
+
+	UPROPERTY()
+	FName SortByCollumnID = NAME_None;
 };
 
 UENUM()
