DMX - Fix Auto-Assign does not align

diff --git a/Engine/Plugins/VirtualProduction/DMX/DMXEngine/Source/DMXEditor/Private/FixturePatchAutoAssignUtility.cpp b/Engine/Plugins/VirtualProduction/DMX/DMXEngine/Source/DMXEditor/Private/FixturePatchAutoAssignUtility.cpp
--- a/Engine/Plugins/VirtualProduction/DMX/DMXEngine/Source/DMXEditor/Private/FixturePatchAutoAssignUtility.cpp
+++ b/Engine/Plugins/VirtualProduction/DMX/DMXEngine/Source/DMXEditor/Private/FixturePatchAutoAssignUtility.cpp
@@ -227,7 +227,6 @@ namespace UE::DMXEditor::AutoAssign::Private
 		{
 			return;
 		}
-
 		Algo::StableSortBy(ElementsToSpread, &FAutoAssignElement::GetLowerBoundValue);
 
 		// If the largest patch would exceed universe size, offset the patches so the largest patch fits into the universe as well
@@ -270,6 +269,7 @@ namespace UE::DMXEditor::AutoAssign::Private
 			return;
 		}
 		Algo::StableSortBy(AutoAssignElements, &FAutoAssignElement::GetLowerBoundValue);
+		Align(AutoAssignElements);
 
 		TArray<int64> StartingChannels;
 		for (const TSharedRef<FAutoAssignElement>& FreeElement : FreeElements)
