DMX - Add patch menu no longer allows to patch to invalid universes / channels

diff --git a/Engine/Plugins/VirtualProduction/DMX/DMXEngine/Source/DMXEditor/Private/Widgets/FixturePatch/SAddFixturePatchMenu.cpp b/Engine/Plugins/VirtualProduction/DMX/DMXEngine/Source/DMXEditor/Private/Widgets/FixturePatch/SAddFixturePatchMenu.cpp
--- a/Engine/Plugins/VirtualProduction/DMX/DMXEngine/Source/DMXEditor/Private/Widgets/FixturePatch/SAddFixturePatchMenu.cpp
+++ b/Engine/Plugins/VirtualProduction/DMX/DMXEngine/Source/DMXEditor/Private/Widgets/FixturePatch/SAddFixturePatchMenu.cpp
@@ -202,14 +202,15 @@ namespace UE::DMXEditor::FixturePatchEditor
 			SNew(SBox)
 			.HAlign(HAlign_Left)
 			[
-				SNew(SEditableTextBox)
+				SAssignNew(UniverseChannelEditableTextBox, SEditableTextBox)
 				.IsEnabled(this, &SAddFixturePatchMenu::HasValidFixtureTypeAndMode)
 				.MinDesiredWidth(60.f)
 				.Font(FAppStyle::GetFontStyle(TEXT("PropertyWindow.NormalFont")))
 				.SelectAllTextWhenFocused(true)
 				.ClearKeyboardFocusOnCommit(true)
+				.RevertTextOnEscape(true)
 				.Text(this, &SAddFixturePatchMenu::GetUniverseChannelText)
-				.OnVerifyTextChanged(this, &SAddFixturePatchMenu::OnVerifyUniverseChannelText)
+				.OnTextChanged(this, &SAddFixturePatchMenu::OnUniverseChannelTextChanged)
 				.OnTextCommitted(this, &SAddFixturePatchMenu::OnUniverseChannelTextCommitted)
 			];
 	}
@@ -305,7 +306,6 @@ namespace UE::DMXEditor::FixturePatchEditor
 	{
 		RequestRefreshModeComboBoxTimerHandle.Invalidate();
 
-
 		if (!WeakFixtureType.IsValid() && WeakDMXEditor.IsValid())
 		{
 			if (UDMXLibrary* DMXLibrary = WeakDMXEditor.Pin()->GetDMXLibrary())
@@ -361,21 +361,60 @@ namespace UE::DMXEditor::FixturePatchEditor
 		}
 	}
 
-	bool SAddFixturePatchMenu::OnVerifyUniverseChannelText(const FText& InNewText, FText& OutErrorMessage)
-	{
-		const FString String = InNewText.ToString();
-		int32 IndexOfDot = INDEX_NONE;
-		InNewText.ToString().FindChar('.', IndexOfDot);
-		if (IndexOfDot == INDEX_NONE && !String.IsNumeric())
+	void SAddFixturePatchMenu::OnUniverseChannelTextChanged(const FText& Text)
+	{		
+		if (!UniverseChannelEditableTextBox.IsValid())
+		{
+			return;
+		}
+
+		static const TCHAR* Delimiter[] = { TEXT("."), TEXT(","), TEXT(":"), TEXT(" ") };
+		TArray<FString> Substrings;
+		constexpr bool bCullEmpty = true;
+		Text.ToString().ParseIntoArray(Substrings, Delimiter, 4, bCullEmpty);
+
+		static const FText InvalidStringErrorMessage = LOCTEXT("InvalidUniverseString", "Must be in the form of 'Universe' or 'Universe.Channel'. E.g. '4', or '4.5'.");
+
+		FText ErrorMessage;
+		int32 ParsedUniverse;
+		if (Substrings.IsValidIndex(0) &&
+			LexTryParseString(ParsedUniverse, *Substrings[0]))
+		{
+			if (ParsedUniverse < 1 || ParsedUniverse > DMX_MAX_UNIVERSE)
+			{
+				ErrorMessage = FText::Format(LOCTEXT("InvalidUniverseValue", "Universe must be between 1 and {0}."), FText::FromString(FString::FromInt((int32)DMX_MAX_UNIVERSE)));
+			}
+		}
+		else if (Substrings.IsValidIndex(0))
 		{
-			OutErrorMessage = LOCTEXT("InvalidUniverseInfo", "Universe, for example '3' or Universe.Channel, for example '4.5'");
-			return false;
+			ErrorMessage = InvalidStringErrorMessage;
 		}
-		return true;
+
+		// Channel is optional, only test if Substrings[1] was parsed
+		int32 ParsedChannel;
+		if (Substrings.IsValidIndex(1) &&
+			LexTryParseString(ParsedChannel, *Substrings[1]))
+		{
+			if (ParsedChannel < 1 || ParsedChannel > DMX_MAX_ADDRESS)
+			{
+				ErrorMessage = FText::Format(LOCTEXT("InvalidChannelValue", "Channel must be between 1 and {0}."), FText::FromString(FString::FromInt((int32)DMX_MAX_ADDRESS)));
+			}
+		}
+		else if (Substrings.IsValidIndex(1))
+		{
+			ErrorMessage = InvalidStringErrorMessage;
+		}
+
+		UniverseChannelEditableTextBox->SetError(ErrorMessage);
 	}
 
 	void SAddFixturePatchMenu::OnUniverseChannelTextCommitted(const FText& Text, ETextCommit::Type CommitType)
 	{
+		if (!UniverseChannelEditableTextBox.IsValid())
+		{
+			return;
+		}
+
 		static const TCHAR* Delimiter[] = { TEXT("."), TEXT(","), TEXT(":"), TEXT(" ") };
 		TArray<FString> Substrings;
 		constexpr bool bCullEmpty = true;
@@ -385,7 +424,7 @@ namespace UE::DMXEditor::FixturePatchEditor
 		if (Substrings.IsValidIndex(0) &&
 			LexTryParseString(NewUniverse, *Substrings[0]))
 		{
-			Universe = NewUniverse;
+			Universe = FMath::Clamp(NewUniverse, 1, DMX_MAX_UNIVERSE);
 			SharedData->SelectUniverse(NewUniverse);
 		}
 		else
@@ -394,21 +433,30 @@ namespace UE::DMXEditor::FixturePatchEditor
 		}
 
 		int32 NewChannel;
-		if (Substrings.IsValidIndex(1) &&
+		if (Universe.IsSet() && // Only parse the channel if parsing universe was successful
+			Substrings.IsValidIndex(1) &&
 			LexTryParseString(NewChannel, *Substrings[1]))
 		{
-			Channel = NewChannel;
+			Channel = FMath::Clamp(NewChannel, 1, DMX_MAX_ADDRESS);
 		}
 		else
 		{
 			Channel.Reset();
 		}
 
+		if (UniverseChannelEditableTextBox.IsValid())
+		{
+			UniverseChannelEditableTextBox->SetError(FText::GetEmpty());
+		}
+
 		// Add fixture patches if enter is pressed
 		if (CommitType == ETextCommit::OnEnter)
 		{
 			OnAddFixturePatchButtonClicked();
 		}
+
+		// Never show an error, values are always mended.
+		UniverseChannelEditableTextBox->SetError(FText::GetEmpty());
 	}
 
 	FReply SAddFixturePatchMenu::OnAddFixturePatchButtonClicked()
diff --git a/Engine/Plugins/VirtualProduction/DMX/DMXEngine/Source/DMXEditor/Private/Widgets/FixturePatch/SAddFixturePatchMenu.h b/Engine/Plugins/VirtualProduction/DMX/DMXEngine/Source/DMXEditor/Private/Widgets/FixturePatch/SAddFixturePatchMenu.h
--- a/Engine/Plugins/VirtualProduction/DMX/DMXEngine/Source/DMXEditor/Private/Widgets/FixturePatch/SAddFixturePatchMenu.h
+++ b/Engine/Plugins/VirtualProduction/DMX/DMXEngine/Source/DMXEditor/Private/Widgets/FixturePatch/SAddFixturePatchMenu.h
@@ -33,7 +33,7 @@ namespace UE::DMXEditor::FixturePatchEditor
 		/** Constructs this widget */
 		void Construct(const FArguments& InArgs, TWeakPtr<FDMXEditor> InWeakDMXEditor);
 
-		/** Refreshes the mode combo box on the next tick */
+		/** Refreshes the widget on the next tick */
 		void RequestRefresh();
 
 	private:
@@ -58,7 +58,7 @@ namespace UE::DMXEditor::FixturePatchEditor
 		/** Generates an entry in the mode combo box */
 		TSharedRef<SWidget> GenerateModeComboBoxEntry(const TSharedPtr<uint32> InMode) const;
 
-		/** Refreshes the mode combo box */
+		/** Refreshes the widget */
 		void ForceRefresh();
 
 		/** Called when an entity was added or removed from the DMX Library */
@@ -73,8 +73,8 @@ namespace UE::DMXEditor::FixturePatchEditor
 		/** Returns the current universe channel text */
 		FText GetUniverseChannelText() const;
 
-		/** Called to verify the universe channel text */
-		bool OnVerifyUniverseChannelText(const FText& InNewText, FText& OutErrorMessage);
+		/** Called when the universe channel text changed */
+		void OnUniverseChannelTextChanged(const FText& InText);
 
 		/** Called when the universe channel text was committed */
 		void OnUniverseChannelTextCommitted(const FText& Text, ETextCommit::Type CommitType);
@@ -109,8 +109,8 @@ namespace UE::DMXEditor::FixturePatchEditor
 		/** Widget to select the fixture type */
 		TSharedPtr<SDMXEntityPickerButton<UDMXEntityFixtureType>> FixtureTypeSelector;
 
-		/** Text boxt to set the channel where the patch will be spawned*/
-		TSharedPtr<SEditableTextBox> ChannelSelectWidget;
+		/** Text boxt to set the channel where the patch will be spawned */
+		TSharedPtr<SEditableTextBox> UniverseChannelEditableTextBox;
 
 		/** Sources for the mode combo box */
 		TArray<TSharedPtr<uint32>> ModeSources;
diff --git a/Engine/Plugins/VirtualProduction/DMX/DMXEngine/Source/DMXEditor/Private/Widgets/FixturePatch/SDMXFixturePatcher.cpp b/Engine/Plugins/VirtualProduction/DMX/DMXEngine/Source/DMXEditor/Private/Widgets/FixturePatch/SDMXFixturePatcher.cpp
--- a/Engine/Plugins/VirtualProduction/DMX/DMXEngine/Source/DMXEditor/Private/Widgets/FixturePatch/SDMXFixturePatcher.cpp
+++ b/Engine/Plugins/VirtualProduction/DMX/DMXEngine/Source/DMXEditor/Private/Widgets/FixturePatch/SDMXFixturePatcher.cpp
@@ -172,16 +172,22 @@ TSharedRef<SWidget> SDMXFixturePatcher::CreateToolbar()
 	ToolbarBuilder.BeginSection("UniverseSelection");
 	{
 		ToolbarBuilder.AddWidget(
-			SNew(STextBlock)
-			.MinDesiredWidth(75.0f)
-			.Font(FAppStyle::GetFontStyle("PropertyWindow.NormalFont"))
-			.TextStyle(FAppStyle::Get(), "DetailsView.CategoryTextStyle")
-			.IsEnabled(this, &SDMXFixturePatcher::IsUniverseSelectionEnabled)
-			.Text(LOCTEXT("UniverseSelectorLabel", "Local Universe"))
+			SNew(SBox)
+			.Padding(4.f, 0.f)
+			.VAlign(VAlign_Center)
+			[
+				SNew(STextBlock)
+				.MinDesiredWidth(75.0f)
+				.Font(FAppStyle::GetFontStyle("PropertyWindow.NormalFont"))
+				.TextStyle(FAppStyle::Get(), "DetailsView.CategoryTextStyle")
+				.IsEnabled(this, &SDMXFixturePatcher::IsUniverseSelectionEnabled)
+				.Text(LOCTEXT("UniverseSelectorLabel", "Local Universe"))
+			]
 		);
 
 		ToolbarBuilder.AddWidget(
 			SNew(SBox)
+			.Padding(4.f, 0.f)
 			.MinDesiredWidth(210.0f)
 			.MaxDesiredWidth(420.0f)
 			[
@@ -550,7 +556,7 @@ void SDMXFixturePatcher::OnUniverseSelectionChanged()
 
 void SDMXFixturePatcher::SelectUniverse(int32 NewUniverseID)
 {
-	UniverseToSetNextTick = NewUniverseID;
+	UniverseToSetNextTick = FMath::Clamp(NewUniverseID, 1, DMX_MAX_UNIVERSE);
 }
 
 int32 SDMXFixturePatcher::GetSelectedUniverse() const
diff --git a/Engine/Plugins/VirtualProduction/DMX/DMXEngine/Source/DMXEditor/Private/Widgets/FixturePatch/SDMXPatchedUniverse.cpp b/Engine/Plugins/VirtualProduction/DMX/DMXEngine/Source/DMXEditor/Private/Widgets/FixturePatch/SDMXPatchedUniverse.cpp
--- a/Engine/Plugins/VirtualProduction/DMX/DMXEngine/Source/DMXEditor/Private/Widgets/FixturePatch/SDMXPatchedUniverse.cpp
+++ b/Engine/Plugins/VirtualProduction/DMX/DMXEngine/Source/DMXEditor/Private/Widgets/FixturePatch/SDMXPatchedUniverse.cpp
@@ -95,6 +95,7 @@ void SDMXPatchedUniverse::Construct(const FArguments& InArgs)
 			+ SOverlay::Slot()
 			[
 				SAssignNew(UniverseName, SBorder)
+				.Padding(4.f, 2.f)
 				.VAlign(VAlign_Top)
 				.HAlign(HAlign_Fill)
 				.BorderBackgroundColor(FLinearColor(0.6f, 0.6f, 0.6f, 0.3f))
