DMX - Update monitors to display 'Local Universe' instead of just 'Universe' to be explicit about what's monitored

diff --git a/Engine/Plugins/VirtualProduction/DMX/DMXEngine/Source/DMXEditor/Private/Widgets/Monitors/SDMXActivityMonitor.cpp b/Engine/Plugins/VirtualProduction/DMX/DMXEngine/Source/DMXEditor/Private/Widgets/Monitors/SDMXActivityMonitor.cpp
--- a/Engine/Plugins/VirtualProduction/DMX/DMXEngine/Source/DMXEditor/Private/Widgets/Monitors/SDMXActivityMonitor.cpp
+++ b/Engine/Plugins/VirtualProduction/DMX/DMXEngine/Source/DMXEditor/Private/Widgets/Monitors/SDMXActivityMonitor.cpp
@@ -2,6 +2,8 @@
 
 #include "Widgets/Monitors/SDMXActivityMonitor.h"
 
+#include "Algo/RemoveIf.h"
+#include "Containers/UnrealString.h"
 #include "DMXEditorSettings.h"
 #include "DMXEditorStyle.h"
 #include "DMXEditorUtils.h"
@@ -10,18 +12,16 @@
 #include "IO/DMXOutputPort.h"
 #include "IO/DMXPortManager.h"
 #include "IO/DMXRawListener.h"
-#include "Widgets/Monitors/SDMXActivityInUniverse.h"
-#include "Widgets/Monitors/SDMXMonitorSourceSelector.h"
-
 #include "Styling/AppStyle.h"
 #include "SlateOptMacros.h"
-#include "Containers/UnrealString.h"
 #include "Widgets/Input/SButton.h"
 #include "Widgets/Input/SEditableTextBox.h"
 #include "Widgets/Layout/SScrollBox.h"
 #include "Widgets/Layout/SSeparator.h"
-#include "Widgets/Input/SSpinBox.h"
 #include "Widgets/Layout/SWrapBox.h"
+#include "Widgets/Input/SSpinBox.h"
+#include "Widgets/Monitors/SDMXActivityInUniverse.h"
+#include "Widgets/Monitors/SDMXMonitorSourceSelector.h"
 #include "Widgets/Views/SListView.h"
 
 
@@ -80,7 +80,8 @@ void SDMXActivityMonitor::Construct(const FArguments& InArgs)
 					.AutoWidth()
 					[
 						SNew(STextBlock)
-						.Text(LOCTEXT("UniversesLabel", "Universes"))
+						.Font(FAppStyle::GetFontStyle(TEXT("PropertyWindow.NormalFont")))
+						.Text(LOCTEXT("LocalUniversesLabel", "Monitored Local Universes"))
 					] 
 
 					+ SHorizontalBox::Slot()
@@ -151,7 +152,7 @@ void SDMXActivityMonitor::Construct(const FArguments& InArgs)
 				.VAlign(VAlign_Top)
 				[
 					SNew(STextBlock)
-					.Text(LOCTEXT("UniverseLabel", "Universe"))
+					.Text(LOCTEXT("LocalUniverseLabel", "Universe"))
 					.Font(FDMXEditorStyle::Get().GetFontStyle("DMXEditor.Font.InputUniverseHeader"))
 				]
 
@@ -196,7 +197,7 @@ void SDMXActivityMonitor::Construct(const FArguments& InArgs)
 
 	UpdateListenerRegistration();
 
-	FDMXPortManager::Get().OnPortsChanged.AddSP(this, &SDMXActivityMonitor::UpdateListenerRegistration);
+	FDMXPortManager::Get().OnPortsChanged.AddSP(this, &SDMXActivityMonitor::OnPortsChanged);
 }
 END_SLATE_FUNCTION_BUILD_OPTIMIZATION
 
@@ -243,13 +244,12 @@ void SDMXActivityMonitor::Tick(const FGeometry& AllottedGeometry, const double I
 		}
 	}
 
-	for (TTuple<int32, TArray<uint8, TFixedAllocator<DMX_UNIVERSE_SIZE>>>& UniverseToDataKvp : UniverseToDataMap)
+	for (const TTuple<int32, TArray<uint8, TFixedAllocator<DMX_UNIVERSE_SIZE>>>& UniverseToDataKvp : UniverseToDataMap)
 	{
 		const TSharedRef<SDMXActivityInUniverse>& ActivityWidget = GetOrCreateActivityWidget(UniverseToDataKvp.Key);
 		ActivityWidget->VisualizeBuffer(UniverseToDataKvp.Value);
 	}
 }
-
 TSharedRef<ITableRow> SDMXActivityMonitor::OnGenerateUniverseRow(TSharedPtr<SDMXActivityInUniverse> ActivityWidget, const TSharedRef<STableViewBase>& OwnerTable)
 {
 	return
@@ -304,7 +304,7 @@ void SDMXActivityMonitor::SaveMonitorSettings() const
 	DMXEditorSettings->SaveConfig();
 }
 
-TSharedRef<SDMXActivityInUniverse> SDMXActivityMonitor::GetOrCreateActivityWidget(uint16 UniverseID)
+TSharedRef<SDMXActivityInUniverse> SDMXActivityMonitor::GetOrCreateActivityWidget(int32 UniverseID)
 {
 	TSharedPtr<SDMXActivityInUniverse>* BufferViewPtr = UniverseListSource.FindByPredicate([&](const TSharedPtr<SDMXActivityInUniverse>& BufferViewCandidate) {
 			return BufferViewCandidate->GetUniverseID() == UniverseID;
@@ -408,9 +408,6 @@ void SDMXActivityMonitor::OnMaxUniverseIDValueCommitted(const FText& InNewText,
 
 void SDMXActivityMonitor::OnSourceSelected()
 {
-	FDMXEditorUtils::ClearAllDMXPortBuffers();
-	FDMXEditorUtils::ClearFixturePatchCachedData();
-
 	ClearDisplay();
 
 	SaveMonitorSettings();
@@ -418,6 +415,12 @@ void SDMXActivityMonitor::OnSourceSelected()
 	UpdateListenerRegistration();
 }
 
+void SDMXActivityMonitor::OnPortsChanged()
+{
+	ClearDisplay();
+
+	UpdateListenerRegistration();
+}
 
 FReply SDMXActivityMonitor::OnClearButtonClicked()
 {
diff --git a/Engine/Plugins/VirtualProduction/DMX/DMXEngine/Source/DMXEditor/Private/Widgets/Monitors/SDMXActivityMonitor.h b/Engine/Plugins/VirtualProduction/DMX/DMXEngine/Source/DMXEditor/Private/Widgets/Monitors/SDMXActivityMonitor.h
--- a/Engine/Plugins/VirtualProduction/DMX/DMXEngine/Source/DMXEditor/Private/Widgets/Monitors/SDMXActivityMonitor.h
+++ b/Engine/Plugins/VirtualProduction/DMX/DMXEngine/Source/DMXEditor/Private/Widgets/Monitors/SDMXActivityMonitor.h
@@ -61,8 +61,8 @@ private:
 	/** Saves settings of the monitor in config */
 	void SaveMonitorSettings() const;
 
-	/** Returns a buffer view for specified universe ID, creates one if it doesn't exist yet */
-	TSharedRef<SDMXActivityInUniverse> GetOrCreateActivityWidget(uint16 UniverseID);
+	/** Returns a view for specified universe ID, creates one if it doesn't exist yet */
+	TSharedRef<SDMXActivityInUniverse> GetOrCreateActivityWidget(int32 UniverseID);
 
 	/** Called when a Min Universe ID value was commited */
 	void OnMinUniverseIDValueCommitted(const FText& InNewText, ETextCommit::Type CommitType);
@@ -73,6 +73,9 @@ private:
 	/** Called when a Source was selected in the Monitor Source Selector */
 	void OnSourceSelected();
 
+	/** Called when ports changed */
+	void OnPortsChanged();
+
 	/** Called when the clear ui values button was clicked */
 	FReply OnClearButtonClicked();
 
