DMX - Fix patching window values overlap with "Monitor DMX Inputs" values

diff --git a/Engine/Plugins/VirtualProduction/DMX/DMXEngine/Source/DMXEditor/Private/Widgets/FixturePatch/SDMXFixturePatchFragment.cpp b/Engine/Plugins/VirtualProduction/DMX/DMXEngine/Source/DMXEditor/Private/Widgets/FixturePatch/SDMXFixturePatchFragment.cpp
--- a/Engine/Plugins/VirtualProduction/DMX/DMXEngine/Source/DMXEditor/Private/Widgets/FixturePatch/SDMXFixturePatchFragment.cpp
+++ b/Engine/Plugins/VirtualProduction/DMX/DMXEngine/Source/DMXEditor/Private/Widgets/FixturePatch/SDMXFixturePatchFragment.cpp
@@ -353,6 +353,11 @@ void SDMXFixturePatchFragment::Construct(const FArguments& InArgs, const TShared
 	}
 }
 
+void SDMXFixturePatchFragment::SetChannelVisibility(EVisibility NewVisibility)
+{
+	ChannelsHorizontalBox->SetVisibility(NewVisibility);
+}
+
 void SDMXFixturePatchFragment::InitializeBorderPadding()
 {
 	if (bIsHead && FixturePatchNodeGroup.Num() > 1)
diff --git a/Engine/Plugins/VirtualProduction/DMX/DMXEngine/Source/DMXEditor/Private/Widgets/FixturePatch/SDMXFixturePatchFragment.h b/Engine/Plugins/VirtualProduction/DMX/DMXEngine/Source/DMXEditor/Private/Widgets/FixturePatch/SDMXFixturePatchFragment.h
--- a/Engine/Plugins/VirtualProduction/DMX/DMXEngine/Source/DMXEditor/Private/Widgets/FixturePatch/SDMXFixturePatchFragment.h
+++ b/Engine/Plugins/VirtualProduction/DMX/DMXEngine/Source/DMXEditor/Private/Widgets/FixturePatch/SDMXFixturePatchFragment.h
@@ -58,6 +58,9 @@ public:
 
 	/** Sets if the fragment is conflicting with others */
 	void SetConflicting(bool bConflicting) { bIsConflicting = bConflicting; }
+	
+	/** Sets the visibility of the channels of the fragment */
+	void SetChannelVisibility(EVisibility NewVisibility);
 
 	/** Gets the fixture patch node this widget refers to */
 	const TSharedPtr<FDMXFixturePatchNode>& GetFixturePatchNode() const { return FixturePatchNode; }
diff --git a/Engine/Plugins/VirtualProduction/DMX/DMXEngine/Source/DMXEditor/Private/Widgets/FixturePatch/SDMXPatchedUniverse.cpp b/Engine/Plugins/VirtualProduction/DMX/DMXEngine/Source/DMXEditor/Private/Widgets/FixturePatch/SDMXPatchedUniverse.cpp
--- a/Engine/Plugins/VirtualProduction/DMX/DMXEngine/Source/DMXEditor/Private/Widgets/FixturePatch/SDMXPatchedUniverse.cpp
+++ b/Engine/Plugins/VirtualProduction/DMX/DMXEngine/Source/DMXEditor/Private/Widgets/FixturePatch/SDMXPatchedUniverse.cpp
@@ -163,8 +163,22 @@ void SDMXPatchedUniverse::SetMonitorInputsEnabled(bool bEnabled)
 {
 	bMonitorInputs = bEnabled;
 
-	if (!bEnabled)
+	if (bEnabled)
 	{
+		// Don't show the channel overlay of the fragment, so only the channel connectors display the values
+		for (const TSharedRef<SDMXFixturePatchFragment>& Fragment : FixturePatchFragmentWidgets)
+		{
+			Fragment->SetChannelVisibility(EVisibility::Hidden);
+		}
+	}
+	else
+	{
+		// Show the channel overlay of the fragment, so the fragments show their hovered state
+		for (const TSharedRef<SDMXFixturePatchFragment>& Fragment : FixturePatchFragmentWidgets)
+		{
+			Fragment->SetChannelVisibility(EVisibility::Visible);
+		}
+
 		ResetMonitor();
 	}
 }
