DMX - Fix root component of MVR Scene Actor, allowing for attachement as expected

diff --git a/Engine/Plugins/VirtualProduction/DMX/DMXEngine/Source/DMXRuntime/Private/MVR/DMXMVRSceneActor.cpp b/Engine/Plugins/VirtualProduction/DMX/DMXEngine/Source/DMXRuntime/Private/MVR/DMXMVRSceneActor.cpp
--- a/Engine/Plugins/VirtualProduction/DMX/DMXEngine/Source/DMXRuntime/Private/MVR/DMXMVRSceneActor.cpp
+++ b/Engine/Plugins/VirtualProduction/DMX/DMXEngine/Source/DMXRuntime/Private/MVR/DMXMVRSceneActor.cpp
@@ -48,6 +48,7 @@ ADMXMVRSceneActor::ADMXMVRSceneActor()
 
 	MVRSceneRoot = CreateDefaultSubobject<USceneComponent>("MVRSceneRoot");
 	SetRootComponent(MVRSceneRoot);
+	AddInstanceComponent(MVRSceneRoot);
 }
 
 ADMXMVRSceneActor::~ADMXMVRSceneActor()
