DMX - Fix ensure condition being met when patching modes without any functions

diff --git a/Engine/Plugins/VirtualProduction/DMX/DMXEngine/Source/DMXEditor/Private/FixturePatchAutoAssignUtility.cpp b/Engine/Plugins/VirtualProduction/DMX/DMXEngine/Source/DMXEditor/Private/FixturePatchAutoAssignUtility.cpp
--- a/Engine/Plugins/VirtualProduction/DMX/DMXEngine/Source/DMXEditor/Private/FixturePatchAutoAssignUtility.cpp
+++ b/Engine/Plugins/VirtualProduction/DMX/DMXEngine/Source/DMXEditor/Private/FixturePatchAutoAssignUtility.cpp
@@ -317,7 +317,10 @@ namespace UE::DMXEditor::AutoAssign::Private
 		const bool bValidElements = Algo::FindByPredicate(Elements, [](const TSharedRef<FAutoAssignElement>& Element)
 			{
 				// Find an invalid element
-				return Element->GetLowerBoundValue() > Element->GetUpperBoundValue() || Element->GetSize() < 1 || Element->GetSize() > DMX_UNIVERSE_SIZE;
+				return 
+					Element->GetLowerBoundValue() > Element->GetUpperBoundValue() || 
+					Element->GetSize() < 0 || 
+					Element->GetSize() > DMX_UNIVERSE_SIZE;
 			}) == nullptr;
 		if (!ensureMsgf(bValidElements, TEXT("Elements to assign are not in valid DMX range. Ignoring call.")))
 		{
