DMX library - Fix Monitor DMX input will stop working if a fixture patch is duplicated

diff --git a/Engine/Plugins/VirtualProduction/DMX/DMXEngine/Source/DMXEditor/Private/Widgets/FixturePatch/SDMXFixturePatcher.cpp b/Engine/Plugins/VirtualProduction/DMX/DMXEngine/Source/DMXEditor/Private/Widgets/FixturePatch/SDMXFixturePatcher.cpp
--- a/Engine/Plugins/VirtualProduction/DMX/DMXEngine/Source/DMXEditor/Private/Widgets/FixturePatch/SDMXFixturePatcher.cpp
+++ b/Engine/Plugins/VirtualProduction/DMX/DMXEngine/Source/DMXEditor/Private/Widgets/FixturePatch/SDMXFixturePatcher.cpp
@@ -99,9 +99,6 @@ void SDMXFixturePatcher::Construct(const FArguments& InArgs)
 	}		
 
 	ShowSelectedUniverse();
-
-	const UDMXEditorSettings* DMXEditorSettings = GetDefault<UDMXEditorSettings>();
-	SetDMXMonitorEnabled(DMXEditorSettings->FixturePatcherSettings.bMonitorEnabled);
 }
 END_SLATE_FUNCTION_BUILD_OPTIMIZATION
 
@@ -585,6 +582,10 @@ void SDMXFixturePatcher::ShowSelectedUniverse()
 		];
 
 	PatchedUniversesByID.Add(SelectedUniverseID, NewPatchedUniverse);
+
+	// Update monitor
+	UDMXEditorSettings* DMXEditorSettings = GetMutableDefault<UDMXEditorSettings>();
+	SetDMXMonitorEnabled(DMXEditorSettings->FixturePatcherSettings.bMonitorEnabled);
 }
 
 void SDMXFixturePatcher::ShowAllPatchedUniverses(bool bForceReconstructWidget)
@@ -659,6 +660,10 @@ void SDMXFixturePatcher::ShowAllPatchedUniverses(bool bForceReconstructWidget)
 			PatchedUniverseScrollBox->ScrollDescendantIntoView(PatchedUniversesByID[SelectedUniverse]);
 		}
 	}
+
+	// Update monitor
+	UDMXEditorSettings* DMXEditorSettings = GetMutableDefault<UDMXEditorSettings>();
+	SetDMXMonitorEnabled(DMXEditorSettings->FixturePatcherSettings.bMonitorEnabled);
 }
 
 void SDMXFixturePatcher::AddUniverse(int32 UniverseID)
@@ -766,7 +771,7 @@ FText SDMXFixturePatcher::GetTooltipText() const
 {
 	if (!HasAnyPorts())
 	{
-		return LOCTEXT("NoPorts", "No ports available. Please create add Ports to the DMX Library.");
+		return LOCTEXT("NoPortsTooltip", "No ports available. Please create ports in Project Settings -> Plugins -> DMX.");
 	}
 
 	return FText::GetEmpty();
