DMX - Fix changing categories of multiple fixture types in the DMX Library stalls the engine

diff --git a/Engine/Plugins/VirtualProduction/DMX/DMXEngine/Source/DMXEditor/Private/Customizations/DMXEntityFixturePatchDetails.cpp b/Engine/Plugins/VirtualProduction/DMX/DMXEngine/Source/DMXEditor/Private/Customizations/DMXEntityFixturePatchDetails.cpp
--- a/Engine/Plugins/VirtualProduction/DMX/DMXEngine/Source/DMXEditor/Private/Customizations/DMXEntityFixturePatchDetails.cpp
+++ b/Engine/Plugins/VirtualProduction/DMX/DMXEngine/Source/DMXEditor/Private/Customizations/DMXEntityFixturePatchDetails.cpp
@@ -120,7 +120,7 @@ TSharedRef<SWidget> FDMXEntityFixturePatchDetails::GenerateActiveModeWidget(cons
 void FDMXEntityFixturePatchDetails::OnParentFixtureTypeChanged(UDMXEntity* NewTemplate) const
 {
 	ParentFixtureTypeHandle->SetValue(Cast<UDMXEntityFixtureType>(NewTemplate));
-	PropertyUtilities->ForceRefresh();
+	PropertyUtilities->RequestRefresh();
 }
 
 void FDMXEntityFixturePatchDetails::OnFixtureTypeChanged(const UDMXEntityFixtureType* FixtureType)
@@ -138,7 +138,7 @@ void FDMXEntityFixturePatchDetails::OnFixtureTypeChanged(const UDMXEntityFixture
 			}
 		}
 	
-		PropertyUtilities->ForceRefresh();
+		PropertyUtilities->RequestRefresh();
 	}
 }
 
diff --git a/Engine/Plugins/VirtualProduction/DMX/DMXEngine/Source/DMXEditor/Private/Customizations/DMXFixtureCategoryCustomization.cpp b/Engine/Plugins/VirtualProduction/DMX/DMXEngine/Source/DMXEditor/Private/Customizations/DMXFixtureCategoryCustomization.cpp
--- a/Engine/Plugins/VirtualProduction/DMX/DMXEngine/Source/DMXEditor/Private/Customizations/DMXFixtureCategoryCustomization.cpp
+++ b/Engine/Plugins/VirtualProduction/DMX/DMXEngine/Source/DMXEditor/Private/Customizations/DMXFixtureCategoryCustomization.cpp
@@ -76,23 +76,16 @@ FName FDMXFixtureCategoryCustomization::GetValue() const
 
 void FDMXFixtureCategoryCustomization::SetValue(FName NewValue)
 {
-	FStructProperty* StructProperty = CastFieldChecked<FStructProperty>(StructPropertyHandle->GetProperty());
+	const TSharedPtr<IPropertyHandle> NameHandle = StructPropertyHandle->GetChildHandle(GET_MEMBER_NAME_CHECKED(FDMXFixtureCategory, Name));
 
-	TArray<void*> RawData;
-	StructPropertyHandle->AccessRawData(RawData);
-
-	for (void* SingleRawData : RawData)
+	// Avoid recursion
+	FName OldValue;
+	if (NameHandle->GetValue(OldValue) == FPropertyAccess::Success &&
+		OldValue != NewValue)
 	{
-		FDMXFixtureCategory* PreviousValue = reinterpret_cast<FDMXFixtureCategory*>(SingleRawData);
-		FDMXFixtureCategory NewCategoryName;
-		NewCategoryName.Name = NewValue;
-
-		// Export new value to text format that can be imported later
-		FString TextValue;
-		StructProperty->Struct->ExportText(TextValue, &NewCategoryName, PreviousValue, nullptr, EPropertyPortFlags::PPF_None, nullptr);
-
-		// Set values on edited property handle from exported text
-		ensure(StructPropertyHandle->SetValueFromFormattedString(TextValue, EPropertyValueSetFlags::DefaultFlags) == FPropertyAccess::Result::Success);
+		NameHandle->NotifyPreChange();
+		NameHandle->SetValue(NewValue);
+		NameHandle->NotifyPostChange(EPropertyChangeType::ValueSet);
 	}
 }
 
