DMX - Selected mode is now displayed correctly in 'Add Fixture Patch' menu in DMX Library

diff --git a/Engine/Plugins/VirtualProduction/DMX/DMXEngine/Source/DMXEditor/Private/Widgets/FixturePatch/SAddFixturePatchMenu.cpp b/Engine/Plugins/VirtualProduction/DMX/DMXEngine/Source/DMXEditor/Private/Widgets/FixturePatch/SAddFixturePatchMenu.cpp
--- a/Engine/Plugins/VirtualProduction/DMX/DMXEngine/Source/DMXEditor/Private/Widgets/FixturePatch/SAddFixturePatchMenu.cpp
+++ b/Engine/Plugins/VirtualProduction/DMX/DMXEngine/Source/DMXEditor/Private/Widgets/FixturePatch/SAddFixturePatchMenu.cpp
@@ -28,7 +28,8 @@ namespace UE::DMXEditor::FixturePatchEditor
 {
 	SAddFixturePatchMenu::~SAddFixturePatchMenu()
 	{
-		SaveConfig();
+		UDMXAddFixturePatchMenuData* MenuData = GetMutableDefault<UDMXAddFixturePatchMenuData>();
+		MenuData->SoftFixtureType = WeakFixtureType.Get();
 	}
 
 	void SAddFixturePatchMenu::Construct(const FArguments& InArgs, TWeakPtr<FDMXEditor> InWeakDMXEditor)
@@ -53,7 +54,6 @@ namespace UE::DMXEditor::FixturePatchEditor
 			SelectedFixtureTypePtr = FixtureTypes.IsEmpty() ? nullptr : &FixtureTypes[0];
 		}
 		WeakFixtureType = SelectedFixtureTypePtr ? *SelectedFixtureTypePtr : nullptr;
-		ActiveModeIndex = WeakFixtureType.IsValid() && WeakFixtureType->Modes.IsValidIndex(ActiveModeIndex) ? ActiveModeIndex : 0;
 
 		ChildSlot
 		[
@@ -158,14 +158,6 @@ namespace UE::DMXEditor::FixturePatchEditor
 		}
 	}
 
-	void SAddFixturePatchMenu::SaveConfig()
-	{
-		UDMXAddFixturePatchMenuData* Data = GetMutableDefault<UDMXAddFixturePatchMenuData>();
-		Data->SoftFixtureType = WeakFixtureType.Get();
-		Data->ActiveModeIndex = ActiveModeIndex;
-		Data->SaveConfig();
-	}
-
 	TSharedRef<SWidget> SAddFixturePatchMenu::MakeFixtureTypeSelectWidget()
 	{
 		return
@@ -341,7 +333,9 @@ namespace UE::DMXEditor::FixturePatchEditor
 
 	void SAddFixturePatchMenu::OnModeSelected(TSharedPtr<uint32> InSelectedMode, ESelectInfo::Type SelectInfo)
 	{
-		ActiveModeIndex = InSelectedMode.IsValid() ? *InSelectedMode : 0;
+		UDMXAddFixturePatchMenuData* MenuData = GetMutableDefault<UDMXAddFixturePatchMenuData>();
+		MenuData->ActiveModeIndex = InSelectedMode.IsValid() ? *InSelectedMode : 0;
+		MenuData->SaveConfig();
 	}
 
 	FText SAddFixturePatchMenu::GetUniverseChannelText() const
@@ -482,7 +476,8 @@ namespace UE::DMXEditor::FixturePatchEditor
 		}
 
 		// Ensure valid mode
-		if (!ensureMsgf(FixtureType->Modes.IsValidIndex(ActiveModeIndex), TEXT("Cannot apply mode. Mode index is invalid.")))
+		const UDMXAddFixturePatchMenuData* MenuData = GetDefault<UDMXAddFixturePatchMenuData>();
+		if (!ensureMsgf(FixtureType->Modes.IsValidIndex(MenuData->ActiveModeIndex), TEXT("Cannot apply mode. Mode index is invalid.")))
 		{
 			return FReply::Handled();
 		}
@@ -499,7 +494,7 @@ namespace UE::DMXEditor::FixturePatchEditor
 		{
 			FDMXEntityFixturePatchConstructionParams FixturePatchConstructionParams;
 			FixturePatchConstructionParams.FixtureTypeRef = FDMXEntityFixtureTypeRef(FixtureType);
-			FixturePatchConstructionParams.ActiveMode = ActiveModeIndex;
+			FixturePatchConstructionParams.ActiveMode = MenuData->ActiveModeIndex;
 			FixturePatchConstructionParams.UniverseID = PatchToUniverse;
 			FixturePatchConstructionParams.StartingAddress = Channel.IsSet() ? Channel.GetValue() : 1;
 
diff --git a/Engine/Plugins/VirtualProduction/DMX/DMXEngine/Source/DMXEditor/Private/Widgets/FixturePatch/SAddFixturePatchMenu.h b/Engine/Plugins/VirtualProduction/DMX/DMXEngine/Source/DMXEditor/Private/Widgets/FixturePatch/SAddFixturePatchMenu.h
--- a/Engine/Plugins/VirtualProduction/DMX/DMXEngine/Source/DMXEditor/Private/Widgets/FixturePatch/SAddFixturePatchMenu.h
+++ b/Engine/Plugins/VirtualProduction/DMX/DMXEngine/Source/DMXEditor/Private/Widgets/FixturePatch/SAddFixturePatchMenu.h
@@ -37,9 +37,6 @@ namespace UE::DMXEditor::FixturePatchEditor
 		void RequestRefresh();
 
 	private:
-		/** Saves state to config */
-		void SaveConfig();
-
 		/** Creates a widget to select a fixture type */
 		TSharedRef<SWidget> MakeFixtureTypeSelectWidget();
 
@@ -91,9 +88,6 @@ namespace UE::DMXEditor::FixturePatchEditor
 		/** Current fixture type */
 		TWeakObjectPtr<UDMXEntityFixtureType> WeakFixtureType;
 
-		/** The number of fixture patches to add */
-		int32 ActiveModeIndex = 0;
-
 		/** Universe where the patches will be added */
 		TOptional<int32> Universe;
 
