DMX - Auto-assign and related methods now work for channel 512, menu now offers an option to increment the channels after patching

diff --git a/Engine/Plugins/VirtualProduction/DMX/DMXEngine/Source/DMXEditor/Private/FixturePatchAutoAssignUtility.cpp b/Engine/Plugins/VirtualProduction/DMX/DMXEngine/Source/DMXEditor/Private/FixturePatchAutoAssignUtility.cpp
--- a/Engine/Plugins/VirtualProduction/DMX/DMXEngine/Source/DMXEditor/Private/FixturePatchAutoAssignUtility.cpp
+++ b/Engine/Plugins/VirtualProduction/DMX/DMXEngine/Source/DMXEditor/Private/FixturePatchAutoAssignUtility.cpp
@@ -73,8 +73,14 @@ namespace UE::DMXEditor::AutoAssign::Private
 	{
 		if (UDMXEntityFixturePatch* Patch = WeakFixturePatch.Get())
 		{
-			const int64 Universe = (GetLowerBoundValue() + 1) / DMX_UNIVERSE_SIZE;
-			const int64 Channel =  GetLowerBoundValue() + 1 - Universe * DMX_UNIVERSE_SIZE;
+			// Transition back to the 1-512 range
+			int64 Universe = (GetLowerBoundValue() + 1) / DMX_UNIVERSE_SIZE;
+			int64 Channel =  (GetLowerBoundValue() + 1) % DMX_UNIVERSE_SIZE;
+			if (Channel == 0)
+			{
+				Channel = 512;
+				--Universe;
+			}
 			if (!ensureMsgf(Universe >= 1 && Channel >= 1 && Channel <= DMX_UNIVERSE_SIZE, TEXT("Invalid universe and channel resulting from auto assign. Channels not applied for '%s'."), *Patch->Name))
 			{
 				return;
diff --git a/Engine/Plugins/VirtualProduction/DMX/DMXEngine/Source/DMXEditor/Private/Widgets/FixturePatch/DMXAddFixturePatchMenuData.h b/Engine/Plugins/VirtualProduction/DMX/DMXEngine/Source/DMXEditor/Private/Widgets/FixturePatch/DMXAddFixturePatchMenuData.h
--- a/Engine/Plugins/VirtualProduction/DMX/DMXEngine/Source/DMXEditor/Private/Widgets/FixturePatch/DMXAddFixturePatchMenuData.h
+++ b/Engine/Plugins/VirtualProduction/DMX/DMXEngine/Source/DMXEditor/Private/Widgets/FixturePatch/DMXAddFixturePatchMenuData.h
@@ -25,4 +25,7 @@ public:
 
 	UPROPERTY(Config)
 	int32 NumPatches = 0;
+
+	UPROPERTY(Config)
+	bool bIncrementChannelAfterPatching = true;
 };
diff --git a/Engine/Plugins/VirtualProduction/DMX/DMXEngine/Source/DMXEditor/Private/Widgets/FixturePatch/SAddFixturePatchMenu.cpp b/Engine/Plugins/VirtualProduction/DMX/DMXEngine/Source/DMXEditor/Private/Widgets/FixturePatch/SAddFixturePatchMenu.cpp
--- a/Engine/Plugins/VirtualProduction/DMX/DMXEngine/Source/DMXEditor/Private/Widgets/FixturePatch/SAddFixturePatchMenu.cpp
+++ b/Engine/Plugins/VirtualProduction/DMX/DMXEngine/Source/DMXEditor/Private/Widgets/FixturePatch/SAddFixturePatchMenu.cpp
@@ -16,10 +16,11 @@
 #include "TimerManager.h"
 #include "UObject/Package.h"
 #include "Widgets/Input/SButton.h"
+#include "Widgets/Input/SCheckBox.h"
+#include "Widgets/Input/SEditableTextBox.h"
 #include "Widgets/Layout/SGridPanel.h"
 #include "Widgets/SBoxPanel.h"
 #include "Widgets/SDMXEntityDropdownMenu.h"
-#include "Widgets/Input/SEditableTextBox.h"
 
 
 #define LOCTEXT_NAMESPACE "SAddFixturePatchMenu"
@@ -112,7 +113,20 @@ namespace UE::DMXEditor::FixturePatchEditor
 				.Padding(4.f)
 				.VAlign(VAlign_Center)
 				[
-					MakeUniverseChannelSelectWidget()
+					SNew(SHorizontalBox)
+
+					+ SHorizontalBox::Slot()
+					.AutoWidth()
+					[
+						MakeUniverseChannelSelectWidget()
+					]
+
+					+ SHorizontalBox::Slot()
+					.AutoWidth()
+					.Padding(4.f, 0.f, 0.f, 0.f)
+					[
+						MakeAutoIncrementChannelCheckBox()
+					]
 				]
 
 				// Num Fixture Patches Label
@@ -207,6 +221,30 @@ namespace UE::DMXEditor::FixturePatchEditor
 			];
 	}
 
+	TSharedRef<SWidget> SAddFixturePatchMenu::MakeAutoIncrementChannelCheckBox()
+	{
+		return 
+			SNew(SCheckBox)
+			.IsEnabled(this, &SAddFixturePatchMenu::HasValidFixtureTypeAndMode)
+			.IsChecked_Lambda([]()
+				{
+					const UDMXAddFixturePatchMenuData* MenuData = GetDefault<UDMXAddFixturePatchMenuData>();
+					return MenuData->bIncrementChannelAfterPatching ? ECheckBoxState::Checked : ECheckBoxState::Unchecked;
+				})
+			.OnCheckStateChanged_Lambda([](ECheckBoxState NewState)
+				{
+					UDMXAddFixturePatchMenuData* MenuData = GetMutableDefault<UDMXAddFixturePatchMenuData>();
+					MenuData->bIncrementChannelAfterPatching = NewState == ECheckBoxState::Checked;
+					MenuData->SaveConfig();
+				})
+			[
+				SNew(STextBlock)
+				.Font(FAppStyle::GetFontStyle(TEXT("PropertyWindow.NormalFont")))
+				.Text(LOCTEXT("AutoIncrementChannelLabel", "Increment after patching"))
+				.ToolTipText(LOCTEXT("AutoIncrementChannelTooltip", "Automatically increments the universe/channel to the first subsequent channel after patching."))
+			];
+	}
+
 	TSharedRef<SWidget> SAddFixturePatchMenu::MakeNumFixturePatchesEditableTextBox()
 	{
 		return
@@ -298,6 +336,7 @@ namespace UE::DMXEditor::FixturePatchEditor
 	{
 		RequestRefreshModeComboBoxTimerHandle.Invalidate();
 
+		// Mend the fixture type
 		if (!WeakFixtureType.IsValid() && WeakDMXEditor.IsValid())
 		{
 			if (UDMXLibrary* DMXLibrary = WeakDMXEditor.Pin()->GetDMXLibrary())
@@ -307,6 +346,14 @@ namespace UE::DMXEditor::FixturePatchEditor
 			}
 		}
 
+		// Mend the active mode index
+		UDMXAddFixturePatchMenuData* MenuData = GetMutableDefault<UDMXAddFixturePatchMenuData>();
+		if (WeakFixtureType.IsValid() && !WeakFixtureType->Modes.IsValidIndex(MenuData->ActiveModeIndex))
+		{
+			MenuData->ActiveModeIndex = 0;
+			MenuData->SaveConfig();
+		}
+
 		ModeSources.Reset();
 		if (UDMXEntityFixtureType* FixtureType = WeakFixtureType.Get())
 		{
@@ -517,6 +564,18 @@ namespace UE::DMXEditor::FixturePatchEditor
 		Algo::Copy(NewFixturePatches, NewWeakFixturePatches);
 		SharedData->SelectFixturePatches(NewWeakFixturePatches);
 
+		// Increment universe/channel if desired
+		if (!NewFixturePatches.IsEmpty() && GetDefault<UDMXAddFixturePatchMenuData>()->bIncrementChannelAfterPatching)
+		{
+			Channel = NewFixturePatches.Last()->GetStartingChannel() + NewFixturePatches.Last()->GetChannelSpan();
+			Universe = NewFixturePatches.Last()->GetUniverseID();
+			if (Channel.GetValue() > DMX_UNIVERSE_SIZE)
+			{
+				Channel = 1;
+				Universe = Universe.GetValue() + 1;
+			}
+		}
+
 		return FReply::Handled();
 	}
 
diff --git a/Engine/Plugins/VirtualProduction/DMX/DMXEngine/Source/DMXEditor/Private/Widgets/FixturePatch/SAddFixturePatchMenu.h b/Engine/Plugins/VirtualProduction/DMX/DMXEngine/Source/DMXEditor/Private/Widgets/FixturePatch/SAddFixturePatchMenu.h
--- a/Engine/Plugins/VirtualProduction/DMX/DMXEngine/Source/DMXEditor/Private/Widgets/FixturePatch/SAddFixturePatchMenu.h
+++ b/Engine/Plugins/VirtualProduction/DMX/DMXEngine/Source/DMXEditor/Private/Widgets/FixturePatch/SAddFixturePatchMenu.h
@@ -45,6 +45,9 @@ namespace UE::DMXEditor::FixturePatchEditor
 
 		/** Creates an widget to select the Universe and Channel */
 		TSharedRef<SWidget> MakeUniverseChannelSelectWidget();
+		
+		/** Creates a widget to toggle if the channels should be auto-incremented after patching */
+		TSharedRef<SWidget> MakeAutoIncrementChannelCheckBox();
 
 		/** Creates an editable text box to specify the number of patches */
 		TSharedRef<SWidget> MakeNumFixturePatchesEditableTextBox();
