DMX - In the DMX Library, auto assigning from the fixture list left now correctly updates the patch view right

diff --git a/Engine/Plugins/VirtualProduction/DMX/DMXEngine/Source/DMXEditor/Private/Widgets/FixturePatch/SDMXMVRFixtureList.cpp b/Engine/Plugins/VirtualProduction/DMX/DMXEngine/Source/DMXEditor/Private/Widgets/FixturePatch/SDMXMVRFixtureList.cpp
--- a/Engine/Plugins/VirtualProduction/DMX/DMXEngine/Source/DMXEditor/Private/Widgets/FixturePatch/SDMXMVRFixtureList.cpp
+++ b/Engine/Plugins/VirtualProduction/DMX/DMXEngine/Source/DMXEditor/Private/Widgets/FixturePatch/SDMXMVRFixtureList.cpp
@@ -921,11 +921,13 @@ void SDMXMVRFixtureList::AutoAssignFixturePatches(UE::DMXEditor::AutoAssign::EAu
 
 	TArray<UDMXEntityFixturePatch*> FixturePatchesToAutoAssign;
 	const TArray<TWeakObjectPtr<UDMXEntityFixturePatch>> SelectedFixturePatches = FixturePatchSharedData->GetSelectedFixturePatches();
-	for (TWeakObjectPtr<UDMXEntityFixturePatch> FixturePatch : SelectedFixturePatches)
+	for (TWeakObjectPtr<UDMXEntityFixturePatch> WeakFixturePatch : SelectedFixturePatches)
 	{
-		if (FixturePatch.IsValid())
+		if (UDMXEntityFixturePatch* FixturePatch = WeakFixturePatch.Get())
 		{
-			FixturePatchesToAutoAssign.Add(FixturePatch.Get());
+			FixturePatch->PreEditChange(UDMXEntityFixturePatch::StaticClass()->FindPropertyByName(UDMXEntityFixturePatch::GetUniverseIDPropertyNameChecked()));
+			FixturePatch->PreEditChange(UDMXEntityFixturePatch::StaticClass()->FindPropertyByName(UDMXEntityFixturePatch::GetStartingChannelPropertyNameChecked()));
+			FixturePatchesToAutoAssign.Add(FixturePatch);
 		}
 	}
 
@@ -938,6 +940,12 @@ void SDMXMVRFixtureList::AutoAssignFixturePatches(UE::DMXEditor::AutoAssign::EAu
 	FAutoAssignUtility::AutoAssign(Mode, WeakDMXEditor.Pin().ToSharedRef(), FixturePatchesToAutoAssign, 1, 1);
 	FixturePatchSharedData->SelectUniverse(FixturePatchesToAutoAssign[0]->GetUniverseID());
 
+	for (UDMXEntityFixturePatch* FixturePatch : FixturePatchesToAutoAssign)
+	{
+		FixturePatch->PostEditChange();
+		FixturePatch->PostEditChange();
+	}
+
 	RequestListRefresh();
 }
 
