DMX - Fix output ports may access stale uobject ptr during engine shutdown, fix engine crash on shutdown when running with -stompmalloc

diff --git a/Engine/Plugins/VirtualProduction/DMX/DMXEngine/Source/DMXEditor/Private/DMXEditorModule.cpp b/Engine/Plugins/VirtualProduction/DMX/DMXEngine/Source/DMXEditor/Private/DMXEditorModule.cpp
--- a/Engine/Plugins/VirtualProduction/DMX/DMXEngine/Source/DMXEditor/Private/DMXEditorModule.cpp
+++ b/Engine/Plugins/VirtualProduction/DMX/DMXEngine/Source/DMXEditor/Private/DMXEditorModule.cpp
@@ -434,7 +434,7 @@ void FDMXEditorModule::RegisterAssetTypeAction(TSharedRef<IAssetTypeActions> Act
 
 void FDMXEditorModule::UnregisterAssetTypeActions()
 {
-	if (FModuleManager::Get().IsModuleLoaded("AssetTools"))
+	if (UObjectInitialized() && FModuleManager::Get().IsModuleLoaded("AssetTools"))
 	{
 		IAssetTools& AssetToolsModule = FModuleManager::GetModuleChecked<FAssetToolsModule>("AssetTools").Get();
 		for (TSharedPtr<IAssetTypeActions>& AssetIt : RegisteredAssetTypeActions)
@@ -461,9 +461,10 @@ void FDMXEditorModule::RegisterCustomClassLayout(FName ClassName, FOnGetDetailCu
 
 void FDMXEditorModule::UnregisterCustomClassLayouts()
 {
-	FPropertyEditorModule* PropertyModule = FModuleManager::GetModulePtr<FPropertyEditorModule>("PropertyEditor");
-	if (PropertyModule)
+	if (UObjectInitialized() && FModuleManager::Get().IsModuleLoaded(TEXT("PropertyEditor")))
 	{
+		FPropertyEditorModule* PropertyModule = FModuleManager::GetModulePtr<FPropertyEditorModule>("PropertyEditor");
+
 		for (const FName& ClassName : RegisteredClassNames)
 		{
 			PropertyModule->UnregisterCustomClassLayout(ClassName);
@@ -488,9 +489,10 @@ void FDMXEditorModule::RegisterCustomPropertyTypeLayout(FName PropertyTypeName,
 
 void FDMXEditorModule::UnregisterCustomPropertyTypeLayouts()
 {
-	FPropertyEditorModule* PropertyModule = FModuleManager::GetModulePtr<FPropertyEditorModule>("PropertyEditor");
-	if (PropertyModule)
+	if (UObjectInitialized() && FModuleManager::Get().IsModuleLoaded(TEXT("PropertyEditor")))
 	{
+		FPropertyEditorModule* PropertyModule = FModuleManager::GetModulePtr<FPropertyEditorModule>("PropertyEditor");
+
 		for (const FName& PropertyTypeName : RegisteredPropertyTypes)
 		{
 			PropertyModule->UnregisterCustomPropertyTypeLayout(PropertyTypeName);
@@ -508,9 +510,10 @@ void  FDMXEditorModule::RegisterCustomSequencerTrackType(const FOnCreateTrackEdi
 
 void FDMXEditorModule::UnregisterCustomSequencerTrackTypes()
 {
-	ISequencerModule* SequencerModule = FModuleManager::GetModulePtr<ISequencerModule>("Sequencer");
-	if (SequencerModule)
+	if (UObjectInitialized() && FModuleManager::Get().IsModuleLoaded(TEXT("Sequencer")))
 	{
+		ISequencerModule* SequencerModule = FModuleManager::GetModulePtr<ISequencerModule>("Sequencer");
+
 		for (const FDelegateHandle& TrackCreateHandle : RegisteredSequencerTrackHandles)
 		{
 			SequencerModule->UnRegisterTrackEditor(TrackCreateHandle);
diff --git a/Engine/Plugins/VirtualProduction/DMX/DMXProtocol/Source/DMXProtocol/Private/DMXProtocolSettings.cpp b/Engine/Plugins/VirtualProduction/DMX/DMXProtocol/Source/DMXProtocol/Private/DMXProtocolSettings.cpp
--- a/Engine/Plugins/VirtualProduction/DMX/DMXProtocol/Source/DMXProtocol/Private/DMXProtocolSettings.cpp
+++ b/Engine/Plugins/VirtualProduction/DMX/DMXProtocol/Source/DMXProtocol/Private/DMXProtocolSettings.cpp
@@ -174,7 +174,11 @@ void UDMXProtocolSettings::PostEditChangeChainProperty(FPropertyChangedChainEven
 	const UScriptStruct* OutputPortConfigStruct = FDMXOutputPortConfig::StaticStruct();
 	const UStruct* PropertyOwnerStruct = Property ? Property->GetOwnerStruct() : nullptr;
 
-	if (PropertyName == GET_MEMBER_NAME_CHECKED(UDMXProtocolSettings, FixtureCategories))
+	if (PropertyName == GET_MEMBER_NAME_CHECKED(UDMXProtocolSettings, SendingRefreshRate))
+	{
+		OnSendingRefresRateChanged.Broadcast();
+	}
+	else if (PropertyName == GET_MEMBER_NAME_CHECKED(UDMXProtocolSettings, FixtureCategories))
 	{
 		if (FixtureCategories.Num() == 0)
 		{
diff --git a/Engine/Plugins/VirtualProduction/DMX/DMXProtocol/Source/DMXProtocol/Private/IO/DMXOutputPort.cpp b/Engine/Plugins/VirtualProduction/DMX/DMXProtocol/Source/DMXProtocol/Private/IO/DMXOutputPort.cpp
--- a/Engine/Plugins/VirtualProduction/DMX/DMXProtocol/Source/DMXProtocol/Private/IO/DMXOutputPort.cpp
+++ b/Engine/Plugins/VirtualProduction/DMX/DMXProtocol/Source/DMXProtocol/Private/IO/DMXOutputPort.cpp
@@ -541,12 +541,13 @@ FDMXOutputPortSharedRef FDMXOutputPort::CreateFromConfig(FDMXOutputPortConfig& O
 	UDMXProtocolSettings* Settings = GetMutableDefault<UDMXProtocolSettings>();
 	check(Settings);
 
+	NewOutputPort->SendRate = Settings->SendingRefreshRate;
 	NewOutputPort->CommunicationDeterminator.SetSendEnabled(Settings->IsSendDMXEnabled());
 	NewOutputPort->CommunicationDeterminator.SetReceiveEnabled(Settings->IsReceiveDMXEnabled());
 
+	Settings->GetOnSendingRefresRateChanged().AddThreadSafeSP(NewOutputPort, &FDMXOutputPort::OnSendingRefreshRateChanged);
 	Settings->GetOnSetSendDMXEnabled().AddThreadSafeSP(NewOutputPort, &FDMXOutputPort::OnSetSendDMXEnabled);
 	Settings->GetOnSetReceiveDMXEnabled().AddThreadSafeSP(NewOutputPort, &FDMXOutputPort::OnSetReceiveDMXEnabled);
-
 	NewOutputPort->UpdateFromConfig(OutputPortConfig);
 
 	const FString SenderThreadName = FString(TEXT("DMXOutputPort_")) + OutputPortConfig.GetPortName();
@@ -908,6 +909,12 @@ void FDMXOutputPort::Unregister()
 	bRegistered = false;
 }
 
+void FDMXOutputPort::OnSendingRefreshRateChanged()
+{
+	const UDMXProtocolSettings* ProtocolSettings = GetDefault< UDMXProtocolSettings>();
+	SendRate = ProtocolSettings->SendingRefreshRate;
+}
+
 void FDMXOutputPort::OnSetSendDMXEnabled(bool bEnabled)
 {
 	checkf(IsInGameThread(), TEXT("Only the game thread can enable and disable send DMX."));
@@ -935,20 +942,22 @@ bool FDMXOutputPort::Init()
 
 uint32 FDMXOutputPort::Run()
 {
-	UDMXProtocolSettings* DMXSettings = GetMutableDefault<UDMXProtocolSettings>();
-	check(DMXSettings);
-
 	SendDMXEvent = FPlatformProcess::GetSynchEventFromPool();
 	check(SendDMXEvent != nullptr);
 
 	while (!bStopping)
 	{
 		const double StartTime = FPlatformTime::Seconds();
+		if (IsEngineExitRequested())
+		{
+			// Stop sending DMX when the engine is shutting down
+			break;
+		}
 
 		ProcessSendDMX();
 
 		const double EndTime = FPlatformTime::Seconds();
-		const double WaitTimeMs = ((1.0 / DMXSettings->SendingRefreshRate) - (EndTime - StartTime)) * 1000.0;
+		const double WaitTimeMs = ((1.0 / SendRate) - (EndTime - StartTime)) * 1000.0;
 		
 		if (WaitTimeMs > 0.0 && WaitTimeMs < std::numeric_limits<uint32>::max())
 		{
diff --git a/Engine/Plugins/VirtualProduction/DMX/DMXProtocol/Source/DMXProtocol/Public/DMXProtocolSettings.h b/Engine/Plugins/VirtualProduction/DMX/DMXProtocol/Source/DMXProtocol/Public/DMXProtocolSettings.h
--- a/Engine/Plugins/VirtualProduction/DMX/DMXProtocol/Source/DMXProtocol/Public/DMXProtocolSettings.h
+++ b/Engine/Plugins/VirtualProduction/DMX/DMXProtocol/Source/DMXProtocol/Public/DMXProtocolSettings.h
@@ -96,25 +96,31 @@ public:
 	/** Returns an automatic/unique name for an OutputPort */
 	FString GetUniqueOutputPortName() const;
 
-	/** Gets a Delegate Broadcast when send DMX is enabled or disabled */
+	/** Gets a delegate broadcast when the Sending Refresh Rate (the DMX send rate) changed */
+	FSimpleMulticastDelegate& GetOnSendingRefresRateChanged()
+	{
+		return OnSendingRefresRateChanged;
+	}
+
+	/** Gets a delegate broadcast when send DMX is enabled or disabled */
 	FDMXOnSendDMXEnabled& GetOnSetSendDMXEnabled()
 	{
 		return OnSetSendDMXEnabledDelegate;
 	}
 
-	/** Gets a Delegate Broadcast when receive DMX is enabled or disabled */
+	/** Gets a delegate broadcast when receive DMX is enabled or disabled */
 	FDMXOnSendDMXEnabled& GetOnSetReceiveDMXEnabled()
 	{
 		return OnSetReceiveDMXEnabledDelegate;
 	}
 
-	/** Gets a Delegate Broadcast when default attributes were changed */
+	/** Gets a delegate broadcast when default attributes were changed */
 	FSimpleMulticastDelegate& GetOnDefaultAttributesChanged()
 	{
 		return OnDefaultAttributesChanged;
 	}
 
-	/** Gets a Delegate Broadcast when default fixture categories were changed */
+	/** Gets a delegate broadcast when default fixture categories were changed */
 	FSimpleMulticastDelegate& GetOnDefaultFixtureCategoriesChanged()
 	{
 		return OnDefaultFixtureCategoriesChanged;
@@ -165,6 +171,9 @@ private:
 	/** Overrides the default bDefaultReceiveDMXEnabled value at runtime */
 	bool bOverrideReceiveDMXEnabled = true;
 
+	/** Broadcast the Sending Refresh Rate (the DMX send rate) changed */
+	FSimpleMulticastDelegate OnSendingRefresRateChanged;
+
 	/** Broadcast when send DMX is enabled or disabled */
 	FDMXOnSendDMXEnabled OnSetSendDMXEnabledDelegate;
 
diff --git a/Engine/Plugins/VirtualProduction/DMX/DMXProtocol/Source/DMXProtocol/Public/IO/DMXOutputPort.h b/Engine/Plugins/VirtualProduction/DMX/DMXProtocol/Source/DMXProtocol/Public/IO/DMXOutputPort.h
--- a/Engine/Plugins/VirtualProduction/DMX/DMXProtocol/Source/DMXProtocol/Public/IO/DMXOutputPort.h
+++ b/Engine/Plugins/VirtualProduction/DMX/DMXProtocol/Source/DMXProtocol/Public/IO/DMXOutputPort.h
@@ -4,6 +4,7 @@
 
 #include "DMXPort.h"
 
+#include <atomic>
 #include "DMXProtocolCommon.h"
 #include "Containers/Queue.h" 
 #include "HAL/Runnable.h"
@@ -168,7 +169,10 @@ protected:
 	virtual void Unregister() override;
 	//~ End DMXPort Interface
 
-		/** Called to set if DMX should be enabled */
+	/** Called when the DMX send rate changed */
+	void OnSendingRefreshRateChanged();
+
+	/** Called to set if DMX should be enabled */
 	void OnSetSendDMXEnabled(bool bEnabled);
 
 	/** Called to set if DMX should be enabled */
@@ -223,6 +227,9 @@ private:
 	/** The unique identifier of this port, shared with the port config this was constructed from. Should not be changed after construction. */
 	FGuid PortGuid;
 
+	/** Current send rate, set from project settings (see OnSendingRefreshRateChanged) */
+	TAtomic<uint32> SendRate = 44.f;
+
 	/** Delay to apply on packets being sent */
 	double DelaySeconds = 0.0;
 
diff --git a/Engine/Plugins/VirtualProduction/DMX/DMXProtocol/Source/DMXProtocolEditor/Private/DMXProtocolEditorModule.cpp b/Engine/Plugins/VirtualProduction/DMX/DMXProtocol/Source/DMXProtocolEditor/Private/DMXProtocolEditorModule.cpp
--- a/Engine/Plugins/VirtualProduction/DMX/DMXProtocol/Source/DMXProtocolEditor/Private/DMXProtocolEditorModule.cpp
+++ b/Engine/Plugins/VirtualProduction/DMX/DMXProtocol/Source/DMXProtocolEditor/Private/DMXProtocolEditorModule.cpp
@@ -53,9 +53,10 @@ void FDMXProtocolEditorModule::RegisterDetailsCustomizations()
 
 void FDMXProtocolEditorModule::UnregisterDetailsCustomizations()
 {
-	FPropertyEditorModule* PropertyModule = FModuleManager::GetModulePtr<FPropertyEditorModule>("PropertyEditor");
-	if (PropertyModule)
+	if (UObjectInitialized() && FModuleManager::Get().IsModuleLoaded(TEXT("PropertyEditor")))
 	{
+		FPropertyEditorModule* PropertyModule = FModuleManager::GetModulePtr<FPropertyEditorModule>("PropertyEditor");
+
 		PropertyModule->UnregisterCustomPropertyTypeLayout(FDMXInputPortConfig::StaticStruct()->GetFName());
 		PropertyModule->UnregisterCustomPropertyTypeLayout(FDMXOutputPortConfig::StaticStruct()->GetFName());
 
