DMX - Fix patch viewer and activity monitor do not allow user to scroll properly

diff --git a/Engine/Plugins/VirtualProduction/DMX/DMXEngine/Source/DMXEditor/Private/Widgets/FixturePatch/SDMXFixturePatcher.cpp b/Engine/Plugins/VirtualProduction/DMX/DMXEngine/Source/DMXEditor/Private/Widgets/FixturePatch/SDMXFixturePatcher.cpp
--- a/Engine/Plugins/VirtualProduction/DMX/DMXEngine/Source/DMXEditor/Private/Widgets/FixturePatch/SDMXFixturePatcher.cpp
+++ b/Engine/Plugins/VirtualProduction/DMX/DMXEngine/Source/DMXEditor/Private/Widgets/FixturePatch/SDMXFixturePatcher.cpp
@@ -23,6 +23,7 @@
 #include "Widgets/Input/SComboButton.h"
 #include "Widgets/Input/SNumericEntryBox.h"
 #include "Widgets/Layout/SBox.h"
+#include "Widgets/Layout/SScrollBar.h"
 #include "Widgets/Layout/SScrollBox.h"
 #include "Widgets/Layout/SSeparator.h"
 
@@ -32,7 +33,6 @@
 BEGIN_SLATE_FUNCTION_BUILD_OPTIMIZATION
 void SDMXFixturePatcher::Construct(const FArguments& InArgs)
 {
-
 	DMXEditorPtr = InArgs._DMXEditor;
 	if (!DMXEditorPtr.IsValid())
 	{
@@ -41,6 +41,12 @@ void SDMXFixturePatcher::Construct(const FArguments& InArgs)
 
 	SharedData = DMXEditorPtr.Pin()->GetFixturePatchSharedData();
 
+	const TSharedRef<SScrollBar> HorizontalScrollBar = SNew(SScrollBar)
+		.Orientation(Orient_Horizontal);
+
+	const TSharedRef<SScrollBar> VerticalScrollBar = SNew(SScrollBar)
+		.Orientation(Orient_Vertical);
+
 	const FLinearColor BackgroundTint(0.6f, 0.6f, 0.6f, 1.0f);
 	ChildSlot			
 		[
@@ -54,7 +60,6 @@ void SDMXFixturePatcher::Construct(const FArguments& InArgs)
 				.HAlign(HAlign_Fill)
 				.AutoHeight()
 				[
-
 					SNew(SBorder)					
 					.HAlign(HAlign_Fill)
 					.BorderBackgroundColor(BackgroundTint)
@@ -68,10 +73,34 @@ void SDMXFixturePatcher::Construct(const FArguments& InArgs)
 				.HAlign(HAlign_Left)
 				.VAlign(VAlign_Fill)
 				[	
-					SAssignNew(PatchedUniverseScrollBox, SScrollBox)					
-					.Orientation(EOrientation::Orient_Vertical)	
-					.ScrollBarAlwaysVisible(true)
-					.ScrollBarThickness(FVector2D(5.f, 0.f))
+					SNew(SHorizontalBox)
+
+					+SHorizontalBox::Slot()
+					[
+						SAssignNew(PatchedUniverseHorizontalScrollBox, SScrollBox)
+						.Orientation(EOrientation::Orient_Horizontal)
+						.ExternalScrollbar(HorizontalScrollBar)
+
+						+SScrollBox::Slot()
+						.FillSize(1.f)
+						[
+							SAssignNew(PatchedUniverseVerticalScrollBox, SScrollBox)
+							.Orientation(EOrientation::Orient_Vertical)	
+							.ExternalScrollbar(VerticalScrollBar)
+						]
+					]
+
+					+ SHorizontalBox::Slot()
+					.AutoWidth()
+					[
+						VerticalScrollBar
+					]
+				]
+
+				+ SVerticalBox::Slot()
+				.AutoHeight()
+				[
+					HorizontalScrollBar
 				]
 			]
 		];
@@ -534,7 +563,8 @@ void SDMXFixturePatcher::OnFixturePatchSelectionChanged()
 	}
 	else
 	{
-		PatchedUniverseScrollBox->ScrollDescendantIntoView(*UniverseWidgetPtr);
+		PatchedUniverseHorizontalScrollBox->ScrollDescendantIntoView(*UniverseWidgetPtr);
+		PatchedUniverseVerticalScrollBox->ScrollDescendantIntoView(*UniverseWidgetPtr);
 	}
 }
 
@@ -564,7 +594,7 @@ int32 SDMXFixturePatcher::GetSelectedUniverse() const
 
 void SDMXFixturePatcher::ShowSelectedUniverse()
 {
-	PatchedUniverseScrollBox->ClearChildren();
+	PatchedUniverseVerticalScrollBox->ClearChildren();
 	PatchedUniversesByID.Reset();
 	
 	const int32 SelectedUniverseID = GetSelectedUniverse();
@@ -575,7 +605,7 @@ void SDMXFixturePatcher::ShowSelectedUniverse()
 		.OnDragEnterChannel(this, &SDMXFixturePatcher::OnDragEnterChannel)
 		.OnDropOntoChannel(this, &SDMXFixturePatcher::OnDropOntoChannel);
 
-	PatchedUniverseScrollBox->AddSlot()
+	PatchedUniverseVerticalScrollBox->AddSlot()
 		.Padding(FMargin(0.0f, 3.0f, 0.0f, 12.0f))			
 		[
 			NewPatchedUniverse
@@ -590,11 +620,11 @@ void SDMXFixturePatcher::ShowSelectedUniverse()
 
 void SDMXFixturePatcher::ShowAllPatchedUniverses(bool bForceReconstructWidget)
 {
-	check(PatchedUniverseScrollBox.IsValid());
+	check(PatchedUniverseHorizontalScrollBox.IsValid() && PatchedUniverseVerticalScrollBox.IsValid());
 
 	if (bForceReconstructWidget)
 	{
-		PatchedUniverseScrollBox->ClearChildren();
+		PatchedUniverseVerticalScrollBox->ClearChildren();
 		PatchedUniversesByID.Reset();
 	}
 
@@ -634,7 +664,7 @@ void SDMXFixturePatcher::ShowAllPatchedUniverses(bool bForceReconstructWidget)
 			{
 				// Remove universe widgets without patches
 				PatchedUniversesByID.Remove(UniverseByIDKvp.Key);
-				PatchedUniverseScrollBox->RemoveSlot(UniverseByIDKvp.Value.ToSharedRef());
+				PatchedUniverseVerticalScrollBox->RemoveSlot(UniverseByIDKvp.Value.ToSharedRef());
 			}
 		}
 
@@ -657,7 +687,8 @@ void SDMXFixturePatcher::ShowAllPatchedUniverses(bool bForceReconstructWidget)
 		const int32 SelectedUniverse = SharedData->GetSelectedUniverse();
 		if (PatchedUniversesByID.Contains(SelectedUniverse))
 		{
-			PatchedUniverseScrollBox->ScrollDescendantIntoView(PatchedUniversesByID[SelectedUniverse]);
+			PatchedUniverseHorizontalScrollBox->ScrollDescendantIntoView(PatchedUniversesByID[SelectedUniverse]);
+			PatchedUniverseVerticalScrollBox->ScrollDescendantIntoView(PatchedUniversesByID[SelectedUniverse]);
 		}
 	}
 
@@ -675,7 +706,7 @@ void SDMXFixturePatcher::AddUniverse(int32 UniverseID)
 		.OnDragEnterChannel(this, &SDMXFixturePatcher::OnDragEnterChannel)
 		.OnDropOntoChannel(this, &SDMXFixturePatcher::OnDropOntoChannel);
 
-	PatchedUniverseScrollBox->AddSlot()
+	PatchedUniverseVerticalScrollBox->AddSlot()
 		.Padding(FMargin(0.0f, 3.0f, 0.0f, 12.0f))
 		[
 			PatchedUniverse
diff --git a/Engine/Plugins/VirtualProduction/DMX/DMXEngine/Source/DMXEditor/Private/Widgets/FixturePatch/SDMXFixturePatcher.h b/Engine/Plugins/VirtualProduction/DMX/DMXEngine/Source/DMXEditor/Private/Widgets/FixturePatch/SDMXFixturePatcher.h
--- a/Engine/Plugins/VirtualProduction/DMX/DMXEngine/Source/DMXEditor/Private/Widgets/FixturePatch/SDMXFixturePatcher.h
+++ b/Engine/Plugins/VirtualProduction/DMX/DMXEngine/Source/DMXEditor/Private/Widgets/FixturePatch/SDMXFixturePatcher.h
@@ -162,8 +162,11 @@ protected:
 	/** Checkbox to determine if the DMX Monitor functionality should be enabled */
 	TSharedPtr<SCheckBox> EnableDMXMonitorCheckBox;
 
-	/** Scrollbox containing all patch universes */
-	TSharedPtr<SScrollBox> PatchedUniverseScrollBox;
+	/** Horizontal scrollbox containing all the patched universes */
+	TSharedPtr<SScrollBox> PatchedUniverseHorizontalScrollBox;
+
+	/** Vertical scrollbox containing all the patched universes */
+	TSharedPtr<SScrollBox> PatchedUniverseVerticalScrollBox;
 
 	/** Universe widgets by ID */
 	TMap<int32, TSharedPtr<SDMXPatchedUniverse>> PatchedUniversesByID;
diff --git a/Engine/Plugins/VirtualProduction/DMX/DMXEngine/Source/DMXEditor/Private/Widgets/Monitors/SDMXActivityMonitor.cpp b/Engine/Plugins/VirtualProduction/DMX/DMXEngine/Source/DMXEditor/Private/Widgets/Monitors/SDMXActivityMonitor.cpp
--- a/Engine/Plugins/VirtualProduction/DMX/DMXEngine/Source/DMXEditor/Private/Widgets/Monitors/SDMXActivityMonitor.cpp
+++ b/Engine/Plugins/VirtualProduction/DMX/DMXEngine/Source/DMXEditor/Private/Widgets/Monitors/SDMXActivityMonitor.cpp
@@ -16,6 +16,7 @@
 #include "SlateOptMacros.h"
 #include "Widgets/Input/SButton.h"
 #include "Widgets/Input/SEditableTextBox.h"
+#include "Widgets/Layout/SScrollBar.h"
 #include "Widgets/Layout/SScrollBox.h"
 #include "Widgets/Layout/SSeparator.h"
 #include "Widgets/Layout/SWrapBox.h"
@@ -47,6 +48,12 @@ void SDMXActivityMonitor::Construct(const FArguments& InArgs)
 {
 	SetCanTick(true);
 
+	const TSharedRef<SScrollBar> HorizontalScrollBar = SNew(SScrollBar)
+		.Orientation(Orient_Horizontal);
+
+	const TSharedRef<SScrollBar> VerticalScrollBar = SNew(SScrollBar)
+		.Orientation(Orient_Vertical);
+
 	ChildSlot
 		[
 			SNew(SVerticalBox)
@@ -179,18 +186,44 @@ void SDMXActivityMonitor::Construct(const FArguments& InArgs)
 			.VAlign(VAlign_Fill)
 			.HAlign(HAlign_Fill)
 			[
-				SNew(SScrollBox)
-				.Orientation(EOrientation::Orient_Horizontal)
+				SNew(SHorizontalBox)
+
+				+ SHorizontalBox::Slot()
+				[
+					SNew(SScrollBox)
+					.Orientation(EOrientation::Orient_Horizontal)
+					.ExternalScrollbar(HorizontalScrollBar)
+
+					+ SScrollBox::Slot()
+					.FillSize(1.f)
+					[
+						SNew(SScrollBox)
+						.Orientation(EOrientation::Orient_Vertical)
+						.ExternalScrollbar(VerticalScrollBar)
+
+						+ SScrollBox::Slot()				
+						[
+							SAssignNew(UniverseList, SListView<TSharedPtr<SDMXActivityInUniverse>>)
+							.ItemHeight(20.0f)
+							.ListItemsSource(&UniverseListSource)
+							.Visibility(EVisibility::Visible)
+							.OnGenerateRow(this, &SDMXActivityMonitor::OnGenerateUniverseRow)
+						]
+					]
+				]
 
-				+ SScrollBox::Slot()				
+				+ SHorizontalBox::Slot()
+				.AutoWidth()
 				[
-					SAssignNew(UniverseList, SListView<TSharedPtr<SDMXActivityInUniverse>>)
-					.ItemHeight(20.0f)
-					.ListItemsSource(&UniverseListSource)
-					.Visibility(EVisibility::Visible)
-					.OnGenerateRow(this, &SDMXActivityMonitor::OnGenerateUniverseRow)
+					VerticalScrollBar
 				]
 			]
+
+			+ SVerticalBox::Slot()
+			.AutoHeight()
+			[
+				HorizontalScrollBar
+			]
 		];
 
 	LoadMonitorSettings();
