DMX - Function attributes now have the same delete icon as cell attributes do

diff --git a/Engine/Plugins/VirtualProduction/DMX/DMXEngine/Source/DMXEditor/Private/Widgets/FixtureType/DMXFixtureTypeFunctionsEditorFunctionItem.cpp b/Engine/Plugins/VirtualProduction/DMX/DMXEngine/Source/DMXEditor/Private/Widgets/FixtureType/DMXFixtureTypeFunctionsEditorFunctionItem.cpp
--- a/Engine/Plugins/VirtualProduction/DMX/DMXEngine/Source/DMXEditor/Private/Widgets/FixtureType/DMXFixtureTypeFunctionsEditorFunctionItem.cpp
+++ b/Engine/Plugins/VirtualProduction/DMX/DMXEngine/Source/DMXEditor/Private/Widgets/FixtureType/DMXFixtureTypeFunctionsEditorFunctionItem.cpp
@@ -189,6 +189,19 @@ void FDMXFixtureTypeFunctionsEditorFunctionItem::SetFunctionName(const FText& De
 	}
 }
 
+void FDMXFixtureTypeFunctionsEditorFunctionItem::DeleteFunction()
+{
+	if (ensureMsgf(FixtureType.IsValid() && FixtureType->Modes.IsValidIndex(ModeIndex) && FixtureType->Modes[ModeIndex].Functions.IsValidIndex(FunctionIndex), TEXT("Invalid Fixture Type, Mode or Function in FDMXFixtureTypeFunctionsEditorFunctionItem.")))
+	{
+		const FScopedTransaction DeleteFunctionTransaction(LOCTEXT("DeleteFunctionTransaction", "Delete Fixture Function"));
+		FixtureType->PreEditChange(FDMXFixtureFunction::StaticStruct()->FindPropertyByName(GET_MEMBER_NAME_CHECKED(FDMXFixtureMode, Functions)));
+
+		FixtureType->Modes[ModeIndex].Functions.RemoveAt(FunctionIndex);
+
+		FixtureType->PostEditChange();
+	}
+}
+
 TSharedPtr<FDMXFixtureTypeSharedData> FDMXFixtureTypeFunctionsEditorFunctionItem::GetFixtureTypeSharedData() const
 {
 	if (TSharedPtr<FDMXEditor> DMXEditor = WeakDMXEditor.Pin())
diff --git a/Engine/Plugins/VirtualProduction/DMX/DMXEngine/Source/DMXEditor/Private/Widgets/FixtureType/DMXFixtureTypeFunctionsEditorFunctionItem.h b/Engine/Plugins/VirtualProduction/DMX/DMXEngine/Source/DMXEditor/Private/Widgets/FixtureType/DMXFixtureTypeFunctionsEditorFunctionItem.h
--- a/Engine/Plugins/VirtualProduction/DMX/DMXEngine/Source/DMXEditor/Private/Widgets/FixtureType/DMXFixtureTypeFunctionsEditorFunctionItem.h
+++ b/Engine/Plugins/VirtualProduction/DMX/DMXEngine/Source/DMXEditor/Private/Widgets/FixtureType/DMXFixtureTypeFunctionsEditorFunctionItem.h
@@ -51,6 +51,9 @@ public:
 	/** Sets a Function Name, makes it unique. Returns the OutUniqueFunctionName as output argument */
 	void SetFunctionName(const FText& DesiredFunctionName, FText& OutUniqueFunctionName);
 
+	/** Deletes the function this item stands for */
+	void DeleteFunction();
+
 	/** Returns the Fixture Type Shared Data for the editor that makes use of this Item */
 	TSharedPtr<FDMXFixtureTypeSharedData> GetFixtureTypeSharedData() const;
 
diff --git a/Engine/Plugins/VirtualProduction/DMX/DMXEngine/Source/DMXEditor/Private/Widgets/FixtureType/SDMXFixtureTypeFunctionsEditor.cpp b/Engine/Plugins/VirtualProduction/DMX/DMXEngine/Source/DMXEditor/Private/Widgets/FixtureType/SDMXFixtureTypeFunctionsEditor.cpp
--- a/Engine/Plugins/VirtualProduction/DMX/DMXEngine/Source/DMXEditor/Private/Widgets/FixtureType/SDMXFixtureTypeFunctionsEditor.cpp
+++ b/Engine/Plugins/VirtualProduction/DMX/DMXEngine/Source/DMXEditor/Private/Widgets/FixtureType/SDMXFixtureTypeFunctionsEditor.cpp
@@ -25,10 +25,11 @@
 
 #define LOCTEXT_NAMESPACE "SDMXFixtureTypeFunctionsEditor"
 
-const FName FDMXFixtureTypeFunctionsEditorCollumnIDs::Status = "Status";
-const FName FDMXFixtureTypeFunctionsEditorCollumnIDs::Channel = "Channel";
-const FName FDMXFixtureTypeFunctionsEditorCollumnIDs::Name = "Name";
-const FName FDMXFixtureTypeFunctionsEditorCollumnIDs::Attribute = "Attribute";
+const FName SDMXFixtureTypeFunctionsEditor::FCollumnID::Status = "Status";
+const FName SDMXFixtureTypeFunctionsEditor::FCollumnID::Channel = "Channel";
+const FName SDMXFixtureTypeFunctionsEditor::FCollumnID::Name = "Name";
+const FName SDMXFixtureTypeFunctionsEditor::FCollumnID::Attribute = "Attribute";
+const FName SDMXFixtureTypeFunctionsEditor::FCollumnID::DeleteAttribute = "DeleteAttribute";
 
 SDMXFixtureTypeFunctionsEditor::~SDMXFixtureTypeFunctionsEditor()
 {
@@ -371,38 +372,46 @@ void SDMXFixtureTypeFunctionsEditor::OnFixtureTypeChanged(const UDMXEntityFixtur
 TSharedRef<SHeaderRow> SDMXFixtureTypeFunctionsEditor::GenerateHeaderRow()
 {
 	const float StatusColumnWidth = FMath::Max(FAppStyle::Get().GetBrush("Icons.Warning")->GetImageSize().X + 6.f, FAppStyle::Get().GetBrush("Icons.Error")->GetImageSize().X + 6.f);
+	const float DeleteAttributeColumnWidth = FAppStyle::GetBrush("Icons.Delete")->GetImageSize().X + 8.f;
 
 	HeaderRow = SNew(SHeaderRow);
 	SHeaderRow::FColumn::FArguments ColumnArgs;
 
 	HeaderRow->AddColumn(
 		SHeaderRow::FColumn::FArguments()
-		.ColumnId(FDMXFixtureTypeFunctionsEditorCollumnIDs::Status)
+		.ColumnId(SDMXFixtureTypeFunctionsEditor::FCollumnID::Status)
 		.DefaultLabel(LOCTEXT("StatusColumnLabel", ""))
 		.FixedWidth(StatusColumnWidth)
 	);
 
 	HeaderRow->AddColumn(
 		SHeaderRow::FColumn::FArguments()
-		.ColumnId(FDMXFixtureTypeFunctionsEditorCollumnIDs::Channel)
+		.ColumnId(SDMXFixtureTypeFunctionsEditor::FCollumnID::Channel)
 		.DefaultLabel(LOCTEXT("ChannelColumnLabel", "Ch."))
 		.FixedWidth(56.f)
 	);
 
 	HeaderRow->AddColumn(
 		SHeaderRow::FColumn::FArguments()
-		.ColumnId(FDMXFixtureTypeFunctionsEditorCollumnIDs::Name)
+		.ColumnId(SDMXFixtureTypeFunctionsEditor::FCollumnID::Name)
 		.DefaultLabel(LOCTEXT("NameColumnLabel", "Name"))
 		.FillWidth(0.5f)
 	);
 
 	HeaderRow->AddColumn(
 		SHeaderRow::FColumn::FArguments()
-		.ColumnId(FDMXFixtureTypeFunctionsEditorCollumnIDs::Attribute)
+		.ColumnId(SDMXFixtureTypeFunctionsEditor::FCollumnID::Attribute)
 		.DefaultLabel(LOCTEXT("AttributeColumnLabel", "Attribute"))
 		.FillWidth(0.5f)
 	);
 
+	HeaderRow->AddColumn(
+		SHeaderRow::FColumn::FArguments()
+		.ColumnId(SDMXFixtureTypeFunctionsEditor::FCollumnID::DeleteAttribute)
+		.DefaultLabel(LOCTEXT("DeleteAttributeColumnLabel", ""))
+		.FixedWidth(DeleteAttributeColumnWidth)
+	);
+
 	// Restore user settings
 	if (const UDMXEditorSettings* EditorSettings = GetDefault<UDMXEditorSettings>())
 	{
@@ -415,7 +424,7 @@ TSharedRef<SHeaderRow> SDMXFixtureTypeFunctionsEditor::GenerateHeaderRow()
 		const float AttributeColumnWidth = EditorSettings->FixtureTypeFunctionsEditorSettings.AttributeColumnWidth;
 		if (AttributeColumnWidth > 10.f)
 		{
-			HeaderRow->SetColumnWidth(FDMXFixtureTypeFunctionsEditorCollumnIDs::Attribute, AttributeColumnWidth);
+			HeaderRow->SetColumnWidth(SDMXFixtureTypeFunctionsEditor::FCollumnID::Attribute, AttributeColumnWidth);
 		}
 	}
 
@@ -429,12 +438,12 @@ void SDMXFixtureTypeFunctionsEditor::SaveHeaderRowSettings()
 	{
 		for (const SHeaderRow::FColumn& Column : HeaderRow->GetColumns())
 		{
-			if (Column.ColumnId == FDMXFixtureTypeFunctionsEditorCollumnIDs::Name)
+			if (Column.ColumnId == SDMXFixtureTypeFunctionsEditor::FCollumnID::Name)
 			{
 				EditorSettings->FixtureTypeFunctionsEditorSettings.NameColumnWidth = Column.Width.Get();
 
 			}
-			else if (Column.ColumnId == FDMXFixtureTypeFunctionsEditorCollumnIDs::Attribute)
+			else if (Column.ColumnId == SDMXFixtureTypeFunctionsEditor::FCollumnID::Attribute)
 			{
 				EditorSettings->FixtureTypeFunctionsEditorSettings.AttributeColumnWidth = Column.Width.Get();
 			}
diff --git a/Engine/Plugins/VirtualProduction/DMX/DMXEngine/Source/DMXEditor/Private/Widgets/FixtureType/SDMXFixtureTypeFunctionsEditor.h b/Engine/Plugins/VirtualProduction/DMX/DMXEngine/Source/DMXEditor/Private/Widgets/FixtureType/SDMXFixtureTypeFunctionsEditor.h
--- a/Engine/Plugins/VirtualProduction/DMX/DMXEngine/Source/DMXEditor/Private/Widgets/FixtureType/SDMXFixtureTypeFunctionsEditor.h
+++ b/Engine/Plugins/VirtualProduction/DMX/DMXEngine/Source/DMXEditor/Private/Widgets/FixtureType/SDMXFixtureTypeFunctionsEditor.h
@@ -21,15 +21,6 @@ template <typename ItemType> class SListView;
 class STableViewBase;
 class SVerticalBox;
 
-/** Collumn IDs in the Functions Editor */
-struct FDMXFixtureTypeFunctionsEditorCollumnIDs 
-{
-	static const FName Status;
-	static const FName Channel;
-	static const FName Name;
-	static const FName Attribute;
-};
-
 
 using FDMXFixtureTypeFunctionsEditorFunctionListType = SListView<TSharedPtr<FDMXFixtureTypeFunctionsEditorItemBase>>;
 
@@ -38,6 +29,16 @@ class SDMXFixtureTypeFunctionsEditor
 	: public SCompoundWidget
 {
 public:
+	/** Collumn IDs in the Functions Editor */
+	struct FCollumnID
+	{
+		static const FName Status;
+		static const FName Channel;
+		static const FName Name;
+		static const FName Attribute;
+		static const FName DeleteAttribute;
+	};
+
 	SLATE_BEGIN_ARGS(SDMXFixtureTypeFunctionsEditor)
 	{}
 		
diff --git a/Engine/Plugins/VirtualProduction/DMX/DMXEngine/Source/DMXEditor/Private/Widgets/FixtureType/SDMXFixtureTypeFunctionsEditorFunctionRow.cpp b/Engine/Plugins/VirtualProduction/DMX/DMXEngine/Source/DMXEditor/Private/Widgets/FixtureType/SDMXFixtureTypeFunctionsEditorFunctionRow.cpp
--- a/Engine/Plugins/VirtualProduction/DMX/DMXEngine/Source/DMXEditor/Private/Widgets/FixtureType/SDMXFixtureTypeFunctionsEditorFunctionRow.cpp
+++ b/Engine/Plugins/VirtualProduction/DMX/DMXEngine/Source/DMXEditor/Private/Widgets/FixtureType/SDMXFixtureTypeFunctionsEditorFunctionRow.cpp
@@ -9,13 +9,13 @@
 #include "DragDrop/DMXFixtureFunctionDragDropOp.h"
 #include "DragDrop/DMXFixtureMatrixDragDropOp.h"
 #include "Library/DMXEntityFixtureType.h"
+#include "ScopedTransaction.h"
 #include "Widgets/SNameListPicker.h"
 #include "Widgets/FixtureType/DMXFixtureTypeFunctionsEditorFunctionItem.h"
 #include "Widgets/FixtureType/DMXFixtureTypeFunctionsEditorMatrixItem.h"
 #include "Widgets/FixtureType/SDMXFixtureTypeFunctionsEditor.h"
 #include "Widgets/FixtureType/SDMXFixtureTypeFunctionsEditorMatrixRow.h"
-
-#include "ScopedTransaction.h"
+#include "Widgets/Input/SButton.h"
 #include "Widgets/Input/SCheckBox.h"
 #include "Widgets/Text/SInlineEditableTextBlock.h"
 
@@ -61,7 +61,7 @@ FReply SDMXFixtureTypeFunctionsEditorFunctionRow::OnDragDetected(const FGeometry
 
 TSharedRef<SWidget> SDMXFixtureTypeFunctionsEditorFunctionRow::GenerateWidgetForColumn(const FName& ColumnName)
 {
-	if (ColumnName == FDMXFixtureTypeFunctionsEditorCollumnIDs::Status)
+	if (ColumnName == SDMXFixtureTypeFunctionsEditor::FCollumnID::Status)
 	{
 		return
 			SNew(SBox)
@@ -99,7 +99,7 @@ TSharedRef<SWidget> SDMXFixtureTypeFunctionsEditorFunctionRow::GenerateWidgetFor
 					})
 			];
 	}
-	else if (ColumnName == FDMXFixtureTypeFunctionsEditorCollumnIDs::Channel)
+	else if (ColumnName == SDMXFixtureTypeFunctionsEditor::FCollumnID::Channel)
 	{
 		return
 			SNew(SBorder)
@@ -121,7 +121,7 @@ TSharedRef<SWidget> SDMXFixtureTypeFunctionsEditorFunctionRow::GenerateWidgetFor
 				.IsSelected(IsSelected)
 			];
 	}
-	else if (ColumnName == FDMXFixtureTypeFunctionsEditorCollumnIDs::Name)
+	else if (ColumnName == SDMXFixtureTypeFunctionsEditor::FCollumnID::Name)
 	{
 		return
 			SNew(SBorder)
@@ -162,7 +162,7 @@ TSharedRef<SWidget> SDMXFixtureTypeFunctionsEditorFunctionRow::GenerateWidgetFor
 				.IsSelected(IsSelected)
 			];
 	}
-	else if (ColumnName == FDMXFixtureTypeFunctionsEditorCollumnIDs::Attribute)
+	else if (ColumnName == SDMXFixtureTypeFunctionsEditor::FCollumnID::Attribute)
 	{
 		return
 			SNew(SBorder)
@@ -178,6 +178,24 @@ TSharedRef<SWidget> SDMXFixtureTypeFunctionsEditorFunctionRow::GenerateWidgetFor
 				.OnValueChanged(this, &SDMXFixtureTypeFunctionsEditorFunctionRow::OnUserChangedAttributeName)
 			];
 	}
+	else if (ColumnName == SDMXFixtureTypeFunctionsEditor::FCollumnID::DeleteAttribute)
+	{
+		return
+			SNew(SButton)
+			.ContentPadding(2.0f)
+			.ButtonStyle(FAppStyle::Get(), "HoverHintOnly")
+			.OnClicked(this, &SDMXFixtureTypeFunctionsEditorFunctionRow::OnDeleteAttributeClicked)
+			.ToolTipText(LOCTEXT("RemoveCellAttributeTooltip", "Removes the Cell Attribute"))
+			.ForegroundColor(FSlateColor::UseForeground())
+			.HAlign(HAlign_Center)
+			.VAlign(VAlign_Center)
+			.Content()
+			[
+				SNew(SImage)
+				.Image(FAppStyle::GetBrush("Icons.Delete"))
+				.ColorAndOpacity(FLinearColor(0.6f, 0.6f, 0.6f))
+			];
+	}
 
 	return SNullWidget::NullWidget;
 }
@@ -347,4 +365,11 @@ void SDMXFixtureTypeFunctionsEditorFunctionRow::OnUserChangedAttributeName(FName
 	}
 }
 
+FReply SDMXFixtureTypeFunctionsEditorFunctionRow::OnDeleteAttributeClicked()
+{
+	FunctionItem->DeleteFunction();
+
+	return FReply::Handled();
+}
+
 #undef LOCTEXT_NAMESPACE
diff --git a/Engine/Plugins/VirtualProduction/DMX/DMXEngine/Source/DMXEditor/Private/Widgets/FixtureType/SDMXFixtureTypeFunctionsEditorFunctionRow.h b/Engine/Plugins/VirtualProduction/DMX/DMXEngine/Source/DMXEditor/Private/Widgets/FixtureType/SDMXFixtureTypeFunctionsEditorFunctionRow.h
--- a/Engine/Plugins/VirtualProduction/DMX/DMXEngine/Source/DMXEditor/Private/Widgets/FixtureType/SDMXFixtureTypeFunctionsEditorFunctionRow.h
+++ b/Engine/Plugins/VirtualProduction/DMX/DMXEngine/Source/DMXEditor/Private/Widgets/FixtureType/SDMXFixtureTypeFunctionsEditorFunctionRow.h
@@ -81,6 +81,9 @@ private:
 	/** Called when the Attribute Name changed */
 	void OnUserChangedAttributeName(FName NewValue);
 
+	/** Called when the delete attribute button was clicked */
+	FReply OnDeleteAttributeClicked();
+
 	/** True while this is being dragged in a drag drop op */
 	bool bIsBeingDragged = false;
 
diff --git a/Engine/Plugins/VirtualProduction/DMX/DMXEngine/Source/DMXEditor/Private/Widgets/FixtureType/SDMXFixtureTypeFunctionsEditorMatrixRow.cpp b/Engine/Plugins/VirtualProduction/DMX/DMXEngine/Source/DMXEditor/Private/Widgets/FixtureType/SDMXFixtureTypeFunctionsEditorMatrixRow.cpp
--- a/Engine/Plugins/VirtualProduction/DMX/DMXEngine/Source/DMXEditor/Private/Widgets/FixtureType/SDMXFixtureTypeFunctionsEditorMatrixRow.cpp
+++ b/Engine/Plugins/VirtualProduction/DMX/DMXEngine/Source/DMXEditor/Private/Widgets/FixtureType/SDMXFixtureTypeFunctionsEditorMatrixRow.cpp
@@ -60,7 +60,7 @@ TSharedRef<SWidget> SDMXFixtureTypeFunctionsEditorMatrixRow::GenerateWidgetForCo
 		TWeakObjectPtr<UDMXEntityFixtureType> FixtureType = MatrixItem->GetFixtureType();
 		const int32 ModeIndex = MatrixItem->GetModeIndex();
 
-		if (ColumnName == FDMXFixtureTypeFunctionsEditorCollumnIDs::Status)
+		if (ColumnName == SDMXFixtureTypeFunctionsEditor::FCollumnID::Status)
 		{
 			return
 				SNew(SBorder)
@@ -97,7 +97,7 @@ TSharedRef<SWidget> SDMXFixtureTypeFunctionsEditorMatrixRow::GenerateWidgetForCo
 						})
 				];
 		}
-		else if (ColumnName == FDMXFixtureTypeFunctionsEditorCollumnIDs::Channel)
+		else if (ColumnName == SDMXFixtureTypeFunctionsEditor::FCollumnID::Channel)
 		{
 			return
 				SNew(SBorder)
@@ -117,7 +117,7 @@ TSharedRef<SWidget> SDMXFixtureTypeFunctionsEditorMatrixRow::GenerateWidgetForCo
 					.IsSelected(IsSelected)
 				];
 		}
-		else if (ColumnName == FDMXFixtureTypeFunctionsEditorCollumnIDs::Name)
+		else if (ColumnName == SDMXFixtureTypeFunctionsEditor::FCollumnID::Name)
 		{
 			return
 				SNew(SBorder)
@@ -175,7 +175,7 @@ TSharedRef<SWidget> SDMXFixtureTypeFunctionsEditorMatrixRow::GenerateWidgetForCo
 					]
 				];
 		}
-		else if (ColumnName == FDMXFixtureTypeFunctionsEditorCollumnIDs::Attribute)
+		else if (ColumnName == SDMXFixtureTypeFunctionsEditor::FCollumnID::Attribute)
 		{
 			return
 				SNew(SBorder)
