DMX - Fix CIS Build Engine Localization

diff --git a/Engine/Plugins/VirtualProduction/DMX/DMXEngine/Source/DMXEditor/Private/Widgets/FixturePatch/SDMXReadOnlyFixturePatchListRow.cpp b/Engine/Plugins/VirtualProduction/DMX/DMXEngine/Source/DMXEditor/Private/Widgets/FixturePatch/SDMXReadOnlyFixturePatchListRow.cpp
--- a/Engine/Plugins/VirtualProduction/DMX/DMXEngine/Source/DMXEditor/Private/Widgets/FixturePatch/SDMXReadOnlyFixturePatchListRow.cpp
+++ b/Engine/Plugins/VirtualProduction/DMX/DMXEngine/Source/DMXEditor/Private/Widgets/FixturePatch/SDMXReadOnlyFixturePatchListRow.cpp
@@ -149,5 +149,3 @@ TSharedRef<SWidget> SDMXReadOnlyFixturePatchListRow::GeneratePatchRow()
 			.ColorAndOpacity(FLinearColor::White)
 		];
 }
-
-#undef LOCTEXT_NAMESPACE
