DMX - Fix misleading mistake  in comment of UDMXLibrary::SetMVRGeneralSceneDescription

diff --git a/Engine/Plugins/VirtualProduction/DMX/DMXEngine/Source/DMXRuntime/Public/Library/DMXLibrary.h b/Engine/Plugins/VirtualProduction/DMX/DMXEngine/Source/DMXRuntime/Public/Library/DMXLibrary.h
--- a/Engine/Plugins/VirtualProduction/DMX/DMXEngine/Source/DMXRuntime/Public/Library/DMXLibrary.h
+++ b/Engine/Plugins/VirtualProduction/DMX/DMXEngine/Source/DMXRuntime/Public/Library/DMXLibrary.h
@@ -210,7 +210,7 @@ public:
 	/** 
 	 * Sets the MVR General Scene Description of the Library. 
 	 * Note, this will not add any patches that occur in the new General Scene Description. 
-	 * Note, this will remove any patches that do not occur in the new General Scene Description. 
+	 * Note, this will not remove any patches that do not occur in the new General Scene Description. 
 	 */
 	void SetMVRGeneralSceneDescription(UDMXMVRGeneralSceneDescription* NewGeneralSceneDescription);
 
