Fix CIS Issue 368931: fixed wrong usages of TGuardValue causing compilation error now that its ctor is `nodiscard`

diff --git a/Engine/Plugins/VirtualProduction/DMX/DMXEngine/Source/DMXEditor/Private/Widgets/FixturePatch/SDMXMVRFixtureList.cpp b/Engine/Plugins/VirtualProduction/DMX/DMXEngine/Source/DMXEditor/Private/Widgets/FixturePatch/SDMXMVRFixtureList.cpp
--- a/Engine/Plugins/VirtualProduction/DMX/DMXEngine/Source/DMXEditor/Private/Widgets/FixturePatch/SDMXMVRFixtureList.cpp
+++ b/Engine/Plugins/VirtualProduction/DMX/DMXEngine/Source/DMXEditor/Private/Widgets/FixturePatch/SDMXMVRFixtureList.cpp
@@ -1345,7 +1345,7 @@ void SDMXMVRFixtureList::OnPasteItems()
 		return;
 	}
 
-	TGuardValue<bool>(bChangingDMXLibrary, true);
+	TGuardValue Guard(bChangingDMXLibrary, true);
 	const FText TransactionText = LOCTEXT("PasteFixturePatchesTransaction", "Paste Fixture Patches");
 	const FScopedTransaction PasteTransaction(TransactionText);
 
@@ -1386,7 +1386,7 @@ void SDMXMVRFixtureList::OnDuplicateItems()
 		return;
 	}
 
-	TGuardValue<bool>(bChangingDMXLibrary, true);
+	TGuardValue Guard(bChangingDMXLibrary, true);
 
 	const FText TransactionText = LOCTEXT("DuplicateFixturePatchesTransaction", "Duplicate Fixture Patches");
 	const FScopedTransaction PasteTransaction(TransactionText);
@@ -1443,7 +1443,7 @@ bool SDMXMVRFixtureList::CanDeleteItems() const
 
 void SDMXMVRFixtureList::OnDeleteItems()
 {
-	TGuardValue<bool>(bChangingDMXLibrary, true);
+	TGuardValue Guard(bChangingDMXLibrary, true);
 
 	const TArray<TSharedPtr<FDMXMVRFixtureListItem>> SelectedItems = ListView->GetSelectedItems();
 
diff --git a/Engine/Plugins/VirtualProduction/DMX/DMXPixelMapping/Source/DMXPixelMappingEditor/Private/Toolkits/DMXPixelMappingToolkit.cpp b/Engine/Plugins/VirtualProduction/DMX/DMXPixelMapping/Source/DMXPixelMappingEditor/Private/Toolkits/DMXPixelMappingToolkit.cpp
--- a/Engine/Plugins/VirtualProduction/DMX/DMXPixelMapping/Source/DMXPixelMappingEditor/Private/Toolkits/DMXPixelMappingToolkit.cpp
+++ b/Engine/Plugins/VirtualProduction/DMX/DMXPixelMapping/Source/DMXPixelMappingEditor/Private/Toolkits/DMXPixelMappingToolkit.cpp
@@ -313,7 +313,7 @@ TArray<UDMXPixelMappingBaseComponent*> FDMXPixelMappingToolkit::CreateComponents
 	TArray<UDMXPixelMappingBaseComponent*> NewComponents;
 	if (Templates.Num() > 0)
 	{
-		TGuardValue<bool>(bAddingComponents, true);
+		TGuardValue Guard(bAddingComponents, true);
 
 		if (ensureMsgf(RootComponent && Target, TEXT("Tried to create components from template but RootComponent or Target were invalid.")))
 		{
@@ -344,7 +344,7 @@ void FDMXPixelMappingToolkit::DeleteSelectedComponents()
 		return;
 	}
 
-	TGuardValue<bool>(bRemovingComponents, true);
+	TGuardValue Guard(bRemovingComponents, true);
 
 	TSet<FDMXPixelMappingComponentReference> ParentComponentReferences;
 	for (const FDMXPixelMappingComponentReference& SelectedComponentReference : SelectedComponents)
diff --git a/Engine/Plugins/VirtualProduction/DMX/DMXPixelMapping/Source/DMXPixelMappingRuntime/Private/Components/DMXPixelMappingMatrixComponent.cpp b/Engine/Plugins/VirtualProduction/DMX/DMXPixelMapping/Source/DMXPixelMappingRuntime/Private/Components/DMXPixelMappingMatrixComponent.cpp
--- a/Engine/Plugins/VirtualProduction/DMX/DMXPixelMapping/Source/DMXPixelMappingRuntime/Private/Components/DMXPixelMappingMatrixComponent.cpp
+++ b/Engine/Plugins/VirtualProduction/DMX/DMXPixelMapping/Source/DMXPixelMappingRuntime/Private/Components/DMXPixelMappingMatrixComponent.cpp
@@ -500,8 +500,6 @@ PRAGMA_ENABLE_DEPRECATION_WARNINGS
 
 void UDMXPixelMappingMatrixComponent::HandleMatrixChanged()
 {
-	TGuardValue<bool>(bIsUpdatingChildren, true);
-
 	UDMXPixelMapping* PixelMapping = GetPixelMapping();
 	UDMXPixelMappingRootComponent* RootComponent = PixelMapping ? PixelMapping->GetRootComponent() : nullptr;
 	if (!RootComponent)
diff --git a/Engine/Plugins/VirtualProduction/DMX/DMXPixelMapping/Source/DMXPixelMappingRuntime/Public/Components/DMXPixelMappingMatrixComponent.h b/Engine/Plugins/VirtualProduction/DMX/DMXPixelMapping/Source/DMXPixelMappingRuntime/Public/Components/DMXPixelMappingMatrixComponent.h
--- a/Engine/Plugins/VirtualProduction/DMX/DMXPixelMapping/Source/DMXPixelMappingRuntime/Public/Components/DMXPixelMappingMatrixComponent.h
+++ b/Engine/Plugins/VirtualProduction/DMX/DMXPixelMapping/Source/DMXPixelMappingRuntime/Public/Components/DMXPixelMappingMatrixComponent.h
@@ -105,9 +105,6 @@ protected:
 #endif // WITH_EDITORONLY_DATA
 
 private:
-	/** True while the component is updating its children */
-	bool bIsUpdatingChildren = false;
-
 	/** Position before it was changed */
 	FVector2D PreEditChangePosition;
 
