DMX Control Console - Create asset editor toolkit for Control Console

diff --git a/Engine/Plugins/VirtualProduction/DMX/DMXEngine/Source/DMXEditor/Private/Widgets/FixturePatch/SDMXReadOnlyFixturePatchList.cpp b/Engine/Plugins/VirtualProduction/DMX/DMXEngine/Source/DMXEditor/Private/Widgets/FixturePatch/SDMXReadOnlyFixturePatchList.cpp
--- a/Engine/Plugins/VirtualProduction/DMX/DMXEngine/Source/DMXEditor/Private/Widgets/FixturePatch/SDMXReadOnlyFixturePatchList.cpp
+++ b/Engine/Plugins/VirtualProduction/DMX/DMXEngine/Source/DMXEditor/Private/Widgets/FixturePatch/SDMXReadOnlyFixturePatchList.cpp
@@ -193,7 +193,11 @@ TArray<UDMXEntityFixturePatch*> SDMXReadOnlyFixturePatchList::GetSelectedFixture
 
 void SDMXReadOnlyFixturePatchList::SetDMXLibrary(UDMXLibrary* InDMXLibrary)
 {
-	WeakDMXLibrary = InDMXLibrary;
+	if (WeakDMXLibrary.Get() != InDMXLibrary)
+	{
+		WeakDMXLibrary = InDMXLibrary;
+		RequestRefresh();
+	}
 }
 
 void SDMXReadOnlyFixturePatchList::SetExcludedFixturePatches(const TArray<UDMXEntityFixturePatch*>& NewExcludedFixturePatches)
@@ -394,12 +398,14 @@ TSharedRef<SWidget> SDMXReadOnlyFixturePatchList::GenerateHeaderRowFilterMenu()
 		
 		auto AddMenuEntryLambda = [this, &Section](const FName& Name, const FText& Label, const FText& ToolTip, const FName& ColumnID)
 		{
-			Section.AddMenuEntry(
+			Section.AddMenuEntry
+			(
 				Name,
 				Label,
 				ToolTip,
 				FSlateIcon(),
-				FUIAction(
+				FUIAction
+				(
 					FExecuteAction::CreateSP(this, &SDMXReadOnlyFixturePatchList::ToggleColumnShowState, ColumnID),
 					FCanExecuteAction(),
 					FIsActionChecked::CreateSP(this, &SDMXReadOnlyFixturePatchList::IsColumnShown, ColumnID)
diff --git a/Engine/Plugins/VirtualProduction/DMX/DMXEngine/Source/DMXEditor/Public/DMXEditorSettings.h b/Engine/Plugins/VirtualProduction/DMX/DMXEngine/Source/DMXEditor/Public/DMXEditorSettings.h
--- a/Engine/Plugins/VirtualProduction/DMX/DMXEngine/Source/DMXEditor/Public/DMXEditorSettings.h
+++ b/Engine/Plugins/VirtualProduction/DMX/DMXEngine/Source/DMXEditor/Public/DMXEditorSettings.h
@@ -2,7 +2,6 @@
 
 #pragma once
 
-#include "CoreMinimal.h"
 #include "Containers/UnrealString.h"
 #include "Engine/DeveloperSettings.h"
 #include "Widgets/Views/SHeaderRow.h"
@@ -189,9 +188,15 @@ public:
 	PRAGMA_ENABLE_DEPRECATION_WARNINGS
 #endif // WITH_EDITORONLY_DATA
 
+	// Control Console
+public:
+	/** Path to the last control console opened in the editor */
+	UPROPERTY(Config)
+	FString LastOpenedControlConsolePath;
+
 	// Channels Monitor
 public:
-	/** The Universe ID to be monitored in the Channels Monitor  */
+	/** The Universe ID to be monitored in the Channels Monitor */
 	UPROPERTY(Config)
 	int32 ChannelsMonitorUniverseID = 1;
 
diff --git a/Engine/Plugins/VirtualProduction/DMX/DMXEngine/Source/DMXEditor/Public/Widgets/SDMXReadOnlyFixturePatchList.h b/Engine/Plugins/VirtualProduction/DMX/DMXEngine/Source/DMXEditor/Public/Widgets/SDMXReadOnlyFixturePatchList.h
--- a/Engine/Plugins/VirtualProduction/DMX/DMXEngine/Source/DMXEditor/Public/Widgets/SDMXReadOnlyFixturePatchList.h
+++ b/Engine/Plugins/VirtualProduction/DMX/DMXEngine/Source/DMXEditor/Public/Widgets/SDMXReadOnlyFixturePatchList.h
@@ -176,6 +176,9 @@ protected:
 	/** Called to generate a row in the list */
 	virtual TSharedRef<ITableRow> OnGenerateRow(TSharedPtr<FDMXReadOnlyFixturePatchListItem> InItem, const TSharedRef<STableViewBase>& OwnerTable);
 
+	/** Toggles show state of the column with the given ID */
+	virtual void ToggleColumnShowState(const FName ColumnID);
+
 private:
 	/** Generates a filter menu fot the Header Row of the List */
 	TSharedRef<SWidget> GenerateHeaderRowFilterMenu();
@@ -201,9 +204,6 @@ private:
 	/** Called when a Fixture Type changed */
 	void OnFixtureTypeChanged(const UDMXEntityFixtureType* FixtureType);
 
-	/** Toggles show state of the column with the given ID */
-	void ToggleColumnShowState(const FName ColumnID);
-
 	/** Gets show state of the column with the given ID */
 	bool IsColumnShown(const FName ColumnID) const;
 
