Mark RemoveIf with nodiscard and fix up all erroneous usages

diff --git a/Engine/Plugins/MovieScene/MovieRenderPipeline/Source/MovieRenderPipelineEditor/Private/Tests/MovieGraphEditorTestUtilities.cpp b/Engine/Plugins/MovieScene/MovieRenderPipeline/Source/MovieRenderPipelineEditor/Private/Tests/MovieGraphEditorTestUtilities.cpp
--- a/Engine/Plugins/MovieScene/MovieRenderPipeline/Source/MovieRenderPipelineEditor/Private/Tests/MovieGraphEditorTestUtilities.cpp
+++ b/Engine/Plugins/MovieScene/MovieRenderPipeline/Source/MovieRenderPipelineEditor/Private/Tests/MovieGraphEditorTestUtilities.cpp
@@ -56,11 +56,11 @@ namespace UE::MovieGraph::Private::Tests
 
 	void RemoveAbstractClasses(TArray<UClass*>& InClassArray)
 	{
-		Algo::StableRemoveIf(InClassArray,
+		InClassArray.SetNum(Algo::StableRemoveIf(InClassArray,
 			[](const UClass* Class)
 			{
 				return !Class ||Class->HasAnyClassFlags(CLASS_Abstract);
-			});
+			}));
 	}
 
 	TArray<UClass*> GetNativeClasses(UClass* BaseClass, bool bRecursive)
diff --git a/Engine/Plugins/MovieScene/MovieRenderPipeline/Source/MovieRenderPipelineEditor/Private/Tests/MovieGraphTests.cpp b/Engine/Plugins/MovieScene/MovieRenderPipeline/Source/MovieRenderPipelineEditor/Private/Tests/MovieGraphTests.cpp
--- a/Engine/Plugins/MovieScene/MovieRenderPipeline/Source/MovieRenderPipelineEditor/Private/Tests/MovieGraphTests.cpp
+++ b/Engine/Plugins/MovieScene/MovieRenderPipeline/Source/MovieRenderPipelineEditor/Private/Tests/MovieGraphTests.cpp
@@ -29,11 +29,11 @@ namespace UE::MovieGraph::Private::Tests
 			UMovieGraphSubgraphNode::StaticClass()	// Will be tested separately in another test
 		};
 
-		Algo::StableRemoveIf(InClassArray,
+		InClassArray.SetNum(Algo::StableRemoveIf(InClassArray,
 			[&](const UClass* Class)
 			{
 				return UndesirableClasses.Contains(Class);
-			});
+			}));
 	}
 
 	void OpenGraphConfigInEditorIfDesired(UMovieGraphConfig* InGraphConfig)
diff --git a/Engine/Plugins/VirtualProduction/DMX/DMXEngine/Source/DMXRuntime/Private/MVR/DMXMVRSceneActor.cpp b/Engine/Plugins/VirtualProduction/DMX/DMXEngine/Source/DMXRuntime/Private/MVR/DMXMVRSceneActor.cpp
--- a/Engine/Plugins/VirtualProduction/DMX/DMXEngine/Source/DMXRuntime/Private/MVR/DMXMVRSceneActor.cpp
+++ b/Engine/Plugins/VirtualProduction/DMX/DMXEngine/Source/DMXRuntime/Private/MVR/DMXMVRSceneActor.cpp
@@ -370,7 +370,7 @@ void ADMXMVRSceneActor::UpdateGDTFToDefaultActorClasses(const TSharedRef<FDMXMVR
 	}
 
 	// Remove those GDTFs that are no longer used by any actor
-	Algo::RemoveIf(GDTFToDefaultActorClasses, [this](const FDMXMVRSceneGDTFToActorClassPair& GDTFToActorClassPair)
+	GDTFToDefaultActorClasses.SetNum(Algo::RemoveIf(GDTFToDefaultActorClasses, [this](const FDMXMVRSceneGDTFToActorClassPair& GDTFToActorClassPair)
 		{
 			const TSoftObjectPtr<AActor>* ActorThatUsesGDTF = Algo::FindByPredicate(RelatedActors, [this, GDTFToActorClassPair](const TSoftObjectPtr<AActor>& Actor)
 				{
@@ -381,7 +381,7 @@ void ADMXMVRSceneActor::UpdateGDTFToDefaultActorClasses(const TSharedRef<FDMXMVR
 					return GDTFToActorClassPair.GDTF == GDTF;
 				});
 			return ActorThatUsesGDTF == nullptr;
-		});
+		}));
 }
 
 #if WITH_EDITOR
diff --git a/Engine/Source/Runtime/Core/Public/Algo/RemoveIf.h b/Engine/Source/Runtime/Core/Public/Algo/RemoveIf.h
--- a/Engine/Source/Runtime/Core/Public/Algo/RemoveIf.h
+++ b/Engine/Source/Runtime/Core/Public/Algo/RemoveIf.h
@@ -18,7 +18,7 @@ namespace Algo
 	 * @return The index of the first element after those which were not removed.
 	 */
 	template <typename RangeType, typename Predicate>
-	int32 RemoveIf(RangeType& Range, Predicate Pred)
+	[[nodiscard]] int32 RemoveIf(RangeType& Range, Predicate Pred)
 	{
 		auto* First = GetData(Range);
 		auto* Last  = First + GetNum(Range);
@@ -77,7 +77,7 @@ namespace Algo
 	 * @return The index of the first element after those which were not removed.
 	 */
 	template <typename RangeType, typename Predicate>
-	int32 StableRemoveIf(RangeType& Range, Predicate Pred)
+	[[nodiscard]] int32 StableRemoveIf(RangeType& Range, Predicate Pred)
 	{
 		auto* First = GetData(Range);
 		auto* Last  = First + GetNum(Range);
diff --git a/Engine/Source/Runtime/LiveLinkMessageBusFramework/Private/LiveLinkProvider.cpp b/Engine/Source/Runtime/LiveLinkMessageBusFramework/Private/LiveLinkProvider.cpp
--- a/Engine/Source/Runtime/LiveLinkMessageBusFramework/Private/LiveLinkProvider.cpp
+++ b/Engine/Source/Runtime/LiveLinkMessageBusFramework/Private/LiveLinkProvider.cpp
@@ -95,7 +95,10 @@ void FLiveLinkProvider::ValidateConnections()
 	FConnectionValidator Validator;
 
 	TArray<FMessageAddress> RemovedConnections;
-	Algo::RemoveIf(ConnectedAddresses, [this, &Validator, &RemovedConnections](const FTrackedAddress& Address) mutable
+
+	// Using SetNumUninitialized because FTrackedAddress does not have a default constructor, resulting in SetNum not
+	// compiling (due to the DefaultConstructItems<> usage). Uninitialized is not unsafe here, because we're shrinking.
+	ConnectedAddresses.SetNumUninitialized(Algo::RemoveIf(ConnectedAddresses, [this, &Validator, &RemovedConnections](const FTrackedAddress& Address) mutable
 	{
 		if (!Validator(Address))
 	    {
@@ -103,7 +106,7 @@ void FLiveLinkProvider::ValidateConnections()
 			return true;
 	    }
 		return false;
-	});
+	}));
 
 	if (RemovedConnections.Num() > 0)
 	{
diff --git a/Engine/Source/Runtime/MovieSceneTracks/Private/Systems/MovieSceneEventSystems.cpp b/Engine/Source/Runtime/MovieSceneTracks/Private/Systems/MovieSceneEventSystems.cpp
--- a/Engine/Source/Runtime/MovieSceneTracks/Private/Systems/MovieSceneEventSystems.cpp
+++ b/Engine/Source/Runtime/MovieSceneTracks/Private/Systems/MovieSceneEventSystems.cpp
@@ -92,7 +92,7 @@ void UMovieSceneEventSystem::TriggerAllEvents()
 		{
 			if (bSkipTrigger)
 			{
-				Algo::RemoveIf(Pair.Value, [SkipUntil](FMovieSceneEventTriggerData& Data) { return Data.RootTime <= SkipUntil; });
+				Pair.Value.SetNum(Algo::RemoveIf(Pair.Value, [SkipUntil](FMovieSceneEventTriggerData& Data) { return Data.RootTime <= SkipUntil; }));
 			}
 			Algo::SortBy(Pair.Value, &FMovieSceneEventTriggerData::RootTime);
 		}
@@ -100,7 +100,7 @@ void UMovieSceneEventSystem::TriggerAllEvents()
 		{
 			if (bSkipTrigger)
 			{
-				Algo::RemoveIf(Pair.Value, [SkipUntil](FMovieSceneEventTriggerData& Data) { return Data.RootTime >= SkipUntil; });
+				Pair.Value.SetNum(Algo::RemoveIf(Pair.Value, [SkipUntil](FMovieSceneEventTriggerData& Data) { return Data.RootTime >= SkipUntil; }));
 			}
 			Algo::SortBy(Pair.Value, &FMovieSceneEventTriggerData::RootTime, TGreater<>());
 		}
