DMX - Fix UMovieSceneDMXLibrarySection::GetFixturePatches returns invalid instead of valid patches

diff --git a/Engine/Plugins/VirtualProduction/DMX/DMXEngine/Source/DMXRuntime/Private/Sequencer/MovieSceneDMXLibrarySection.cpp b/Engine/Plugins/VirtualProduction/DMX/DMXEngine/Source/DMXRuntime/Private/Sequencer/MovieSceneDMXLibrarySection.cpp
--- a/Engine/Plugins/VirtualProduction/DMX/DMXEngine/Source/DMXRuntime/Private/Sequencer/MovieSceneDMXLibrarySection.cpp
+++ b/Engine/Plugins/VirtualProduction/DMX/DMXEngine/Source/DMXRuntime/Private/Sequencer/MovieSceneDMXLibrarySection.cpp
@@ -590,12 +590,10 @@ TArray<UDMXEntityFixturePatch*> UMovieSceneDMXLibrarySection::GetFixturePatches(
 	for (const FDMXFixturePatchChannel& PatchRef : FixturePatchChannels)
 	{
 		// Add only valid patches
-		if (UDMXEntityFixturePatch* Patch = PatchRef.Reference.GetFixturePatch())
+		UDMXEntityFixturePatch* Patch = PatchRef.Reference.GetFixturePatch();
+		if (IsValid(Patch))
 		{
-			if (!Patch->IsValidLowLevelFast())
-			{
-				Result.Add(Patch);
-			}
+			Result.Add(Patch);
 		}
 	}
 
