DMX - Conflict monitor now only prints unique primary conflict names, unique channels, fix an issue where it sometimes was not printing all conflicting channels

diff --git a/Engine/Plugins/VirtualProduction/DMX/DMXEngine/Source/DMXEditor/Private/Widgets/Monitors/DMXConflictMonitorConflictModel.cpp b/Engine/Plugins/VirtualProduction/DMX/DMXEngine/Source/DMXEditor/Private/Widgets/Monitors/DMXConflictMonitorConflictModel.cpp
--- a/Engine/Plugins/VirtualProduction/DMX/DMXEngine/Source/DMXEditor/Private/Widgets/Monitors/DMXConflictMonitorConflictModel.cpp
+++ b/Engine/Plugins/VirtualProduction/DMX/DMXEngine/Source/DMXEditor/Private/Widgets/Monitors/DMXConflictMonitorConflictModel.cpp
@@ -2,7 +2,6 @@
 
 #include "DMXConflictMonitorConflictModel.h"
 
-#include "Algo/AnyOf.h"
 #include "Algo/Sort.h"
 #include "DMXEditorSettings.h"
 #include "Framework/Text/ITextDecorator.h"
@@ -63,6 +62,7 @@ namespace UE::DMX
 		const uint8 Depth = EditorSettings->ConflictMonitorSettings.Depth;
 
 		int32 CountStack = 0;
+		TArray<FString> PrimaryNames;
 		for (const FName& NameStack : NameStacks)
 		{
 			++CountStack;
@@ -80,39 +80,49 @@ namespace UE::DMX
 			FString DetailsString;
 			for (int32 SubstringIndex = 0; SubstringIndex < Substrings.Num(); SubstringIndex++)
 			{
-				const bool bFirstSubstring = &Substrings[SubstringIndex] == &Substrings[0];
+				const bool bPrimaryName = &Substrings[SubstringIndex] == &Substrings[0];
 				const FString CleanSubstring = FPaths::GetBaseFilename(Substrings[SubstringIndex]);
 
 				// Create the title
-				if (bFirstSubstring && bFirstNameInStack)
+				if (bPrimaryName && bFirstNameInStack)
 				{
-					FString FirstString = LOCTEXT("ConflictDetected", "Send DMX Conflict: ").ToString() + CleanSubstring;
-					FirstString = StyleString(FirstString, TEXT("ConflictLog.Title"));
-					Title.Append(FirstString);
+					FString FirstTraceInStack = LOCTEXT("ConflictDetected", "Send DMX Conflict: ").ToString() + CleanSubstring;
+					FirstTraceInStack = StyleString(FirstTraceInStack, TEXT("ConflictLog.Title"));
+
+					PrimaryNames.Add(CleanSubstring);
+					Title.Append(FirstTraceInStack);
 				}
-				else if (bFirstSubstring)
+				else if (bPrimaryName)
 				{
-					FString NextConflictString = LOCTEXT("AppendConflict", " / ").ToString() + CleanSubstring;
-					NextConflictString = StyleString(NextConflictString, TEXT("ConflictLog.Title"));
-					Title.Append(NextConflictString);
+					FString FirstTraceInStack = TEXT(" / ") + CleanSubstring;
+					FirstTraceInStack = StyleString(FirstTraceInStack, TEXT("ConflictLog.Title"));
+					
+					// Show primary names of conflicting objects only once, even if there are many conflicts.
+					// E.g. "Obj1 / Obj2", not "Obj1 / Obj2 / Obj2".
+					if (!PrimaryNames.Contains(CleanSubstring))
+					{
+						PrimaryNames.Add(CleanSubstring);
+						Title.Append(FirstTraceInStack);
+					}
 				}
 
+				// Don't add details if depth is 1, details would be the same as the title.
 				if (Depth == 1)
 				{
 					break;
 				}
 
 				// Create details
-				if (bFirstSubstring)
+				if (bPrimaryName)
 				{
 					const FString Number = TEXT("\t\t") + FString::FromInt(CountStack) + TEXT(". ");
-					const FString TraceString = StyleString(Number + CleanSubstring, TEXT("ConflictLog.Warning"));
-					DetailsString = TraceString;
+					const FString FirstTraceInStack = StyleString(Number + CleanSubstring, TEXT("ConflictLog.Warning"));
+					DetailsString = FirstTraceInStack;
 				}
 				else
 				{
-					const FString TraceString = StyleString(TEXT(" -> ") + CleanSubstring, TEXT("ConflictLog.Warning"));
-					DetailsString.Append(TraceString);
+					const FString NextTraceInStack = StyleString(TEXT(" -> ") + CleanSubstring, TEXT("ConflictLog.Warning"));
+					DetailsString.Append(NextTraceInStack);
 				}
 			}
 
@@ -175,23 +185,17 @@ namespace UE::DMX
 		{
 			for (const TSharedRef<FDMXMonitoredOutboundDMXData>& Other : Conflicts)
 			{
-				// Only lexically lower, avoids duplicates and self
-				if (Other->Trace.Compare(Conflict->Trace) >= 0.0)
-				{
-					continue;
-				}
-
 				for (const TTuple<int32, uint8>& ChannelToValuePair : Conflict->ChannelToValueMap)
 				{
 					if (Other->ChannelToValueMap.Contains(ChannelToValuePair.Key))
 					{
-						Channels.Add(ChannelToValuePair.Key);
+						Channels.AddUnique(ChannelToValuePair.Key);
 					}
 				}
 			}
 		}
 		Algo::Sort(Channels);
-
+		
 		FString ChannelsString;
 		constexpr int32 MaxIndex = 31;
 		for (int32 ChannelIndex = 0; ChannelIndex < Channels.Num(); ChannelIndex++)
diff --git a/Engine/Plugins/VirtualProduction/DMX/DMXEngine/Source/DMXEditor/Private/Widgets/Monitors/SDMXConflictMonitor.cpp b/Engine/Plugins/VirtualProduction/DMX/DMXEngine/Source/DMXEditor/Private/Widgets/Monitors/SDMXConflictMonitor.cpp
--- a/Engine/Plugins/VirtualProduction/DMX/DMXEngine/Source/DMXEditor/Private/Widgets/Monitors/SDMXConflictMonitor.cpp
+++ b/Engine/Plugins/VirtualProduction/DMX/DMXEngine/Source/DMXEditor/Private/Widgets/Monitors/SDMXConflictMonitor.cpp
@@ -333,13 +333,13 @@ namespace UE::DMX
 
 	void SDMXConflictMonitor::UpdateStatusInfo()
 	{
-		if (!GetCanTick())
+		if (bIsPaused)
 		{
-			StatusInfo = EDMXConflictMonitorStatusInfo::Idle;
+			StatusInfo = EDMXConflictMonitorStatusInfo::Paused;
 		}
-		else if (bIsPaused)
+		else if (!GetCanTick())
 		{
-			StatusInfo = EDMXConflictMonitorStatusInfo::Paused;
+			StatusInfo = EDMXConflictMonitorStatusInfo::Idle;
 		}
 		else if (Models.IsEmpty())
 		{
@@ -350,7 +350,6 @@ namespace UE::DMX
 			StatusInfo = EDMXConflictMonitorStatusInfo::Conflict;
 		}
 	}
-
 }
 
 #undef LOCTEXT_NAMESPACE
diff --git a/Engine/Plugins/VirtualProduction/DMX/DMXEngine/Source/DMXEditor/Private/Widgets/Monitors/SDMXConflictMonitorToolbar.cpp b/Engine/Plugins/VirtualProduction/DMX/DMXEngine/Source/DMXEditor/Private/Widgets/Monitors/SDMXConflictMonitorToolbar.cpp
--- a/Engine/Plugins/VirtualProduction/DMX/DMXEngine/Source/DMXEditor/Private/Widgets/Monitors/SDMXConflictMonitorToolbar.cpp
+++ b/Engine/Plugins/VirtualProduction/DMX/DMXEngine/Source/DMXEditor/Private/Widgets/Monitors/SDMXConflictMonitorToolbar.cpp
@@ -50,7 +50,7 @@ namespace UE::DMX
 					+ SHorizontalBox::Slot()
 					.VAlign(VAlign_Center)
 					.AutoWidth()
-					.Padding(2.f, 0.f)
+					.Padding(2.f, 0.f, 6.f, 0.f)
 					[
 						SNew(SCircularThrobber)
 						.Radius(8.5f)
