DMX - DMXFixturePatch SendDefaultValues resp. SendZeroValues no longer creates traces to avoid confusing conflict logging, clearer comments

diff --git a/Engine/Plugins/VirtualProduction/DMX/DMXEngine/Source/DMXRuntime/Private/Library/DMXEntityFixturePatch.cpp b/Engine/Plugins/VirtualProduction/DMX/DMXEngine/Source/DMXRuntime/Private/Library/DMXEntityFixturePatch.cpp
--- a/Engine/Plugins/VirtualProduction/DMX/DMXEngine/Source/DMXRuntime/Private/Library/DMXEntityFixturePatch.cpp
+++ b/Engine/Plugins/VirtualProduction/DMX/DMXEngine/Source/DMXRuntime/Private/Library/DMXEntityFixturePatch.cpp
@@ -1291,8 +1291,6 @@ void UDMXEntityFixturePatch::SendResetDataToAllAttributes(bool bUseDefaultValues
 	}
 
 	// Send to the library's output ports
-	UE_DMX_SCOPED_TRACE_SENDDMX(DMXLibrary->GetFName());
-	UE_DMX_SCOPED_TRACE_SENDDMX(*GetDisplayName());
 	for (const FDMXOutputPortSharedRef& OutputPort : DMXLibrary->GetOutputPorts())
 	{
 		OutputPort->SendDMX(UniverseID, DMXChannelToValueMap);
diff --git a/Engine/Plugins/VirtualProduction/DMX/DMXEngine/Source/DMXRuntime/Public/Library/DMXEntityFixturePatch.h b/Engine/Plugins/VirtualProduction/DMX/DMXEngine/Source/DMXRuntime/Public/Library/DMXEntityFixturePatch.h
--- a/Engine/Plugins/VirtualProduction/DMX/DMXEngine/Source/DMXRuntime/Public/Library/DMXEntityFixturePatch.h
+++ b/Engine/Plugins/VirtualProduction/DMX/DMXEngine/Source/DMXRuntime/Public/Library/DMXEntityFixturePatch.h
@@ -114,10 +114,16 @@ public:
 	UFUNCTION(BlueprintCallable, Category = "DMX")
 	void SendDMX(TMap<FDMXAttributeName, int32> AttributeMap);
 
-	/** Sends the default value for all attributes, including matrix attributes */
+	/** 
+	 * Sends the default value for all attributes, including matrix attributes. 
+	 * Note, this call will not raise send dmx traces.
+	 */
 	void SendDefaultValues();
 
-	/** Sends zeroes for all attributes, including matrix attributes */
+	/** 
+	 * Sends zeroes for all attributes, including matrix attributes.
+	 * Note, this call will not raise send dmx traces.
+	 */
 	void SendZeroValues();
 
 	/** 
@@ -482,7 +488,10 @@ public:
 	bool GetAllMatrixCells(TArray<FDMXCell>& Cells);
 
 private:
-	/** Sends reset data to all channels. If bUseDefaultValues is true, sends default values. If it's false it sends zeroes. */
+	/** 
+	 * Sends reset data to all channels. If bUseDefaultValues is true, sends default values. If it's false it sends zeroes.
+	 * Note, this call will not raise send DMX traces.
+	 */
 	void SendResetDataToAllAttributes(bool bUseDefaultValues);
 
 	/** Called when a Fixture Type changed */
