DMX - Fix a BP Compile error when upgrading bluperints that use SendDMXToOutputPort, introduced with conflict logging

diff --git a/Engine/Plugins/VirtualProduction/DMX/DMXEngine/Source/DMXRuntime/Private/DMXSubsystem.cpp b/Engine/Plugins/VirtualProduction/DMX/DMXEngine/Source/DMXRuntime/Private/DMXSubsystem.cpp
--- a/Engine/Plugins/VirtualProduction/DMX/DMXEngine/Source/DMXRuntime/Private/DMXSubsystem.cpp
+++ b/Engine/Plugins/VirtualProduction/DMX/DMXEngine/Source/DMXRuntime/Private/DMXSubsystem.cpp
@@ -149,7 +149,7 @@ void UDMXSubsystem::SendDMXRaw(FDMXProtocolName SelectedProtocol, int32 RemoteUn
 	OutResult = EDMXSendResult::Success;
 }
 
-void UDMXSubsystem::SendDMXToOutputPort(FDMXOutputPortReference OutputPortReference, TMap<int32, uint8> ChannelToValueMap, int32 LocalUniverse, const FName& SenderName)
+void UDMXSubsystem::SendDMXToOutputPort(FDMXOutputPortReference OutputPortReference, TMap<int32, uint8> ChannelToValueMap, int32 LocalUniverse)
 {
 	const FGuid& PortGuid = OutputPortReference.GetPortGuid();
 	const FDMXOutputPortSharedRef* OutputPortPtr = FDMXPortManager::Get().GetOutputPorts().FindByPredicate([PortGuid](const FDMXOutputPortSharedRef& OutputPort) {
diff --git a/Engine/Plugins/VirtualProduction/DMX/DMXEngine/Source/DMXRuntime/Public/DMXSubsystem.h b/Engine/Plugins/VirtualProduction/DMX/DMXEngine/Source/DMXRuntime/Public/DMXSubsystem.h
--- a/Engine/Plugins/VirtualProduction/DMX/DMXEngine/Source/DMXRuntime/Public/DMXSubsystem.h
+++ b/Engine/Plugins/VirtualProduction/DMX/DMXEngine/Source/DMXRuntime/Public/DMXSubsystem.h
@@ -44,7 +44,7 @@ public:
 
 	/**  Sends DMX Values over the Output Port */
 	UFUNCTION(BlueprintCallable, Category = "DMX")
-	static void SendDMXToOutputPort(FDMXOutputPortReference OutputPortReference, TMap<int32, uint8> ChannelToValueMap, int32 LocalUniverse = 1, const FName& SenderName = NAME_None);
+	static void SendDMXToOutputPort(FDMXOutputPortReference OutputPortReference, TMap<int32, uint8> ChannelToValueMap, int32 LocalUniverse = 1);
 
 	/**  DEPRECATED 4.27 */
 	UFUNCTION(BlueprintCallable, Category = "DMX", meta = (DeprecatedFunction, DeprecationMessage = "Deprecated 4.27. Use GetDMXDataFromInputPort or GetDMXDataFromOutputPort instead."))
