DMX - Fix an issue where 'Stop Sends Zero Values' resp. 'Stop Sends Default Values' option did not work for matrix fixtures both in pixel mapping and control console

diff --git a/Engine/Plugins/VirtualProduction/DMX/DMXEngine/Source/DMXRuntime/Private/Library/DMXEntityFixturePatch.cpp b/Engine/Plugins/VirtualProduction/DMX/DMXEngine/Source/DMXRuntime/Private/Library/DMXEntityFixturePatch.cpp
--- a/Engine/Plugins/VirtualProduction/DMX/DMXEngine/Source/DMXRuntime/Private/Library/DMXEntityFixturePatch.cpp
+++ b/Engine/Plugins/VirtualProduction/DMX/DMXEngine/Source/DMXRuntime/Private/Library/DMXEntityFixturePatch.cpp
@@ -1286,6 +1286,8 @@ void UDMXEntityFixturePatch::SendResetDataToAllAttributes(bool bUseDefaultValues
 				{
 					DMXChannelToValueMap.Add(AbsoluteCellStartingChannel + AttributeOffset + ByteIndex, ByteArray[ByteIndex]);
 				}
+
+				AttributeOffset += CellAttribute.GetNumChannels();
 			}
 		}
 	}
