DMX - Cleanup DMXLibraryTrackEditor includes

#rnx

diff --git a/Engine/Plugins/VirtualProduction/DMX/DMXEngine/Source/DMXEditor/Private/Sequencer/DMXLibraryTrackEditor.cpp b/Engine/Plugins/VirtualProduction/DMX/DMXEngine/Source/DMXEditor/Private/Sequencer/DMXLibraryTrackEditor.cpp
--- a/Engine/Plugins/VirtualProduction/DMX/DMXEngine/Source/DMXEditor/Private/Sequencer/DMXLibraryTrackEditor.cpp
+++ b/Engine/Plugins/VirtualProduction/DMX/DMXEngine/Source/DMXEditor/Private/Sequencer/DMXLibraryTrackEditor.cpp
@@ -2,26 +2,20 @@
 
 #include "DMXLibraryTrackEditor.h"
 
-#include "Library/DMXLibrary.h"
-#include "Library/DMXEntityFixturePatch.h"
-#include "Sequencer/MovieSceneDMXLibraryTrack.h"
-#include "Sequencer/MovieSceneDMXLibrarySection.h"
-#include "Sequencer/DMXLibrarySection.h"
-
-#include "Sections/MovieSceneParameterSection.h"
-#include "MVVM/Views/ViewUtilities.h"
-#include "MovieSceneSection.h"
-
-#include "Framework/MultiBox/MultiBoxBuilder.h"
 #include "Algo/Sort.h"
-#include "Styling/AppStyle.h"
-#include "Styling/SlateIconFinder.h"
 #include "ContentBrowserModule.h"
+#include "Framework/Application/SlateApplication.h"
+#include "Framework/MultiBox/MultiBoxBuilder.h"
 #include "IContentBrowserSingleton.h"
+#include "Library/DMXEntityFixturePatch.h"
+#include "Library/DMXLibrary.h"
+#include "MovieSceneSection.h"
+#include "MVVM/Views/ViewUtilities.h"
+#include "Sequencer/DMXLibrarySection.h"
+#include "Sequencer/MovieSceneDMXLibrarySection.h"
+#include "Sequencer/MovieSceneDMXLibraryTrack.h"
+#include "Styling/SlateIconFinder.h"
 #include "Widgets/Layout/SBox.h"
-#include "Widgets/Images/SImage.h"
-#include "Widgets/Input/SButton.h"
-#include "Framework/Application/SlateApplication.h"
 
 #define LOCTEXT_NAMESPACE "DMXLibraryTrackEditor"
 
@@ -161,6 +155,8 @@ TSharedPtr<SWidget> FDMXLibraryTrackEditor::BuildOutlinerEditWidget(const FGuid&
 	// Create combo button "+ Patch" to pick an asset to add
 	// sub menu content callback
 	FOnGetContent AddPatchMenuContent = FOnGetContent::CreateSP(this, &FDMXLibraryTrackEditor::OnGetAddPatchMenuContent, DMXTrack);
+
+	constexpr bool bEnabled = true;
 	return UE::Sequencer::MakeAddButton(LOCTEXT("AddPatchButton", "Patch"), AddPatchMenuContent, Params.ViewModel);
 }
 
diff --git a/Engine/Plugins/VirtualProduction/DMX/DMXEngine/Source/DMXEditor/Private/Sequencer/DMXLibraryTrackEditor.h b/Engine/Plugins/VirtualProduction/DMX/DMXEngine/Source/DMXEditor/Private/Sequencer/DMXLibraryTrackEditor.h
--- a/Engine/Plugins/VirtualProduction/DMX/DMXEngine/Source/DMXEditor/Private/Sequencer/DMXLibraryTrackEditor.h
+++ b/Engine/Plugins/VirtualProduction/DMX/DMXEngine/Source/DMXEditor/Private/Sequencer/DMXLibraryTrackEditor.h
@@ -2,7 +2,6 @@
 
 #pragma once
 
-#include "CoreMinimal.h"
 #include "MovieSceneTrackEditor.h"
 
 struct FAssetData;
