DMX - Upgrade to a new GDTF implementation, initial dev version. UDMXImportGDTF now only holds the raw GDTF data, its GDTF description is deprecated. DMXGDTF now holds the GDTF standard in its own resuable plugin, and can be initialized from raw data.

diff --git a/Engine/Plugins/VirtualProduction/DMX/DMXEngine/DMXEngine.uplugin b/Engine/Plugins/VirtualProduction/DMX/DMXEngine/DMXEngine.uplugin
--- a/Engine/Plugins/VirtualProduction/DMX/DMXEngine/DMXEngine.uplugin
+++ b/Engine/Plugins/VirtualProduction/DMX/DMXEngine/DMXEngine.uplugin
@@ -34,6 +34,10 @@
 			"Name": "DatasmithContent",
 			"Enabled": true
 		},
+		{
+			"Name": "DMXGDTF",
+			"Enabled": true
+		},
 		{
 			"Name": "DMXModularFeatures",
 			"Enabled": true
diff --git a/Engine/Plugins/VirtualProduction/DMX/DMXEngine/Source/DMXEditor/DMXEditor.Build.cs b/Engine/Plugins/VirtualProduction/DMX/DMXEngine/Source/DMXEditor/DMXEditor.Build.cs
--- a/Engine/Plugins/VirtualProduction/DMX/DMXEngine/Source/DMXEditor/DMXEditor.Build.cs
+++ b/Engine/Plugins/VirtualProduction/DMX/DMXEngine/Source/DMXEditor/DMXEditor.Build.cs
@@ -29,6 +29,8 @@ public class DMXEditor : ModuleRules
 				"ContentBrowser",
 				"CoreUObject",
 				"DesktopPlatform",
+				"DMXGDTF",
+				"DMXZip",
 				"EditorFramework",
 				"EditorStyle",
 				"EditorWidgets",
diff --git a/Engine/Plugins/VirtualProduction/DMX/DMXEngine/Source/DMXEditor/Private/Customizations/DMXEntityFixtureTypeDetails.cpp b/Engine/Plugins/VirtualProduction/DMX/DMXEngine/Source/DMXEditor/Private/Customizations/DMXEntityFixtureTypeDetails.cpp
--- a/Engine/Plugins/VirtualProduction/DMX/DMXEngine/Source/DMXEditor/Private/Customizations/DMXEntityFixtureTypeDetails.cpp
+++ b/Engine/Plugins/VirtualProduction/DMX/DMXEngine/Source/DMXEditor/Private/Customizations/DMXEntityFixtureTypeDetails.cpp
@@ -26,20 +26,14 @@ void FDMXEntityFixtureTypeDetails::CustomizeDetails(IDetailLayoutBuilder& Detail
 
 	DetailBuilder.HideProperty(GET_MEMBER_NAME_CHECKED(UDMXEntityFixtureType, Modes));
 
-	GDTFHandle = DetailBuilder.GetProperty(GET_MEMBER_NAME_CHECKED(UDMXEntityFixtureType, DMXImport));
-	GDTFHandle->SetOnPropertyValueChanged(FSimpleDelegate::CreateSP(this, &FDMXEntityFixtureTypeDetails::OnDMXImportChanged));
+	GDTFSourceHandle = DetailBuilder.GetProperty(GET_MEMBER_NAME_CHECKED(UDMXEntityFixtureType, GDTFSource));
+	GDTFSourceHandle->SetOnPropertyValueChanged(FSimpleDelegate::CreateSP(this, &FDMXEntityFixtureTypeDetails::OnGDTFSourceChanged));
 }
 
-void FDMXEntityFixtureTypeDetails::OnDMXImportChanged()
+void FDMXEntityFixtureTypeDetails::OnGDTFSourceChanged()
 {
-	UObject* DMXImportObject = nullptr;
-	if (GDTFHandle->GetValue(DMXImportObject) != FPropertyAccess::Success)
-	{
-		return;
-	}
-
-	UDMXImportGDTF* DMXImportGDTF = Cast<UDMXImportGDTF>(DMXImportObject);
 	const TArray<TWeakObjectPtr<UObject>>& SelectedObjects = PropertyUtilities->GetSelectedObjects();
+
 	for (TWeakObjectPtr<UObject> WeakFixtureTypeObject : SelectedObjects)
 	{
 		if (UDMXEntityFixtureType* FixtureType = Cast<UDMXEntityFixtureType>(WeakFixtureTypeObject.Get()))
@@ -47,19 +41,22 @@ void FDMXEntityFixtureTypeDetails::OnDMXImportChanged()
 			FixtureType->PreEditChange(nullptr);
 			FixtureType->Modes.Reset();
 			FixtureType->PostEditChange();
-			
-			if (!DMXImportGDTF)
+
+			if (FixtureType->GDTFSource.IsNull())
 			{
 				continue;
 			}
+			UDMXImportGDTF* GDTF = FixtureType->GDTFSource.LoadSynchronous();
 
 			// Try to use the work around that supports creation of matrices, otherwise setup the fixture type with the old implementation
 			FixtureType->PreEditChange(nullptr);
-			const bool bAdvancedImportSuccess = FDMXInitializeFixtureTypeFromGDTFHelper::GenerateModesFromGDTF(*FixtureType, *DMXImportGDTF);
+			const bool bAdvancedImportSuccess = FDMXInitializeFixtureTypeFromGDTFHelper::GenerateModesFromGDTF(*FixtureType, *GDTF);
 			if (!bAdvancedImportSuccess)
 			{
+				PRAGMA_DISABLE_DEPRECATION_WARNINGS
 				UE_LOG(LogDMXEditor, Warning, TEXT("Failed to initialize Fixture Type '%s', falling back to legacy method that doesn't support matrix fixtures."), *FixtureType->GetName());
-				FixtureType->SetModesFromDMXImport(DMXImportGDTF);
+				FixtureType->SetModesFromDMXImport(GDTF);
+				PRAGMA_ENABLE_DEPRECATION_WARNINGS
 			}
 			FixtureType->PostEditChange();
 		}
diff --git a/Engine/Plugins/VirtualProduction/DMX/DMXEngine/Source/DMXEditor/Private/Customizations/DMXEntityFixtureTypeDetails.h b/Engine/Plugins/VirtualProduction/DMX/DMXEngine/Source/DMXEditor/Private/Customizations/DMXEntityFixtureTypeDetails.h
--- a/Engine/Plugins/VirtualProduction/DMX/DMXEngine/Source/DMXEditor/Private/Customizations/DMXEntityFixtureTypeDetails.h
+++ b/Engine/Plugins/VirtualProduction/DMX/DMXEngine/Source/DMXEditor/Private/Customizations/DMXEntityFixtureTypeDetails.h
@@ -23,11 +23,11 @@ protected:
 	//~ End IDetailCustomization interface
 
 private:
-	/** Called when the DMXImport property changed */
-	void OnDMXImportChanged();
+	/** Called when the GDTFSource property changed */
+	void OnGDTFSourceChanged();
 
-	/** Handle to the DMXImport property */
-	TSharedPtr<IPropertyHandle> GDTFHandle;
+	/** Handle to the GDTFSource property */
+	TSharedPtr<IPropertyHandle> GDTFSourceHandle;
 
 	/** Property utilities for this customization */
 	TSharedPtr<IPropertyUtilities> PropertyUtilities;
diff --git a/Engine/Plugins/VirtualProduction/DMX/DMXEngine/Source/DMXEditor/Private/Exporters/DMXMVRExporter.cpp b/Engine/Plugins/VirtualProduction/DMX/DMXEngine/Source/DMXEditor/Private/Exporters/DMXMVRExporter.cpp
--- a/Engine/Plugins/VirtualProduction/DMX/DMXEngine/Source/DMXEditor/Private/Exporters/DMXMVRExporter.cpp
+++ b/Engine/Plugins/VirtualProduction/DMX/DMXEngine/Source/DMXEditor/Private/Exporters/DMXMVRExporter.cpp
@@ -147,13 +147,13 @@ bool FDMXMVRExporter::ZipGDTFs(const TSharedRef<FDMXZipper>& Zip, UDMXLibrary* D
 	bool bAllZippedSuccessfully = true;
 	for (UDMXEntityFixtureType* FixtureType : FixtureTypesToExport)
 	{
-		if (!FixtureType->DMXImport)
+		if (FixtureType->GDTFSource.IsNull())
 		{
 			UE_LOG(LogDMXEditor, Warning, TEXT("Cannot export Fixture Type '%s' to MVR, but the Fixture Type has no GDTF set."), *FixtureType->Name);
 			continue;
 		}
 
-		UDMXImportGDTF* DMXImportGDTF = Cast<UDMXImportGDTF>(FixtureType->DMXImport);
+		UDMXImportGDTF* DMXImportGDTF = FixtureType->GDTFSource.LoadSynchronous();
 		if (!DMXImportGDTF)
 		{
 			UE_LOG(LogDMXEditor, Warning, TEXT("Cannot export Fixture Type '%s' to MVR, but the Fixture Type has a DMX Import Type which is not GDTF."), *FixtureType->Name);
diff --git a/Engine/Plugins/VirtualProduction/DMX/DMXEngine/Source/DMXEditor/Private/Factories/DMXGDTFFactory.cpp b/Engine/Plugins/VirtualProduction/DMX/DMXEngine/Source/DMXEditor/Private/Factories/DMXGDTFFactory.cpp
--- a/Engine/Plugins/VirtualProduction/DMX/DMXEngine/Source/DMXEditor/Private/Factories/DMXGDTFFactory.cpp
+++ b/Engine/Plugins/VirtualProduction/DMX/DMXEngine/Source/DMXEditor/Private/Factories/DMXGDTFFactory.cpp
@@ -1,31 +1,32 @@
 // Copyright Epic Games, Inc. All Rights Reserved.
 
 #include "Factories/DMXGDTFFactory.h"
-#include "Factories/DMXGDTFImportUI.h"
-#include "DMXEditorLog.h"
-#include "Factories/DMXGDTFImporter.h"
-#include "Library/DMXGDTFAssetImportData.h"
-#include "Library/DMXImportGDTF.h"
 
+#include "Application/SlateApplicationBase.h"
+#include "DMXEditorLog.h"
 #include "Editor.h"
+#include "Editor/EditorEngine.h"
 #include "EditorReimportHandler.h"
-#include "AssetImportTask.h"
-#include "EditorFramework/AssetImportData.h"
+#include "Factories/DMXGDTFImporter.h"
+#include "Factories/DMXGDTFImportUI.h"
+#include "Framework/Application/SlateApplication.h"
+#include "Framework/Notifications/NotificationManager.h"
 #include "HAL/FileManager.h"
-#include "Misc/CommandLine.h"
-#include "Misc/FeedbackContext.h"
+#include "Interfaces/IMainFrameModule.h"
+#include "Library/DMXGDTFAssetImportData.h"
+#include "Library/DMXImportGDTF.h"
 #include "Misc/Paths.h"
+#include "Widgets/Notifications/SNotificationList.h"
+#include "Widgets/SDMXGDTFOptionWindow.h"
 
-
-const TCHAR* UDMXGDTFFactory::Extension = TEXT("gdtf");
+#define LOCTEXT_NAMESPACE "DMXGDTFFactory"
 
 UDMXGDTFImportUI::UDMXGDTFImportUI()
-    : bUseSubDirectory(false)
-    , bImportXML(true)
-    , bImportTextures(false)
-    , bImportModels(false)
-{
-}
+	: bUseSubDirectory(false)
+	, bImportXML(true)
+	, bImportTextures(false)
+	, bImportModels(false)
+{}
 
 void UDMXGDTFImportUI::ResetToDefault()
 {
@@ -46,28 +47,6 @@ UDMXGDTFFactory::UDMXGDTFFactory()
 	bOperationCanceled = false;
 }
 
-void UDMXGDTFFactory::CleanUp()
-{
-	Super::CleanUp();
-
-    bShowOption = true;
-}
-
-bool UDMXGDTFFactory::ConfigureProperties()
-{
-    Super::ConfigureProperties();
-    EnableShowOption();
-
-    return true;
-}
-
-void UDMXGDTFFactory::PostInitProperties()
-{
-	Super::PostInitProperties();
-
-    ImportUI = NewObject<UDMXGDTFImportUI>(this, NAME_None, RF_NoFlags);
-}
-
 bool UDMXGDTFFactory::DoesSupportClass(UClass* Class)
 {
     return Class == UDMXImportGDTF::StaticClass();
@@ -93,7 +72,7 @@ UObject* UDMXGDTFFactory::FactoryCreateFile(UClass* InClass, UObject* InParent,
 
     CA_ASSUME(InParent);
 
-    if( bOperationCanceled )
+    if (bOperationCanceled)
     {
         bOutOperationCanceled = true;
         GEditor->GetEditorSubsystem<UImportSubsystem>()->BroadcastAssetPostImport(this, nullptr);
@@ -103,68 +82,63 @@ UObject* UDMXGDTFFactory::FactoryCreateFile(UClass* InClass, UObject* InParent,
     GEditor->GetEditorSubsystem<UImportSubsystem>()->BroadcastAssetPreImport(this, InClass, InParent, InName, Type);
 
     UObject* ExistingObject = nullptr;
-    if (InParent != nullptr)
+    if (InParent)
     {
         ExistingObject = StaticFindObject(UObject::StaticClass(), InParent, *(InName.ToString()));
         if (ExistingObject)
         {
-            bShowOption = false;
+			bShowOptions = false;
         }
     }
 
-    // Prepare import options
+	if (FParse::Param(FCommandLine::Get(), TEXT("NoDMXImportOption")))
+	{
+		bShowOptions = false;
+	}
+
+    // Create import options
+	using namespace UE::DMX;
     FDMXGDTFImportArgs ImportArgs;
-    ImportArgs.ImportUI = ImportUI;
     ImportArgs.Name = InName;
     ImportArgs.Parent = InParent;
-    ImportArgs.CurrentFilename = InFilename;
+    ImportArgs.Filename = InFilename;
     ImportArgs.Flags = Flags;
-    ImportArgs.bCancelOperation = bOperationCanceled;
-    const TUniquePtr<FDMXGDTFImporter> Importer = MakeUnique<FDMXGDTFImporter>(ImportArgs);
 
     // Set Import UI
-    if (FParse::Param(FCommandLine::Get(), TEXT("NoDMXImportOption")))
-    {
-        bShowOption = false;
-    }
     bool bIsAutomated = IsAutomatedImport();
-    bool bShowImportDialog = bShowOption && !bIsAutomated;
-    bool bImportAll = false;
-    FDMXGDTFImporter::GetImportOptions(Importer, ImportUI, bShowImportDialog, InParent->GetPathName(), bOperationCanceled, bImportAll, UFactory::CurrentFilename);
-    bOutOperationCanceled = bOperationCanceled;
-    if( bImportAll )
-    {
-        // If the user chose to import all, we don't show the dialog again and use the same settings for each object until importing another set of files
-        bShowOption = false;
-    }
+    bool bShowImportDialog = bShowOptions && !bIsAutomated;
 
-    if (!ImportUI->bImportXML && !ImportUI->bImportModels && !ImportUI->bImportTextures)
+	ImportUI = nullptr;
+	if (bShowImportDialog)
+	{
+		ImportUI = GetMutableDefault<UDMXGDTFImportUI>();
+		bOperationCanceled = GetOptionsFromDialog(InParent);
+		bOutOperationCanceled = bOperationCanceled;
+	}
+
+    if (bImportAll)
     {
-        Warn->Log(ELogVerbosity::Error, TEXT("Nothing to Import") );
-		return nullptr;
+        // If the user chose to import all, we don't show the dialog again and use the same settings for each object until importing another set of files
+		bShowOptions = false;
     }
 
-    // Try to load and parse the content
-    if (!Importer->AttemptImportFromFile())
+    if (ImportUI && !ImportUI->bImportXML && !ImportUI->bImportModels && !ImportUI->bImportTextures)
     {
-        Warn->Log(ELogVerbosity::Error, TEXT("Failed to import GDTF") );
+		const FNotificationInfo Info(LOCTEXT("NothingToimportInfo", "Skipping import of GDTF, nothing to import."));
+		FSlateNotificationManager::Get().AddNotification(Info);
 		return nullptr;
     }
 
 	// Import to the Editor
-	UDMXImportGDTF* GDTF = Importer->Import();
+	FText OutErrorReason;
+	UDMXImportGDTF* GDTF = FDMXGDTFImporter::Import(*this, ImportArgs, OutErrorReason);
 	if (!GDTF)
 	{
-		return nullptr;
-	}
+		const FNotificationInfo Info(OutErrorReason);
+		FSlateNotificationManager::Get().AddNotification(Info);
 
-	// Set Asset Import Data
-	UDMXGDTFAssetImportData* GDTFAssetImportData = GDTF->GetGDTFAssetImportData();
-	if (!ensureAlwaysMsgf(GDTFAssetImportData, TEXT("Unexpected missing Asset Import Data for newly created GDTF %s"), *GDTF->GetName()))
-	{
 		return nullptr;
 	}
-	GDTFAssetImportData->SetSourceFile(InFilename);
 
 	GEditor->GetEditorSubsystem<UImportSubsystem>()->BroadcastAssetPostImport(this, GDTF);
 
@@ -175,12 +149,7 @@ bool UDMXGDTFFactory::FactoryCanImport(const FString& Filename)
 {
 	const FString TargetExtension = FPaths::GetExtension(Filename);
 
-	if(TargetExtension == UDMXGDTFFactory::Extension)
-	{
-		return true;
-	}
-
-	return false;
+	return TargetExtension == TEXT("gdtf");
 }
 
 bool UDMXGDTFFactory::CanReimport(UObject* Obj, TArray<FString>& OutFilenames)
@@ -228,7 +197,55 @@ EReimportResult::Type UDMXGDTFFactory::Reimport(UObject* InObject)
 	return bOutCanceled ? EReimportResult::Cancelled : EReimportResult::Failed;
 }
 
-int32 UDMXGDTFFactory::GetPriority() const
+bool UDMXGDTFFactory::GetOptionsFromDialog(UObject* Parent)
 {
-    return ImportPriority;
+	if (!ensureMsgf(Parent, TEXT("Trying to display import options for transient object. This is not expected.")))
+	{
+		return false;
+	}
+
+	TSharedPtr<SWindow> ParentWindow;
+	if (FModuleManager::Get().IsModuleLoaded("MainFrame"))
+	{
+		IMainFrameModule& MainFrame = FModuleManager::LoadModuleChecked<IMainFrameModule>("MainFrame");
+		ParentWindow = MainFrame.GetParentWindow();
+	}
+
+	// Compute centered window position based on max window size, which include when all categories are expanded
+	const float ImportWindowWidth = 410.0f;
+	const float ImportWindowHeight = 750.0f;
+	const FVector2D ImportWindowSize = FVector2D(ImportWindowWidth, ImportWindowHeight); // Max window size it can get based on current slate
+
+	const FSlateRect WorkAreaRect = FSlateApplicationBase::Get().GetPreferredWorkArea();
+	const FVector2D DisplayTopLeft(WorkAreaRect.Left, WorkAreaRect.Top);
+	const FVector2D DisplaySize(WorkAreaRect.Right - WorkAreaRect.Left, WorkAreaRect.Bottom - WorkAreaRect.Top);
+
+	const FVector2D WindowPosition = (DisplayTopLeft + (DisplaySize - ImportWindowSize) / 2.0f);
+
+	const TSharedRef<SWindow> Window = SNew(SWindow)
+		.Title(LOCTEXT("GDTFImportOpionsTitle", "GDTF Import Options"))
+		.SizingRule(ESizingRule::Autosized)
+		.AutoCenter(EAutoCenter::None)
+		.ClientSize(ImportWindowSize)
+		.ScreenPosition(WindowPosition);
+
+	using namespace UE::DMX;
+	TSharedPtr<SDMXGDTFOptionWindow> OptionWindow;
+	Window->SetContent
+	(
+		SAssignNew(OptionWindow, SDMXGDTFOptionWindow)
+		.ImportUI(ImportUI)
+		.WidgetWindow(Window)
+		.FullPath(FText::FromString(Parent->GetPathName()))
+		.MaxWindowHeight(ImportWindowHeight)
+		.MaxWindowWidth(ImportWindowWidth)
+	);
+
+	FSlateApplication::Get().AddModalWindow(Window, ParentWindow, false);
+
+	bImportAll = OptionWindow->ShouldImportAll();
+
+	return OptionWindow->ShouldImport();
 }
+
+#undef LOCTEXT_NAMESPACE
diff --git a/Engine/Plugins/VirtualProduction/DMX/DMXEngine/Source/DMXEditor/Private/Factories/DMXGDTFFactory.h b/Engine/Plugins/VirtualProduction/DMX/DMXEngine/Source/DMXEditor/Private/Factories/DMXGDTFFactory.h
--- a/Engine/Plugins/VirtualProduction/DMX/DMXEngine/Source/DMXEditor/Private/Factories/DMXGDTFFactory.h
+++ b/Engine/Plugins/VirtualProduction/DMX/DMXEngine/Source/DMXEditor/Private/Factories/DMXGDTFFactory.h
@@ -21,15 +21,6 @@ class UDMXGDTFFactory
 public:
 	UDMXGDTFFactory();
 
-	/**  Set import batch **/
-	void EnableShowOption() { bShowOption = true; }
-
-	//~ Begin UObject Interface
-	virtual void CleanUp() override;
-	virtual bool ConfigureProperties() override;
-	virtual void PostInitProperties() override;
-	//~ End UObject Interface
-
 	//~ Begin UFactory Interface
 	virtual bool DoesSupportClass(UClass * Class) override;
 	virtual UClass* ResolveSupportedClass() override;
@@ -41,17 +32,20 @@ public:
 	virtual bool CanReimport(UObject* Obj, TArray<FString>& OutFilenames) override;
 	virtual void SetReimportPaths(UObject* Obj, const TArray<FString>& NewReimportPaths) override;
 	virtual EReimportResult::Type Reimport(UObject* Obj) override;
-	virtual int32 GetPriority() const override;
 	//~ End FReimportHandler Interface
 
-public:
-	static const TCHAR* Extension;
-
 private:
-	bool bShowOption;
+	/** Shows an options dialog that initializes members. If false is returned, the import should be canceled */
+	[[nodiscard]] bool GetOptionsFromDialog(UObject* Parent);
+
+	/** If true, shows an options dialog. Can be false for example when reimporting */
+	bool bShowOptions = false;
+
+	/** If true, all content of the GDTF should be imported */
+	bool bImportAll = false;
 
-	/** true if the import operation was canceled. */
-	bool bOperationCanceled;
+	/** If true, importing was canceled */
+	bool bOperationCanceled = false;
 
 	UPROPERTY(transient)
 	TObjectPtr<UDMXGDTFImportUI> ImportUI;
diff --git a/Engine/Plugins/VirtualProduction/DMX/DMXEngine/Source/DMXEditor/Private/Factories/DMXGDTFImportUI.h b/Engine/Plugins/VirtualProduction/DMX/DMXEngine/Source/DMXEditor/Private/Factories/DMXGDTFImportUI.h
--- a/Engine/Plugins/VirtualProduction/DMX/DMXEngine/Source/DMXEditor/Private/Factories/DMXGDTFImportUI.h
+++ b/Engine/Plugins/VirtualProduction/DMX/DMXEngine/Source/DMXEditor/Private/Factories/DMXGDTFImportUI.h
@@ -19,17 +19,15 @@ public:
 	void ResetToDefault();
 
 public:
-    UPROPERTY(EditAnywhere, Category = "DMX")
-    bool bUseSubDirectory;
+    UPROPERTY(EditAnywhere, Config, Category = "DMX")
+    bool bUseSubDirectory = true;
 
-    UPROPERTY(EditAnywhere, Category = "DMX")
-    bool bImportXML;
+    UPROPERTY(EditAnywhere, Config, Category = "DMX")
+    bool bImportXML = true;
 
-    UPROPERTY(EditAnywhere, Category = "DMX")
-    bool bImportTextures;
+    UPROPERTY(EditAnywhere, Config, Category = "DMX")
+    bool bImportTextures = true;
 
-    UPROPERTY(EditAnywhere, Category = "DMX")
-    bool bImportModels;
+    UPROPERTY(EditAnywhere, Config, Category = "DMX")
+    bool bImportModels = true;
 };
-
-
diff --git a/Engine/Plugins/VirtualProduction/DMX/DMXEngine/Source/DMXEditor/Private/Factories/DMXGDTFImporter.cpp b/Engine/Plugins/VirtualProduction/DMX/DMXEngine/Source/DMXEditor/Private/Factories/DMXGDTFImporter.cpp
--- a/Engine/Plugins/VirtualProduction/DMX/DMXEngine/Source/DMXEditor/Private/Factories/DMXGDTFImporter.cpp
+++ b/Engine/Plugins/VirtualProduction/DMX/DMXEngine/Source/DMXEditor/Private/Factories/DMXGDTFImporter.cpp
@@ -1,691 +1,79 @@
 // Copyright Epic Games, Inc. All Rights Reserved.
 
 #include "Factories/DMXGDTFImporter.h"
-#include "Library/DMXImportGDTF.h"
-#include "DMXEditorLog.h"
-#include "Widgets/SDMXGDTFOptionWindow.h"
-#include "Factories/DMXGDTFImportUI.h"
 
-#include "Framework/Application/SlateApplication.h"
-#include "Interfaces/IMainFrameModule.h"
-#include "Modules/ModuleManager.h"
-#include "ObjectTools.h"
-#include "HAL/PlatformFileManager.h"
-#include "HAL/FileManagerGeneric.h"
-#include "Misc/FileHelper.h"
-#include "PackageTools.h"
-#include "Serialization/BufferArchive.h"
-#include "Containers/UnrealString.h"
-#include "Factories/TextureFactory.h"
-#include "Engine/Texture.h"
-#include "AssetRegistry/AssetRegistryModule.h"
-#include "AssetToolsModule.h"
+#include "DMXEditorLog.h"
+#include "DMXZipper.h"
 #include "EditorFramework/AssetImportData.h"
-#include "XmlNode.h"
+#include "Factories/DMXGDTFFactory.h"
+#include "GDTF/DMXGDTFDescription.h"
+#include "Library/DMXGDTFAssetImportData.h"
+#include "Library/DMXImportGDTF.h"
+#include "Misc/Paths.h"
 #include "XmlFile.h"
-#include "Engine/Texture2D.h"
 
-#define LOCTEXT_NAMESPACE "DMXGDTFImporter"
 
+#define LOCTEXT_NAMESPACE "DMXGDTFImporter"
 
-namespace UE::DMX::DMXGDTFImporter::Private
+namespace UE::DMX
 {
-	/** 
-	 * Depending on the GDTF Spec, some XML Attribute Names might start with or without DMX. 
-	 * 
-	 * @param ParentNode		The Parent Node to search in 
-	 * @param NodeName			The Attribute Name with DMX tag. This is the correct form, e.g. 'DMXMode', not 'Mode', 'DMXChannels', not 'Channels'.
-	 * @return					Returns a pointer to the child node, or nullptr if the child cannot be found 
-	 */
-	const FXmlNode* FindChildNodeEvenIfDMXSubstringIsMissing(const FXmlNode& ParentNode, const FString& AttributeNameWithDMXTag)
+	UDMXImportGDTF* FDMXGDTFImporter::Import(const UDMXGDTFFactory& InImportFactory, const FDMXGDTFImportArgs& ImportArgs, FText& OutErrorReason)
+	{
+		FDMXGDTFImporter Instance(ImportArgs);
+		return Instance.ImportInternal(InImportFactory, OutErrorReason);
+	}
+
+	FDMXGDTFImporter::FDMXGDTFImporter(const FDMXGDTFImportArgs& InImportArgs)
+		: ImportArgs(InImportArgs)
+	{}
+
+	UDMXImportGDTF* FDMXGDTFImporter::ImportInternal(const UDMXGDTFFactory& InImportFactory, FText& OutErrorReason)
 	{
-		if (const FXmlNode* ChildNodePtrWithSubstring = ParentNode.FindChildNode(AttributeNameWithDMXTag))
+		if (!ensureAlwaysMsgf(FPaths::FileExists(ImportArgs.Filename), TEXT("Cannot import GDTF File '%s'. File does not exist."), *ImportArgs.Filename))
 		{
-			return ChildNodePtrWithSubstring;
+			return nullptr;
 		}
-		else if (const FXmlNode* ChildNodePtrWithoutSubstring = ParentNode.FindChildNode(AttributeNameWithDMXTag.RightChop(3)))
+
+		// Unzip GDTF
+		Zip = MakeShared<FDMXZipper>();
+		if (!Zip->LoadFromFile(ImportArgs.Filename))
 		{
-			return ChildNodePtrWithoutSubstring;
+			OutErrorReason = LOCTEXT("InvalidGDTFError", "Cannot read GDTF.File is not a valid GDTF file.");
+			return nullptr;
 		}
 
-		return nullptr;
+		return CreateGDTF(InImportFactory, OutErrorReason);
 	}
 
-	/** 
-	 * Depending on the GDTF Spec, some XML Attribute Names might start with or without DMX. 
-	 * 
-	 * @param ParentNode		The Parent Node to search in 
-	 * @param NodeName			The Attribute Name with DMX tag. This is the correct form, e.g. 'DMXMode', not 'Mode', 'DMXChannels', not 'Channels'.
-	 * @return					Returns the Attribute as String
-	 */
-	FString FindAttributeEvenIfDMXSubstringIsMissing(const FXmlNode& ParentNode, const FString& AttributeNameWithDMXTag)
+	UDMXImportGDTF* FDMXGDTFImporter::CreateGDTF(const UDMXGDTFFactory& InImportFactory, FText& OutErrorReason) const
 	{
-		FString Attribute = ParentNode.GetAttribute(AttributeNameWithDMXTag);
-		if (Attribute.IsEmpty())
+		if (!ensureMsgf(Zip.IsValid(), TEXT("Invalid zipper. Cannot create GDTF asset.")))
 		{
-			Attribute = ParentNode.GetAttribute(AttributeNameWithDMXTag.RightChop(3));
+			return nullptr;
 		}
 
-		return Attribute;
-	}
-};
-
-FDMXGDTFImporter::FDMXGDTFImporter(const FDMXGDTFImportArgs& InImportArgs)
-    : ImportArgs(InImportArgs)
-    , bIsXMLParsedSuccessfully(false)
-{
-    ImportName = ObjectTools::SanitizeObjectName(ImportArgs.Name.ToString());
-}
-
-UDMXImportGDTF* FDMXGDTFImporter::Import()
-{
-    UDMXImportGDTF* CreatedObject = nullptr;
-
-    // Import GDTF XML into UObject
-    if (ImportArgs.ImportUI->bImportXML && bIsXMLParsedSuccessfully)
-    {
-        CreatedObject = CreateGDTFDesctription();
-    }
-
-    return CreatedObject;
-}
-
-UPackage* FDMXGDTFImporter::GetPackage(const FString& InImportName)
-{
-    UPackage* Package = nullptr;
-    FString NewPackageName;
-
-    if (!ImportArgs.ImportUI->bUseSubDirectory && ImportArgs.Parent != nullptr && ImportArgs.Parent->IsA(UPackage::StaticClass()))
-    {
-        Package = StaticCast<UPackage*>(ImportArgs.Parent.Get());
-    }
-
-    if (Package == nullptr)
-    {
-        if (ImportArgs.Parent != nullptr && ImportArgs.Parent->GetOutermost() != nullptr)
-        {
-            NewPackageName = FPackageName::GetLongPackagePath(ImportArgs.Parent->GetOutermost()->GetName()) + "/" + ImportName;
-        }
-        else
-        {
-            return nullptr;
-        }
-
-        if (ImportArgs.ImportUI->bUseSubDirectory)
-        {
-            NewPackageName /= InImportName;
-        }
-
-        NewPackageName = UPackageTools::SanitizePackageName(NewPackageName);
-        LastPackageName = NewPackageName;
-        Package = CreatePackage(*NewPackageName);
-        Package->FullyLoad();
-    }
-
-    return Package;
-}
-
-bool FDMXGDTFImporter::AttemptImportFromFile()
-{
-    if (ImportArgs.ImportUI->bImportXML)
-    {
-        if (!ParseXML())
-        {
-            UE_LOG_DMXEDITOR(Warning, TEXT("Error parsing XML"));
-
-            return false;
-        }
-    }
-
-    return true;
-}
-
-UDMXImportGDTF* FDMXGDTFImporter::CreateGDTFDesctription()
-{
-    UPackage* Package = GetPackage(ImportName);
-
-    // Create new GDTF asset
-    UDMXImportGDTF* ImportObject = NewObject<UDMXImportGDTF>(Package, FName(*ImportName), ImportArgs.Flags | RF_Public);
-
-    UDMXImportGDTFFixtureType* GDTFFixtureType = ImportObject->CreateNewObject<UDMXImportGDTFFixtureType>();
-    ImportObject->FixtureType = GDTFFixtureType;
-
-    UDMXImportGDTFAttributeDefinitions* GDTFAttributeDefinitions = ImportObject->CreateNewObject<UDMXImportGDTFAttributeDefinitions>();
-    ImportObject->AttributeDefinitions = GDTFAttributeDefinitions;
-
-    UDMXImportGDTFWheels* GDTFWheels = ImportObject->CreateNewObject<UDMXImportGDTFWheels>();
-    ImportObject->Wheels = GDTFWheels;
-
-    UDMXImportGDTFPhysicalDescriptions* GDTFPhysicalDescriptions = ImportObject->CreateNewObject<UDMXImportGDTFPhysicalDescriptions>();
-    ImportObject->PhysicalDescriptions = GDTFPhysicalDescriptions;
-
-    UDMXImportGDTFModels* GDTFModels = ImportObject->CreateNewObject<UDMXImportGDTFModels>();
-    ImportObject->Models = GDTFModels;
-
-    UDMXImportGDTFGeometries* GDTFGeometries = ImportObject->CreateNewObject<UDMXImportGDTFGeometries>();
-    ImportObject->Geometries = GDTFGeometries;
-
-    UDMXImportGDTFDMXModes* GDTFDMXModes = ImportObject->CreateNewObject<UDMXImportGDTFDMXModes>();
-    ImportObject->DMXModes = GDTFDMXModes;
-
-    UDMXImportGDTFProtocols* GDTFProtocols = ImportObject->CreateNewObject<UDMXImportGDTFProtocols>();
-    ImportObject->Protocols = GDTFProtocols;
-
-    if (bIsXMLParsedSuccessfully)
-    {
-        if (FXmlNode* XmlNode = XMLFile->GetRootNode())
-        {
-            if (const FXmlNode* FixtureTypeNode = XmlNode->FindChildNode("FixtureType"))
-            {
-                GDTFFixtureType->Description = FixtureTypeNode->GetAttribute("Description");
-                GDTFFixtureType->FixtureTypeID = FixtureTypeNode->GetAttribute("FixtureTypeID");
-                GDTFFixtureType->LongName = FixtureTypeNode->GetAttribute("LongName");
-                GDTFFixtureType->Manufacturer = FixtureTypeNode->GetAttribute("Manufacturer");
-                GDTFFixtureType->Name = FName(*FixtureTypeNode->GetAttribute("Name"));
-                GDTFFixtureType->ShortName = FixtureTypeNode->GetAttribute("ShortName");
-                GDTFFixtureType->RefFT = FixtureTypeNode->GetAttribute("RefFT");
-
-                if (const FXmlNode* AttributeDefinitionsNode = FixtureTypeNode->FindChildNode("AttributeDefinitions"))
-                {
-                    if (const FXmlNode* ActivationGroupsNode = AttributeDefinitionsNode->FindChildNode("ActivationGroups"))
-                    {
-                        for (const FXmlNode* ActivationGroupNode : ActivationGroupsNode->GetChildrenNodes())
-                        {
-                            FDMXImportGDTFActivationGroup ActivationGroup;
-                            ActivationGroup.Name = FName(*ActivationGroupNode->GetAttribute("Name"));
-                            GDTFAttributeDefinitions->ActivationGroups.Add(ActivationGroup);
-                        }
-                    }
-
-                    if (const FXmlNode* FeatureGroupsNode = AttributeDefinitionsNode->FindChildNode("FeatureGroups"))
-                    {
-                        for (const FXmlNode* FeatureGroupNode : FeatureGroupsNode->GetChildrenNodes())
-                        {
-                            if (FeatureGroupNode != nullptr)
-                            {
-                                FDMXImportGDTFFeatureGroup FeatureGroup;
-                                FeatureGroup.Name = FName(*FeatureGroupNode->GetAttribute("Name"));
-                                FeatureGroup.Pretty = FeatureGroupNode->GetAttribute("Pretty");
-
-                                for (const FXmlNode* FeatureNode : FeatureGroupNode->GetChildrenNodes())
-                                {
-                                    FDMXImportGDTFFeature Feature;
-                                    Feature.Name = FName(*FeatureNode->GetAttribute("Name"));
-                                    FeatureGroup.Features.Add(Feature);
-                                }
-
-                                GDTFAttributeDefinitions->FeatureGroups.Add(FeatureGroup);
-                            }
-                        }
-                    }
-
-                    if (const FXmlNode* AttributesNode = AttributeDefinitionsNode->FindChildNode("Attributes"))
-                    {
-                        for (const FXmlNode* AttributeNode : AttributesNode->GetChildrenNodes())
-                        {
-                            if (AttributeNode != nullptr)
-                            {
-                                FDMXImportGDTFAttribute Attribute ;
-                                Attribute.Name = FName(*AttributeNode->GetAttribute("Name"));
-                                Attribute.Pretty = AttributeNode->GetAttribute("Pretty");
-                                Attribute.ActivationGroup.Name = FName(*AttributeNode->GetAttribute("ActivationGroup"));
-                                GDTFAttributeDefinitions->FindFeature(AttributeNode->GetAttribute("Feature"), Attribute.Feature);
-                                Attribute.MainAttribute = AttributeNode->GetAttribute("MainAttribute");
-                                Attribute.PhysicalUnit = DMXImport::GetEnumValueFromString<EDMXImportGDTFPhysicalUnit>(AttributeNode->GetAttribute("PhysicalUnit"));
-                                Attribute.Color = DMXImport::ParseColorCIE(AttributeNode->GetAttribute("Color"));
-                                GDTFAttributeDefinitions->Attributes.Add(Attribute);
-                            }
-                        }
-                    }
-                }
-
-                if (const FXmlNode* WheelsNode = FixtureTypeNode->FindChildNode("Wheels"))
-                {
-                    for (const FXmlNode* WheelNode : WheelsNode->GetChildrenNodes())
-                    {
-                        FDMXImportGDTFWheel ImportWheel;
-                        ImportWheel.Name = FName(*WheelNode->GetAttribute("Name"));
-
-                        for (const FXmlNode* SlotNode : WheelNode->GetChildrenNodes())
-                        {
-                            FDMXImportGDTFWheelSlot ImportWheelSlot;
-                            ImportWheelSlot.Name = FName(*SlotNode->GetAttribute("Name"));
-                            ImportWheelSlot.Color = DMXImport::ParseColorCIE(SlotNode->GetAttribute("Color"));
-                            ImportWheel.Slots.Add(ImportWheelSlot);
-                        }
-
-                        GDTFWheels->Wheels.Add(ImportWheel);
-                    }
-                }
-
-                if (const FXmlNode* PhysicalDescriptionsNode = FixtureTypeNode->FindChildNode("PhysicalDescriptions"))
-                {
-                    if (const FXmlNode* EmittersNode = PhysicalDescriptionsNode->FindChildNode("Emitters"))
-                    {
-                        for (const FXmlNode* EmitterNode : EmittersNode->GetChildrenNodes())
-                        {
-                            FDMXImportGDTFEmitter ImportEmitter;
-
-                            ImportEmitter.Name = FName(*EmitterNode->GetAttribute("Name"));
-                            ImportEmitter.Color = DMXImport::ParseColorCIE(EmitterNode->GetAttribute("Color"));
-                            ImportEmitter.DiodePart = EmitterNode->GetAttribute("DiodePart");
-                            LexTryParseString(ImportEmitter.DominantWaveLength, *EmitterNode->GetAttribute("DominantWaveLength"));
-                            ImportEmitter.DiodePart = EmitterNode->GetAttribute("DiodePart");
-
-                            if (const FXmlNode* MeasurementNode = EmitterNode->FindChildNode("Measurement"))
-                            {
-                                FDMXImportGDTFMeasurement ImportMeasurement;
-
-                                LexTryParseString(ImportMeasurement.Physical, *MeasurementNode->GetAttribute("DominantWaveLength"));
-                                LexTryParseString(ImportMeasurement.LuminousIntensity, *MeasurementNode->GetAttribute("LuminousIntensity"));
-                                LexTryParseString(ImportMeasurement.Transmission, *MeasurementNode->GetAttribute("Transmission"));
-                                ImportMeasurement.InterpolationTo = DMXImport::GetEnumValueFromString<EDMXImportGDTFInterpolationTo>(MeasurementNode->GetAttribute("InterpolationTo"));
-
-                                for (const FXmlNode* MeasurementPointNode : MeasurementNode->GetChildrenNodes())
-                                {
-                                    FDMXImportGDTFMeasurementPoint ImportMeasurementPoint;
-
-                                    LexTryParseString(ImportMeasurementPoint.Energy, *MeasurementPointNode->GetAttribute("Energy"));
-                                    LexTryParseString(ImportMeasurementPoint.WaveLength, *MeasurementPointNode->GetAttribute("WaveLength"));
-
-                                    ImportMeasurement.MeasurementPoints.Add(ImportMeasurementPoint);
-                                }
-
-                                ImportEmitter.Measurement = ImportMeasurement;
-                            }
-
-                            GDTFPhysicalDescriptions->Emitters.Add(ImportEmitter);
-                        }
-                    }
-
-                    if (const FXmlNode* ColorSpaceNode = PhysicalDescriptionsNode->FindChildNode("ColorSpace"))
-                    {
-                        FDMXImportGDTFColorSpace ImportColorSpace;
-                        ImportColorSpace.Mode = DMXImport::GetEnumValueFromString<EDMXImportGDTFMode>(ColorSpaceNode->GetAttribute("Mode"));
-                        ImportColorSpace.Description = ColorSpaceNode->GetAttribute("Description");
-                        ImportColorSpace.Red = DMXImport::ParseColorCIE(ColorSpaceNode->GetAttribute("Red"));
-                        ImportColorSpace.Blue = DMXImport::ParseColorCIE(ColorSpaceNode->GetAttribute("Blue"));
-                        ImportColorSpace.Green = DMXImport::ParseColorCIE(ColorSpaceNode->GetAttribute("Green"));
-                        ImportColorSpace.WhitePoint = DMXImport::ParseColorCIE(ColorSpaceNode->GetAttribute("WhitePoint"));
-
-                        GDTFPhysicalDescriptions->ColorSpace = ImportColorSpace;
-                    }
-                }
-
-                if (const FXmlNode* ModelsNode = FixtureTypeNode->FindChildNode("Models"))
-                {
-                    for (const FXmlNode* ModelNode : ModelsNode->GetChildrenNodes())
-                    {
-                        FDMXImportGDTFModel ImportModel;
-
-                        ImportModel.Name = FName(*ModelNode->GetAttribute("Name"));
-                        LexTryParseString(ImportModel.Length, *ModelNode->GetAttribute("Length"));
-                        LexTryParseString(ImportModel.Width, *ModelNode->GetAttribute("Width"));
-                        LexTryParseString(ImportModel.Height, *ModelNode->GetAttribute("Height"));
-                        ImportModel.PrimitiveType = DMXImport::GetEnumValueFromString<EDMXImportGDTFPrimitiveType>(ModelNode->GetAttribute("PrimitiveType"));
-
-                        GDTFModels->Models.Add(ImportModel);
-                    }
-                }
-
-                if (const FXmlNode* GeometriesNode = FixtureTypeNode->FindChildNode("Geometries"))
-                {
-                    for (const FXmlNode* GeometryNode : GeometriesNode->GetChildrenNodes())
-                    {
-                        FDMXImportGDTFGeneralGeometry ImportGeneralGeometry;
-                        ImportGeneralGeometry.Name = FName(*GeometryNode->GetAttribute("Name"));
-                        ImportGeneralGeometry.Model = FName(*GeometryNode->GetAttribute("Model"));
-                        ImportGeneralGeometry.Position = DMXImport::ParseMatrix(GeometryNode->GetAttribute("Position"));
-
-                        if (const FXmlNode* GeneralAxisNode = GeometryNode->FindChildNode("Axis"))
-                        {
-                            FDMXImportGDTFGeneralAxis ImportGeneralAxis;
-                            ImportGeneralAxis.Name = FName(*GeneralAxisNode->GetAttribute("Name"));
-                            ImportGeneralAxis.Model = FName(*GeneralAxisNode->GetAttribute("Model"));
-                            ImportGeneralAxis.Position = DMXImport::ParseMatrix(GeneralAxisNode->GetAttribute("Position"));
-
-                            for (const FXmlNode* TypeAxisNode : GeneralAxisNode->GetChildrenNodes())
-                            {
-                                FDMXImportGDTFTypeAxis ImportTypeAxis;
-                                ImportTypeAxis.Name = FName(*TypeAxisNode->GetAttribute("Name"));
-                                ImportTypeAxis.Model = FName(*TypeAxisNode->GetAttribute("Model"));
-                                ImportTypeAxis.Position = DMXImport::ParseMatrix(TypeAxisNode->GetAttribute("Position"));
-
-                                for (const FXmlNode* BeamNode : TypeAxisNode->GetChildrenNodes())
-                                {
-                                    FDMXImportGDTFBeam ImportBeam;
-                                    ImportBeam.Name = FName(*BeamNode->GetAttribute("Name"));
-                                    ImportBeam.Model = FName(*BeamNode->GetAttribute("Model"));
-                                    ImportBeam.Position = DMXImport::ParseMatrix(BeamNode->GetAttribute("Position"));
-                                    ImportBeam.LampType = DMXImport::GetEnumValueFromString<EDMXImportGDTFLampType>(BeamNode->GetAttribute("LampType"));
-
-                                    LexTryParseString(ImportBeam.PowerConsumption, *BeamNode->GetAttribute("PowerConsumption"));
-                                    LexTryParseString(ImportBeam.LuminousFlux, *BeamNode->GetAttribute("LuminousFlux"));
-                                    LexTryParseString(ImportBeam.ColorTemperature, *BeamNode->GetAttribute("ColorTemperature"));
-                                    LexTryParseString(ImportBeam.BeamAngle, *BeamNode->GetAttribute("BeamAngle"));
-                                    LexTryParseString(ImportBeam.FieldAngle, *BeamNode->GetAttribute("FieldAngle"));
-                                    LexTryParseString(ImportBeam.BeamRadius, *BeamNode->GetAttribute("BeamRadius"));
-                                    ImportBeam.BeamType = DMXImport::GetEnumValueFromString<EDMXImportGDTFBeamType>(BeamNode->GetAttribute("BeamType"));
-                                    LexTryParseString(ImportBeam.ColorRenderingIndex, *BeamNode->GetAttribute("ColorRenderingIndex"));
-
-                                    ImportTypeAxis.Beams.Add(ImportBeam);
-                                }
-                                ImportGeneralAxis.Axis.Add(ImportTypeAxis);
-                            }
-                            ImportGeneralGeometry.Axis = ImportGeneralAxis;
-                        }
-
-                        if (const FXmlNode* TypeGeometryNode = GeometryNode->FindChildNode("Geometry"))
-                        {
-                            FDMXImportGDTFTypeGeometry ImportTypeGeometry;
-                            ImportTypeGeometry.Name = FName(*TypeGeometryNode->GetAttribute("Name"));
-                            ImportTypeGeometry.Model = FName(*TypeGeometryNode->GetAttribute("Model"));
-                            ImportTypeGeometry.Position = DMXImport::ParseMatrix(TypeGeometryNode->GetAttribute("Position"));
-
-                            ImportGeneralGeometry.Geometry = ImportTypeGeometry;
-                        }
-
-                        if (const FXmlNode* FilterBeamNode = GeometryNode->FindChildNode("FilterBeam"))
-                        {
-                            FDMXImportGDTFFilterBeam ImportFilterBeam  ;
-                            ImportFilterBeam.Name = FName(*FilterBeamNode->GetAttribute("Name"));
-                            ImportFilterBeam.Model = FName(*FilterBeamNode->GetAttribute("Model"));
-                            ImportFilterBeam.Position = DMXImport::ParseMatrix(FilterBeamNode->GetAttribute("Position"));
-
-                            ImportGeneralGeometry.FilterBeam = ImportFilterBeam ;
-                        }
-
-                        if (const FXmlNode* FilterColorNode = GeometryNode->FindChildNode("FilterColor"))
-                        {
-                            FDMXImportGDTFFilterColor ImportFilterColor;
-                            ImportFilterColor.Name = FName(*FilterColorNode->GetAttribute("Name"));
-                            ImportFilterColor.Model = FName(*FilterColorNode->GetAttribute("Model"));
-                            ImportFilterColor.Position = DMXImport::ParseMatrix(FilterColorNode->GetAttribute("Position"));
-
-                            ImportGeneralGeometry.FilterColor = ImportFilterColor;
-                        }
-
-                        if (const FXmlNode* FilterGoboNode = GeometryNode->FindChildNode("FilterGobo"))
-                        {
-                            FDMXImportGDTFFilterGobo ImportFilterGobo;
-                            ImportFilterGobo.Name = FName(*FilterGoboNode->GetAttribute("Name"));
-                            ImportFilterGobo.Model = FName(*FilterGoboNode->GetAttribute("Model"));
-                            ImportFilterGobo.Position = DMXImport::ParseMatrix(FilterGoboNode->GetAttribute("Position"));
-
-                            ImportGeneralGeometry.FilterGobo = ImportFilterGobo;
-                        }
-
-                        if (const FXmlNode* FilterShaperNode = GeometryNode->FindChildNode("FilterShaper"))
-                        {
-                            FDMXImportGDTFFilterShaper ImportFilterShaper;
-                            ImportFilterShaper.Name = FName(*FilterShaperNode->GetAttribute("Name"));
-                            ImportFilterShaper.Model = FName(*FilterShaperNode->GetAttribute("Model"));
-                            ImportFilterShaper.Position = DMXImport::ParseMatrix(FilterShaperNode->GetAttribute("Position"));
-
-                            ImportGeneralGeometry.FilterShaper = ImportFilterShaper;
-                        }
-
-                        if (const FXmlNode* GeometryReferenceNode = GeometryNode->FindChildNode("GeometryReference"))
-                        {
-                            FDMXImportGDTFGeometryReference ImportGeometryReference;
-                            ImportGeometryReference.Name = FName(*GeometryReferenceNode->GetAttribute("Name"));
-                            ImportGeometryReference.Model = FName(*GeometryReferenceNode->GetAttribute("Model"));
-                            ImportGeometryReference.Position = DMXImport::ParseMatrix(GeometryReferenceNode->GetAttribute("Position"));
-
-                            for (const FXmlNode* BreakNode : GeometryReferenceNode->GetChildrenNodes())
-                            {
-                                FDMXImportGDTFBreak ImportBreak;
-
-								// Depending on the GDTF spec the Offset node may be called either 'DMXOffset' or 'Offset'
-								const FString DMXOffsetAttribute = UE::DMX::DMXGDTFImporter::Private::FindAttributeEvenIfDMXSubstringIsMissing(*BreakNode, TEXT("DMXOffset"));
-								LexTryParseString(ImportBreak.DMXOffset, *DMXOffsetAttribute);
-
-								// Depending on the GDTF spec the Break node may be called either 'DMXBreak' or 'Break'
-								const FString DMXBreakAttribute = UE::DMX::DMXGDTFImporter::Private::FindAttributeEvenIfDMXSubstringIsMissing(*BreakNode, TEXT("DMXBreak"));
-								LexTryParseString(ImportBreak.DMXBreak, *DMXBreakAttribute);
-
-                                ImportGeometryReference.Breaks.Add(ImportBreak);
-                            }
-
-                            ImportGeneralGeometry.GeometryReference = ImportGeometryReference;
-                        }
-
-                        GDTFGeometries->GeneralGeometry.Add(ImportGeneralGeometry);
-                    }
-                }
-
-				// Depending on the GDTF spec the Break node may be called either 'DMXModes' or 'Modes'.
-				const FXmlNode* const ModesNode = UE::DMX::DMXGDTFImporter::Private::FindChildNodeEvenIfDMXSubstringIsMissing(*FixtureTypeNode, TEXT("DMXModes"));
-                if (ModesNode)
-                {
-                    for (const FXmlNode* DMXModeNode : ModesNode->GetChildrenNodes())
-                    {
-                        FDMXImportGDTFDMXMode DMXImportGDTFDMXMode;
-                        DMXImportGDTFDMXMode.Name = FName(*DMXModeNode->GetAttribute("Name"));
-                        DMXImportGDTFDMXMode.Geometry = FName(*DMXModeNode->GetAttribute("Geometry"));
-
-						// Depending on the GDTF spec in use modes may be stored in the "Modes" or "DMXModes" node.
-						const FXmlNode* const ChannelsNode = UE::DMX::DMXGDTFImporter::Private::FindChildNodeEvenIfDMXSubstringIsMissing(*DMXModeNode, TEXT("DMXChannels"));
-                        if (ChannelsNode)
-                        {
-                            for (const FXmlNode* DMXChannelNode : ChannelsNode->GetChildrenNodes())
-                            {
-								FDMXImportGDTFDMXChannel ImportDMXChannel;
-
-								// Ignore Channels that do not specify a valid offset.
-								// E.g. 'Robe Lighting s.r.o.@Robin 800X LEDWash.gdtf' specifies virtual Dimmer channels that have no offset and cannot be accessed.
-								if (!ImportDMXChannel.ParseOffset(DMXChannelNode->GetAttribute("Offset")))
-								{
-									continue;
-								}
-
-                                LexTryParseString(ImportDMXChannel.DMXBreak, *DMXChannelNode->GetAttribute("DMXBreak"));
-                                ImportDMXChannel.Default = FDMXImportGDTFDMXValue(DMXChannelNode->GetAttribute("Default"));
-                                ImportDMXChannel.Highlight = FDMXImportGDTFDMXValue(DMXChannelNode->GetAttribute("Highlight"));
-                                ImportDMXChannel.Geometry = FName(*DMXChannelNode->GetAttribute("Geometry"));
-
-                                if (const FXmlNode* LogicalChannelNode = DMXChannelNode->FindChildNode("LogicalChannel"))
-                                {
-                                    FDMXImportGDTFLogicalChannel ImportLogicalChannel;
-
-                                    GDTFAttributeDefinitions->FindAtributeByName(*LogicalChannelNode->GetAttribute("Attribute"), ImportLogicalChannel.Attribute);
-                                    ImportLogicalChannel.Snap = DMXImport::GetEnumValueFromString<EDMXImportGDTFSnap>(LogicalChannelNode->GetAttribute("Snap"));
-                                    ImportLogicalChannel.Master = DMXImport::GetEnumValueFromString<EDMXImportGDTFMaster>(LogicalChannelNode->GetAttribute("Master"));
-
-                                    LexTryParseString(ImportLogicalChannel.MibFade, *LogicalChannelNode->GetAttribute("MibFade"));
-                                    LexTryParseString(ImportLogicalChannel.DMXChangeTimeLimit, *LogicalChannelNode->GetAttribute("DMXChangeTimeLimit"));
-
-                                    if (const FXmlNode* ChannelFunctionNode = LogicalChannelNode->FindChildNode("ChannelFunction"))
-                                    {
-                                        FDMXImportGDTFChannelFunction ImportChannelFunction;
-                                        ImportChannelFunction.Name = FName(*ChannelFunctionNode->GetAttribute("Name"));
-                                        GDTFAttributeDefinitions->FindAtributeByName(*ChannelFunctionNode->GetAttribute("Attribute"), ImportChannelFunction.Attribute);
-                                        ImportChannelFunction.OriginalAttribute = *ChannelFunctionNode->GetAttribute("OriginalAttribute");
-                                        ImportChannelFunction.DMXFrom = FDMXImportGDTFDMXValue(ChannelFunctionNode->GetAttribute("DMXFrom"));
-                                        ImportChannelFunction.DMXValue = FDMXImportGDTFDMXValue(ChannelFunctionNode->GetAttribute("DMXValue"));
-
-                                        LexTryParseString(ImportChannelFunction.PhysicalFrom, *ChannelFunctionNode->GetAttribute("PhysicalFrom"));
-                                        LexTryParseString(ImportChannelFunction.PhysicalTo, *ChannelFunctionNode->GetAttribute("PhysicalTo"));
-                                        LexTryParseString(ImportChannelFunction.RealFade, *ChannelFunctionNode->GetAttribute("RealFade"));
-
-                                        GDTFWheels->FindWeelByName(*ChannelFunctionNode->GetAttribute("Wheel"), ImportChannelFunction.Wheel);
-                                        GDTFPhysicalDescriptions->FindEmitterByName(*ChannelFunctionNode->GetAttribute("Emitter"), ImportChannelFunction.Emitter);
-
-                                        ImportChannelFunction.DMXInvert = DMXImport::GetEnumValueFromString<EDMXImportGDTFDMXInvert>(ChannelFunctionNode->GetAttribute("DMXInvert"));
-                                        ImportChannelFunction.ModeMaster = *ChannelFunctionNode->GetAttribute("ModeMaster");
-                                        ImportChannelFunction.ModeFrom = FDMXImportGDTFDMXValue(ChannelFunctionNode->GetAttribute("ModeFrom"));
-                                        ImportChannelFunction.ModeTo = FDMXImportGDTFDMXValue(ChannelFunctionNode->GetAttribute("ModeTo"));
-
-                                        for (const FXmlNode* ChannelSetNode : ChannelFunctionNode->GetChildrenNodes())
-                                        {
-                                            FDMXImportGDTFChannelSet ImportChannelSet;
-
-                                            ImportChannelSet.Name = *ChannelSetNode->GetAttribute("Name");
-                                            ImportChannelSet.DMXFrom = FDMXImportGDTFDMXValue(ChannelSetNode->GetAttribute("DMXFrom"));
-                                            LexTryParseString(ImportChannelSet.PhysicalFrom , *ChannelSetNode->GetAttribute("PhysicalFrom"));
-                                            LexTryParseString(ImportChannelSet.PhysicalTo, *ChannelSetNode->GetAttribute("PhysicalTo"));
-                                            LexTryParseString(ImportChannelSet.WheelSlotIndex, *ChannelSetNode->GetAttribute("WheelSlotIndex"));
-
-                                            ImportChannelFunction.ChannelSets.Add(ImportChannelSet);
-                                        }
-
-                                        ImportLogicalChannel.ChannelFunction = ImportChannelFunction;
-                                    }
-
-                                    ImportDMXChannel.LogicalChannel = ImportLogicalChannel;
-                                }
-
-                                DMXImportGDTFDMXMode.DMXChannels.Add(ImportDMXChannel);
-                            }
-                        }
-
-                        if (const FXmlNode* RelationsNode = DMXModeNode->FindChildNode("Relations"))
-                        {
-                            for (const FXmlNode* RelationNode : RelationsNode->GetChildrenNodes())
-                            {
-                                FDMXImportGDTFRelation ImportRelation;
-                                ImportRelation.Name = RelationNode->GetAttribute("Name");
-                                ImportRelation.Master = RelationNode->GetAttribute("Master");
-                                ImportRelation.Follower = RelationNode->GetAttribute("Follower");
-                                ImportRelation.Type = DMXImport::GetEnumValueFromString<EDMXImportGDTFType>(RelationNode->GetAttribute("Type"));
-
-                                DMXImportGDTFDMXMode.Relations.Add(ImportRelation);
-                            }
-                        }
-
-                        if (const FXmlNode* FTMacrosNode = DMXModeNode->FindChildNode("FTMacros"))
-                        {
-                            for (const FXmlNode* FTMacroNode : FTMacrosNode->GetChildrenNodes())
-                            {
-                                FDMXImportGDTFFTMacro ImportFTMacro;
-                                ImportFTMacro.Name = FName(*FTMacroNode->GetAttribute("Name"));
-                                DMXImportGDTFDMXMode.FTMacros.Add(ImportFTMacro);
-                            }
-                        }
-
-                        GDTFDMXModes->DMXModes.Add(DMXImportGDTFDMXMode);
-                    }
-                }
-
-                if (const FXmlNode* ProtocolsNode = FixtureTypeNode->FindChildNode("Protocols"))
-                {
-                    for (const FXmlNode* ProtocolNode : ProtocolsNode->GetChildrenNodes())
-                    {
-                        GDTFProtocols->Protocols.Add(FName(*ProtocolNode->GetTag()));
-                    }
-                }
-            }
-        }
-    }
-
-    return ImportObject;
-}
-
-bool FDMXGDTFImporter::ParseXML()
-{
-    //Load the Compressed data array
-    TArray<uint8> GDTFFileData;
-    if (!FFileHelper::LoadFileToArray(GDTFFileData, *ImportArgs.CurrentFilename))
-    {
-        UE_LOG_DMXEDITOR(Warning, TEXT("InFilename %s"), *ImportArgs.CurrentFilename);
-        return false;
-    }
-
-    // Convert to readable string
-    FString CharString;
-    {
-        const uint8* DataPtr = GDTFFileData.GetData();
-        int32 DataCount = GDTFFileData.Num();
-        CharString.Empty(DataCount);
-        while (DataCount)
-        {
-            uint8 Value = *DataPtr;
-            CharString += TCHAR(Value);
-
-            ++DataPtr;
-            DataCount--;
-        }
-    }
-
-    FString BeginXMLString("<?xml version=");
-    FString EndXMLString("</GDTF>");
-    int32 BeginXMLStringIndex = CharString.Find(BeginXMLString, ESearchCase::CaseSensitive);
-    int32 EndXMLStringIndex = CharString.Find(EndXMLString, ESearchCase::CaseSensitive, ESearchDir::FromEnd);
-    FString ResultXMLString = CharString.Mid(BeginXMLStringIndex, EndXMLStringIndex - BeginXMLStringIndex + EndXMLString.Len());
-
-    // @TODO. UE5 XML Parser can't handle symbol < or > inside attribute and just break the parsing
-    // This small algorithm remove < and > from the XML line
-    TArray<FString> XMLStringLines;
-    ResultXMLString.ParseIntoArrayLines(XMLStringLines, /*bCullEmpty*/false);
-    FString FinalXMLString;
-    for (FString& Line : XMLStringLines)
-    {
-        FinalXMLString.Append(Line + LINE_TERMINATOR);
-    }
-
-    XMLFile = MakeShared<FXmlFile>(FinalXMLString, EConstructMethod::ConstructFromBuffer);
-    bIsXMLParsedSuccessfully = XMLFile->IsValid();
-
-    return bIsXMLParsedSuccessfully;
-}
-
-void FDMXGDTFImporter::GetImportOptions(const TUniquePtr<FDMXGDTFImporter>& Importer, UDMXGDTFImportUI* ImportUI, bool bShowOptionDialog, const FString& FullPath, bool& OutOperationCanceled, bool& bOutImportAll, const FString& InFilename)
-{
-    OutOperationCanceled = false;
-
-    if ( bShowOptionDialog )
-    {
-        TSharedPtr<SWindow> ParentWindow;
-
-        if( FModuleManager::Get().IsModuleLoaded( "MainFrame" ) )
-        {
-            IMainFrameModule& MainFrame = FModuleManager::LoadModuleChecked<IMainFrameModule>( "MainFrame" );
-            ParentWindow = MainFrame.GetParentWindow();
-        }
-
-        // Compute centered window position based on max window size, which include when all categories are expanded
-        const float ImportWindowWidth = 410.0f;
-        const float ImportWindowHeight = 750.0f;
-        FVector2D ImportWindowSize = FVector2D(ImportWindowWidth, ImportWindowHeight); // Max window size it can get based on current slate
-
-        FSlateRect WorkAreaRect = FSlateApplicationBase::Get().GetPreferredWorkArea();
-        FVector2D DisplayTopLeft(WorkAreaRect.Left, WorkAreaRect.Top);
-        FVector2D DisplaySize(WorkAreaRect.Right - WorkAreaRect.Left, WorkAreaRect.Bottom - WorkAreaRect.Top);
-
-        FVector2D WindowPosition = (DisplayTopLeft + (DisplaySize - ImportWindowSize) / 2.0f);
-
-		TSharedRef<SWindow> Window = SNew(SWindow)
-			.Title(LOCTEXT("GDTFImportOpionsTitle", "GDTF Import Options"))
-			.SizingRule(ESizingRule::Autosized)
-			.AutoCenter(EAutoCenter::None)
-			.ClientSize(ImportWindowSize)
-			.ScreenPosition(WindowPosition);
-
-		TSharedPtr<SDMXGDTFOptionWindow> OptionWindow;
-		Window->SetContent
-		(
-			SAssignNew(OptionWindow, SDMXGDTFOptionWindow)
-			.ImportUI(ImportUI)
-			.WidgetWindow(Window)
-			.FullPath(FText::FromString(FullPath))
-			.MaxWindowHeight(ImportWindowHeight)
-			.MaxWindowWidth(ImportWindowWidth)
-		);
+		// Unzip Description.xml
+		constexpr TCHAR DescriptionXmlFilename[] = TEXT("description.xml");
+		TArray64<uint8> DescriptionXmlData;
+		if (!Zip->GetFileContent(DescriptionXmlFilename, DescriptionXmlData))
+		{
+			OutErrorReason = LOCTEXT("InvalidGDTFError", "Cannot read GDTF. Cannot find Description.xml.");
+			return nullptr;
+		}
 
-		// @todo: we can make this slow as showing progress bar later
-		FSlateApplication::Get().AddModalWindow(Window, ParentWindow, false);
+		// Create GDTF
+		UDMXImportGDTF* NewGDTF = NewObject<UDMXImportGDTF>(ImportArgs.Parent.Get(), ImportArgs.Name, ImportArgs.Flags | RF_Public);
 
-        if (OptionWindow->ShouldImport())
-        {
-            bOutImportAll = OptionWindow->ShouldImportAll();
+		// Set Asset Import Data
+		UDMXGDTFAssetImportData* GDTFAssetImportData = NewGDTF->GetGDTFAssetImportData();
+		if (!ensureAlwaysMsgf(GDTFAssetImportData, TEXT("Unexpected missing Asset Import Data for newly created GDTF %s"), *NewGDTF->GetName()))
+		{
+			return nullptr;
+		}
+		GDTFAssetImportData->SetSourceFile(ImportArgs.Filename);
 
-            return;
-        }
-        else
-        {
-            OutOperationCanceled = true;
-        }
-    }
+		return NewGDTF;
+	}
 }
 
 #undef LOCTEXT_NAMESPACE
diff --git a/Engine/Plugins/VirtualProduction/DMX/DMXEngine/Source/DMXEditor/Private/Factories/DMXGDTFImporter.h b/Engine/Plugins/VirtualProduction/DMX/DMXEngine/Source/DMXEditor/Private/Factories/DMXGDTFImporter.h
--- a/Engine/Plugins/VirtualProduction/DMX/DMXEngine/Source/DMXEditor/Private/Factories/DMXGDTFImporter.h
+++ b/Engine/Plugins/VirtualProduction/DMX/DMXEngine/Source/DMXEditor/Private/Factories/DMXGDTFImporter.h
@@ -2,65 +2,54 @@
 
 #pragma once
 
-#include "CoreMinimal.h"
+#include "Containers/UnrealString.h"
+#include "Templates/SharedPointer.h"
+#include "UObject/ObjectMacros.h"
 #include "UObject/WeakObjectPtr.h"
 
+class FDMXZipper;
 class UDMXImportGDTF;
-class UDMXGDTFImportUI;
-class FXmlFile;
-class UPackage;
+class UDMXGDTFFactory;
 
-struct FDMXGDTFImportArgs
+namespace UE::DMX
 {
-public:
-    FDMXGDTFImportArgs()
-        : Parent(nullptr)
-        , Name(NAME_None)
-        , Flags(RF_NoFlags)
-        , ImportUI(nullptr)
-        , bCancelOperation(false)
-    {}
-
-    TWeakObjectPtr<UObject> Parent;
-    FName Name;
-    FString CurrentFilename;
-    EObjectFlags Flags;
-    TWeakObjectPtr<UDMXGDTFImportUI> ImportUI;
-    bool bCancelOperation;
-};
-
-/**
- * GDTF Importer is read and parse GDTF input file
- */
-class FDMXGDTFImporter
-{
-public:
-    FDMXGDTFImporter(const FDMXGDTFImportArgs& InImportArgs);
-
-    /** Try to load the file and parse the data */
-    bool AttemptImportFromFile();
-
-    /** Create new GDTF UObject from imported file */
-    UDMXImportGDTF* Import();
-
-private:
-    bool ParseXML();
-
-    /** Create GDTF from imported xml */
-    UDMXImportGDTF* CreateGDTFDesctription();
-
-    UPackage* GetPackage(const FString& InAssetName);
-public:
-
-    static void GetImportOptions(const TUniquePtr<FDMXGDTFImporter>& Importer, UDMXGDTFImportUI* ImportUI, bool bShowOptionDialog, const FString& FullPath, bool& OutOperationCanceled, bool& bOutImportAll, const FString& InFilename);
-private:
-    TSharedPtr<FXmlFile> XMLFile;
-
-    FDMXGDTFImportArgs ImportArgs;
-
-    bool bIsXMLParsedSuccessfully;
-
-    FString LastPackageName;
-
-    FString ImportName;
-};
\ No newline at end of file
+	struct FDMXGDTFImportArgs
+	{
+	public:
+		FDMXGDTFImportArgs()
+			: Parent(nullptr)
+			, Name(NAME_None)
+			, Flags(RF_NoFlags)
+		{}
+
+		TWeakObjectPtr<UObject> Parent;
+		FName Name;
+		FString Filename;
+		EObjectFlags Flags;
+	};
+
+
+	/** Imports a GDTF asset. */
+	class FDMXGDTFImporter
+	{
+	public:
+		/** Tries to import a GDTF, using params from the import factory. Returns the resulting GDTF object or nullptr if no GDTF asset could be created. */
+		[[nodiscard]] static UDMXImportGDTF* Import(const UDMXGDTFFactory& InImportFactory, const FDMXGDTFImportArgs& InImportArgs, FText& OutErrorReason);
+
+	private:
+		/** Private constructor */
+		FDMXGDTFImporter(const FDMXGDTFImportArgs& InImportArgs);
+
+		/** Non-static implementation */
+		UDMXImportGDTF* ImportInternal(const UDMXGDTFFactory& InImportFactory, FText& OutErrorReason);
+
+		/** Creates the GDTF asset */
+		UDMXImportGDTF* CreateGDTF(const UDMXGDTFFactory& InImportFactory, FText& OutErrorReason) const;
+
+		/** Args for this importer */
+		const FDMXGDTFImportArgs& ImportArgs;
+
+		/** The GDTF File as zip */
+		TSharedPtr<FDMXZipper> Zip;
+	};
+}
diff --git a/Engine/Plugins/VirtualProduction/DMX/DMXEngine/Source/DMXEditor/Private/Factories/DMXLibraryFromMVRImporter.cpp b/Engine/Plugins/VirtualProduction/DMX/DMXEngine/Source/DMXEditor/Private/Factories/DMXLibraryFromMVRImporter.cpp
--- a/Engine/Plugins/VirtualProduction/DMX/DMXEngine/Source/DMXEditor/Private/Factories/DMXLibraryFromMVRImporter.cpp
+++ b/Engine/Plugins/VirtualProduction/DMX/DMXEngine/Source/DMXEditor/Private/Factories/DMXLibraryFromMVRImporter.cpp
@@ -221,7 +221,7 @@ void FDMXLibraryFromMVRImporter::InitializeDMXLibrary(UDMXLibrary* DMXLibrary, c
 
 		UDMXEntityFixtureType* const* ExistingFixtureTypePtr = Algo::FindByPredicate(FixtureTypes, [GDTF](const UDMXEntityFixtureType* FixtureType)
 			{
-				return FixtureType->DMXImport == GDTF;
+				return FixtureType->GDTFSource == GDTF;
 			});
 		if (ExistingFixtureTypePtr)
 		{
@@ -237,10 +237,12 @@ void FDMXLibraryFromMVRImporter::InitializeDMXLibrary(UDMXLibrary* DMXLibrary, c
 			const bool bAdvancedImportSuccess = FDMXInitializeFixtureTypeFromGDTFHelper::GenerateModesFromGDTF(*NewFixtureType, *GDTF);
 			if (!bAdvancedImportSuccess)
 			{
+				PRAGMA_DISABLE_DEPRECATION_WARNINGS
 				UE_LOG(LogDMXEditor, Warning, TEXT("Failed to initialize Fixture Type '%s', falling back to legacy method that doesn't support matrix fixtures."), *NewFixtureType->GetName());
 				NewFixtureType->SetModesFromDMXImport(GDTF);
+				PRAGMA_ENABLE_DEPRECATION_WARNINGS
 			}
-			NewFixtureType->DMXImport = GDTF;
+			NewFixtureType->GDTFSource = GDTF;
 
 			GDTFSpecToFixtureTypeMap.Add(GDTFFilename, NewFixtureType);
 		}
diff --git a/Engine/Plugins/VirtualProduction/DMX/DMXEngine/Source/DMXEditor/Private/Widgets/SDMXGDTFOptionWindow.cpp b/Engine/Plugins/VirtualProduction/DMX/DMXEngine/Source/DMXEditor/Private/Widgets/SDMXGDTFOptionWindow.cpp
--- a/Engine/Plugins/VirtualProduction/DMX/DMXEngine/Source/DMXEditor/Private/Widgets/SDMXGDTFOptionWindow.cpp
+++ b/Engine/Plugins/VirtualProduction/DMX/DMXEngine/Source/DMXEditor/Private/Widgets/SDMXGDTFOptionWindow.cpp
@@ -14,159 +14,211 @@
 
 #define LOCTEXT_NAMESPACE "SDMXGDTFOptionWindow"
 
-void SDMXGDTFOptionWindow::Construct(const FArguments& InArgs)
+namespace UE::DMX
 {
-    check (InArgs._ImportUI);
-
-    ImportUI = InArgs._ImportUI;
-    WidgetWindow = InArgs._WidgetWindow;
-
-    TSharedPtr<SBox> ImportTypeDisplay;
-    TSharedPtr<SHorizontalBox> HeaderButtons;
-    TSharedPtr<SBox> InspectorBox;
-
-    this->ChildSlot
-    [
-        SNew(SBox)
-            .MaxDesiredHeight(InArgs._MaxWindowHeight)
-            .MaxDesiredWidth(InArgs._MaxWindowWidth)
-        [
-            SNew(SVerticalBox)
-            + SVerticalBox::Slot()
-                .AutoHeight()
-                .Padding(2)
-            [
-                SAssignNew(ImportTypeDisplay, SBox)
-            ]
-            +SVerticalBox::Slot()
-                .AutoHeight()
-                .Padding(2)
-            [
-                SNew(SBorder)
-                    .Padding(FMargin(3))
-                    .BorderImage(FAppStyle::GetBrush("ToolPanel.GroupBorder"))
-                [
-                    SNew(SHorizontalBox)
-                    +SHorizontalBox::Slot()
-                        .AutoWidth()
-                    [
-                        SNew(STextBlock)
-                            .Font(FAppStyle::GetFontStyle("CurveEd.LabelFont"))
-                            .Text(LOCTEXT("Import_CurrentFileTitle", "Current Asset: "))
-                    ]
-                    +SHorizontalBox::Slot()
-                        .Padding(5, 0, 0, 0)
-                        .AutoWidth()
-                        .VAlign(VAlign_Center)
-                    [
-                        SNew(STextBlock)
-                            .Font(FAppStyle::GetFontStyle("CurveEd.InfoFont"))
-                            .Text(InArgs._FullPath)
-                            .ToolTipText(InArgs._FullPath)
-                    ]
-                ]
-            ]
-            + SVerticalBox::Slot()
-                .AutoHeight()
-                .Padding(2)
-            [
-                SAssignNew(InspectorBox, SBox)
-                    .MaxDesiredHeight(650.0f)
-                    .WidthOverride(400.0f)
-            ]
-            + SVerticalBox::Slot()
-                .AutoHeight()
-                .HAlign(HAlign_Right)
-                .Padding(2)
-            [
-                SNew(SUniformGridPanel)
-                    .SlotPadding(2)
-                + SUniformGridPanel::Slot(1, 0)
-                [
-                    SNew(SButton)
-                        .HAlign(HAlign_Center)
-                        .Text(LOCTEXT("OptionWindow_ImportAll", "Import All"))
-                        .ToolTipText(LOCTEXT("OptionWindow_ImportAll_ToolTip", "Import all files with these same settings"))
-                        .OnClicked(this, &SDMXGDTFOptionWindow::OnImportAll)
-                ]
-                + SUniformGridPanel::Slot(2, 0)
-                [
-                    SAssignNew(ImportButton, SButton)
-                        .HAlign(HAlign_Center)
-                        .Text(LOCTEXT("OptionWindow_Import", "Import"))
-                        .OnClicked(this, &SDMXGDTFOptionWindow::OnImport)
-                ]
-                + SUniformGridPanel::Slot(3, 0)
-                [
-                    SNew(SButton)
-                        .HAlign(HAlign_Center)
-                        .Text(LOCTEXT("OptionWindow_Cancel", "Cancel"))
-                        .ToolTipText(LOCTEXT("OptionWindow_Cancel_ToolTip", "Cancels importing this file"))
-                        .OnClicked(this, &SDMXGDTFOptionWindow::OnCancel)
-                ]
-            ]
-        ]
-    ];
-
-    FPropertyEditorModule& PropertyEditorModule = FModuleManager::GetModuleChecked<FPropertyEditorModule>("PropertyEditor");
-    FDetailsViewArgs DetailsViewArgs;
-    DetailsViewArgs.bAllowSearch = false;
-    DetailsViewArgs.NameAreaSettings = FDetailsViewArgs::HideNameArea;
-    DetailsView = PropertyEditorModule.CreateDetailView(DetailsViewArgs);
-
-    InspectorBox->SetContent(DetailsView->AsShared());
-
-    ImportTypeDisplay->SetContent(
-        SNew(SBorder)
-            .Padding(FMargin(3))
-            .BorderImage(FAppStyle::GetBrush("ToolPanel.GroupBorder"))
-        [
-            SNew(SHorizontalBox)
-            + SHorizontalBox::Slot()
-                .VAlign(VAlign_Center)
-            [
-                SNew(STextBlock)
-                    .Text(this, &SDMXGDTFOptionWindow::GetImportTypeDisplayText)
-            ]
-            + SHorizontalBox::Slot()
-            [
-                SNew(SBox)
-                    .HAlign(HAlign_Right)
-                [
-                    SAssignNew(HeaderButtons, SHorizontalBox)
-                    + SHorizontalBox::Slot()
-                        .AutoWidth()
-                        .Padding(FMargin(2.0f, 0.0f))
-                    [
-                        SNew(SButton)
-                            .Text(LOCTEXT("OptionWindow_ResetOptions", "Reset to Default"))
-                            .OnClicked(this, &SDMXGDTFOptionWindow::OnResetToDefaultClick)
-                    ]
-                ]
-            ]
-        ]
-    );
-
-    if (UDMXGDTFImportUI* ImportUIPtr = ImportUI.Get())
-    {
-		DetailsView->SetObject(ImportUIPtr);
-    }
-}
+	void SDMXGDTFOptionWindow::Construct(const FArguments& InArgs)
+	{
+		check(InArgs._ImportUI);
 
-FReply SDMXGDTFOptionWindow::OnResetToDefaultClick() const
-{
-    if (UDMXGDTFImportUI* ImportUIPtr = ImportUI.Get())
-    {
-        ImportUIPtr->ResetToDefault();
-		//Refresh the view to make sure the custom UI are updating correctly
-		DetailsView->SetObject(ImportUIPtr, true);
-    }
-    return FReply::Handled();
-}
+		ImportUI = InArgs._ImportUI;
+		WidgetWindow = InArgs._WidgetWindow;
+
+		TSharedPtr<SBox> ImportTypeDisplay;
+		TSharedPtr<SHorizontalBox> HeaderButtons;
+		TSharedPtr<SBox> InspectorBox;
+
+		this->ChildSlot
+			[
+				SNew(SBox)
+				.MaxDesiredHeight(InArgs._MaxWindowHeight)
+				.MaxDesiredWidth(InArgs._MaxWindowWidth)
+				[
+					SNew(SVerticalBox)
+
+					+ SVerticalBox::Slot()
+					.AutoHeight()
+					.Padding(2.f)
+					[
+						SAssignNew(ImportTypeDisplay, SBox)
+					]
+
+					+ SVerticalBox::Slot()
+					.AutoHeight()
+					.Padding(2.f)
+					[
+						SNew(SBorder)
+						.Padding(FMargin(3))
+						.BorderImage(FAppStyle::GetBrush("ToolPanel.GroupBorder"))
+						[
+							SNew(SHorizontalBox)
+
+							+ SHorizontalBox::Slot()
+							.AutoWidth()
+							[
+								SNew(STextBlock)
+								.Font(FAppStyle::GetFontStyle("CurveEd.LabelFont"))
+								.Text(LOCTEXT("Import_CurrentFileTitle", "Current Asset: "))
+							]
+
+							+ SHorizontalBox::Slot()
+							.Padding(5.f, 0.f, 0.f, 0.f)
+							.AutoWidth()
+							.VAlign(VAlign_Center)
+							[
+								SNew(STextBlock)
+								.Font(FAppStyle::GetFontStyle("CurveEd.InfoFont"))
+								.Text(InArgs._FullPath)
+								.ToolTipText(InArgs._FullPath)
+							]
+						]
+					]
+
+					+ SVerticalBox::Slot()
+					.AutoHeight()
+					.Padding(2.f)
+					[
+						SAssignNew(InspectorBox, SBox)
+						.MaxDesiredHeight(650.0f)
+						.WidthOverride(400.0f)
+					]
+
+					+ SVerticalBox::Slot()
+					.AutoHeight()
+					.HAlign(HAlign_Right)
+					.Padding(2.f)
+					[
+						SNew(SUniformGridPanel)
+						.SlotPadding(2.f)
+						+ SUniformGridPanel::Slot(1, 0)
+						[
+							SNew(SButton)
+							.HAlign(HAlign_Center)
+							.Text(LOCTEXT("OptionWindow_ImportAll", "Import All"))
+							.ToolTipText(LOCTEXT("OptionWindow_ImportAll_ToolTip", "Import all files with these same settings"))
+							.OnClicked(this, &SDMXGDTFOptionWindow::OnImportAll)
+						]
+
+						+ SUniformGridPanel::Slot(2, 0)
+						[
+							SAssignNew(ImportButton, SButton)
+							.HAlign(HAlign_Center)
+							.Text(LOCTEXT("OptionWindow_Import", "Import"))
+							.OnClicked(this, &SDMXGDTFOptionWindow::OnImport)
+						]
+
+						+ SUniformGridPanel::Slot(3, 0)
+						[
+							SNew(SButton)
+							.HAlign(HAlign_Center)
+							.Text(LOCTEXT("OptionWindow_Cancel", "Cancel"))
+							.ToolTipText(LOCTEXT("OptionWindow_Cancel_ToolTip", "Cancels importing this file"))
+							.OnClicked(this, &SDMXGDTFOptionWindow::OnCancel)
+						]
+					]
+				]
+			];
+
+		FPropertyEditorModule& PropertyEditorModule = FModuleManager::GetModuleChecked<FPropertyEditorModule>("PropertyEditor");
+		FDetailsViewArgs DetailsViewArgs;
+		DetailsViewArgs.bAllowSearch = false;
+		DetailsViewArgs.NameAreaSettings = FDetailsViewArgs::HideNameArea;
+		DetailsView = PropertyEditorModule.CreateDetailView(DetailsViewArgs);
+
+		InspectorBox->SetContent(DetailsView->AsShared());
+
+		ImportTypeDisplay->SetContent(
+			SNew(SBorder)
+			.Padding(FMargin(3))
+			.BorderImage(FAppStyle::GetBrush("ToolPanel.GroupBorder"))
+			[
+				SNew(SHorizontalBox)
+
+				+ SHorizontalBox::Slot()
+				.VAlign(VAlign_Center)
+				[
+					SNew(STextBlock)
+					.Text(this, &SDMXGDTFOptionWindow::GetImportTypeDisplayText)
+				]
+
+				+ SHorizontalBox::Slot()
+				[
+					SNew(SBox)
+					.HAlign(HAlign_Right)
+					[
+						SAssignNew(HeaderButtons, SHorizontalBox)
+
+						+ SHorizontalBox::Slot()
+						.AutoWidth()
+						.Padding(FMargin(2.f, 0.f))
+						[
+							SNew(SButton)
+							.Text(LOCTEXT("OptionWindow_ResetOptions", "Reset to Default"))
+							.OnClicked(this, &SDMXGDTFOptionWindow::OnResetToDefaultClick)
+						]
+					]
+				]
+			]
+		);
+
+		if (UDMXGDTFImportUI* ImportUIPtr = ImportUI.Get())
+		{
+			DetailsView->SetObject(ImportUIPtr);
+		}
+	}
+
+	FReply SDMXGDTFOptionWindow::OnImport()
+	{
+		bShouldImport = true;
+		if (WidgetWindow.IsValid())
+		{
+			WidgetWindow.Pin()->RequestDestroyWindow();
+		}
+		return FReply::Handled();
+	}
+
+	FReply SDMXGDTFOptionWindow::OnImportAll()
+	{
+		bShouldImportAll = true;
+		return OnImport();
+	}
+
+	FReply SDMXGDTFOptionWindow::OnCancel()
+	{
+		bShouldImport = false;
+		bShouldImportAll = false;
+		if (WidgetWindow.IsValid())
+		{
+			WidgetWindow.Pin()->RequestDestroyWindow();
+		}
+		return FReply::Handled();
+	}
+
+	FReply SDMXGDTFOptionWindow::OnKeyDown(const FGeometry& MyGeometry, const FKeyEvent& InKeyEvent)
+	{
+		if (InKeyEvent.GetKey() == EKeys::Escape)
+		{
+			return OnCancel();
+		}
+
+		return FReply::Unhandled();
+	}
+
+	FReply SDMXGDTFOptionWindow::OnResetToDefaultClick() const
+	{
+		if (UDMXGDTFImportUI* ImportUIPtr = ImportUI.Get())
+		{
+			ImportUIPtr->ResetToDefault();
+			//Refresh the view to make sure the custom UI are updating correctly
+			DetailsView->SetObject(ImportUIPtr, true);
+		}
+		return FReply::Handled();
+	}
+
+	FText SDMXGDTFOptionWindow::GetImportTypeDisplayText() const
+	{
+		return FText::FromString(TEXT("Import GDTF"));
+	}
 
-FText SDMXGDTFOptionWindow::GetImportTypeDisplayText() const
-{
-    return FText::FromString(TEXT("Import GDTF"));
 }
 
 #undef LOCTEXT_NAMESPACE
diff --git a/Engine/Plugins/VirtualProduction/DMX/DMXEngine/Source/DMXEditor/Private/Widgets/SDMXGDTFOptionWindow.h b/Engine/Plugins/VirtualProduction/DMX/DMXEngine/Source/DMXEditor/Private/Widgets/SDMXGDTFOptionWindow.h
--- a/Engine/Plugins/VirtualProduction/DMX/DMXEngine/Source/DMXEditor/Private/Widgets/SDMXGDTFOptionWindow.h
+++ b/Engine/Plugins/VirtualProduction/DMX/DMXEngine/Source/DMXEditor/Private/Widgets/SDMXGDTFOptionWindow.h
@@ -13,92 +13,52 @@ class SButton;
 class UDMXGDTFImportUI;
 class IDetailsView;
 
-class SDMXGDTFOptionWindow : public SCompoundWidget
+namespace UE::DMX
 {
-public:
-    SLATE_BEGIN_ARGS( SDMXGDTFOptionWindow )
-    : _ImportUI(nullptr)
-    , _WidgetWindow()
-    , _FullPath()
-    , _MaxWindowHeight(0.0f)
-    , _MaxWindowWidth(0.0f)
-    {}
 
-    SLATE_ARGUMENT( UDMXGDTFImportUI*, ImportUI )
-    SLATE_ARGUMENT( TSharedPtr<SWindow>, WidgetWindow )
-    SLATE_ARGUMENT( FText, FullPath )
-    SLATE_ARGUMENT( float, MaxWindowHeight)
-    SLATE_ARGUMENT(float, MaxWindowWidth)
-    SLATE_END_ARGS()
-
-public:
-    void Construct(const FArguments& InArgs);
-    virtual bool SupportsKeyboardFocus() const override { return true; }
-
-    FReply OnImport()
-    {
-        bShouldImport = true;
-        if ( WidgetWindow.IsValid() )
-        {
-            WidgetWindow.Pin()->RequestDestroyWindow();
-        }
-        return FReply::Handled();
-    }
-
-    FReply OnImportAll()
-    {
-        bShouldImportAll = true;
-        return OnImport();
-    }
-
-    FReply OnCancel()
-    {
-        bShouldImport = false;
-        bShouldImportAll = false;
-        if ( WidgetWindow.IsValid() )
-        {
-            WidgetWindow.Pin()->RequestDestroyWindow();
-        }
-        return FReply::Handled();
-    }
-
-    virtual FReply OnKeyDown( const FGeometry& MyGeometry, const FKeyEvent& InKeyEvent ) override
-    {
-        if( InKeyEvent.GetKey() == EKeys::Escape )
-        {
-            return OnCancel();
-        }
-
-        return FReply::Unhandled();
-    }
-
-    bool ShouldImport() const
-    {
-        return bShouldImport;
-    }
-
-    bool ShouldImportAll() const
-    {
-        return bShouldImportAll;
-    }
-
-    SDMXGDTFOptionWindow()
-        : ImportUI(nullptr)
-        , bShouldImport(false)
-        , bShouldImportAll(false)
-    {}
-
-private:
-
-    FReply OnResetToDefaultClick() const;
-    FText GetImportTypeDisplayText() const;
-
-
-private:
-    TWeakObjectPtr<UDMXGDTFImportUI> ImportUI;
-    TSharedPtr<IDetailsView> DetailsView;
-    TWeakPtr< SWindow > WidgetWindow;
-    TSharedPtr< SButton > ImportButton;
-    bool			bShouldImport;
-    bool			bShouldImportAll;
-};
+	class SDMXGDTFOptionWindow : public SCompoundWidget
+	{
+	public:
+		SLATE_BEGIN_ARGS(SDMXGDTFOptionWindow)
+			: _ImportUI(nullptr)
+			, _WidgetWindow()
+			, _FullPath()
+			, _MaxWindowHeight(0.0f)
+			, _MaxWindowWidth(0.0f)
+			{}
+
+			SLATE_ARGUMENT(UDMXGDTFImportUI*, ImportUI)
+			SLATE_ARGUMENT(TSharedPtr<SWindow>, WidgetWindow)
+			SLATE_ARGUMENT(FText, FullPath)
+			SLATE_ARGUMENT(float, MaxWindowHeight)
+			SLATE_ARGUMENT(float, MaxWindowWidth)
+
+		SLATE_END_ARGS()
+
+	public:
+		void Construct(const FArguments& InArgs);
+
+		bool ShouldImport() const { return bShouldImport;  }
+
+		bool ShouldImportAll() const { return bShouldImportAll; }
+
+	private:
+		//~ Begin SWidget interface
+		virtual bool SupportsKeyboardFocus() const override { return true; }
+		virtual FReply OnKeyDown(const FGeometry& MyGeometry, const FKeyEvent& InKeyEvent) override;
+		//~ End SWidget interface
+
+		FReply OnImport();
+		FReply OnImportAll();
+		FReply OnCancel();
+		FReply OnResetToDefaultClick() const;
+		FText GetImportTypeDisplayText() const;
+
+		TWeakObjectPtr<UDMXGDTFImportUI> ImportUI;
+		TSharedPtr<IDetailsView> DetailsView;
+		TWeakPtr< SWindow > WidgetWindow;
+		TSharedPtr< SButton > ImportButton;
+		bool bShouldImport = false;
+		bool bShouldImportAll = false;
+	};
+}
diff --git a/Engine/Plugins/VirtualProduction/DMX/DMXEngine/Source/DMXRuntime/DMXRuntime.Build.cs b/Engine/Plugins/VirtualProduction/DMX/DMXEngine/Source/DMXRuntime/DMXRuntime.Build.cs
--- a/Engine/Plugins/VirtualProduction/DMX/DMXEngine/Source/DMXRuntime/DMXRuntime.Build.cs
+++ b/Engine/Plugins/VirtualProduction/DMX/DMXEngine/Source/DMXRuntime/DMXRuntime.Build.cs
@@ -14,8 +14,9 @@ public class DMXRuntime : ModuleRules
 				"Core",
                 "CoreUObject",
 				"DatasmithContent",
-				"Engine",
+				"DMXGDTF",
 				"DMXProtocol",
+				"Engine",
 				"JsonUtilities",
 			}
 		);
diff --git a/Engine/Plugins/VirtualProduction/DMX/DMXEngine/Source/DMXRuntime/Private/DMXRuntimeMainStreamObjectVersion.h b/Engine/Plugins/VirtualProduction/DMX/DMXEngine/Source/DMXRuntime/Private/DMXRuntimeMainStreamObjectVersion.h
--- a/Engine/Plugins/VirtualProduction/DMX/DMXEngine/Source/DMXRuntime/Private/DMXRuntimeMainStreamObjectVersion.h
+++ b/Engine/Plugins/VirtualProduction/DMX/DMXEngine/Source/DMXRuntime/Private/DMXRuntimeMainStreamObjectVersion.h
@@ -30,6 +30,9 @@ struct FDMXRuntimeMainStreamObjectVersion
 		// Remove auto assign features from fixture patch (instead it's auto assigned via FDMXEditorUtils where required)
 		DMXFixturePatchNoLongerImplementsAutoAssign,
 
+		// Change the DMXImportGDTF asset reference in Fixture Type to a soft object ptr, to avoid loading the data
+		DMXImportGDTFIsASoftObjectPtr,
+
 		VersionPlusOne,
 		LatestVersion = VersionPlusOne - 1
 	};
diff --git a/Engine/Plugins/VirtualProduction/DMX/DMXEngine/Source/DMXRuntime/Private/Library/DMXEntityFixtureType.cpp b/Engine/Plugins/VirtualProduction/DMX/DMXEngine/Source/DMXRuntime/Private/Library/DMXEntityFixtureType.cpp
--- a/Engine/Plugins/VirtualProduction/DMX/DMXEngine/Source/DMXRuntime/Private/Library/DMXEntityFixtureType.cpp
+++ b/Engine/Plugins/VirtualProduction/DMX/DMXEngine/Source/DMXRuntime/Private/Library/DMXEntityFixtureType.cpp
@@ -3,11 +3,13 @@
 #include "Library/DMXEntityFixtureType.h"
 
 #include "DMXConversions.h"
+#include "DMXGDTF.h"
 #include "DMXProtocolSettings.h"
 #include "DMXRuntimeLog.h"
 #include "DMXRuntimeMainStreamObjectVersion.h"
 #include "DMXRuntimeUtils.h"
 #include "Library/DMXEntityFixturePatch.h"
+#include "Library/DMXGDTFAssetImportData.h"
 #include "Library/DMXImport.h"
 #include "Library/DMXImportGDTF.h"
 #include "Library/DMXLibrary.h"
@@ -158,6 +160,9 @@ int32 FDMXFixtureMode::AddOrInsertFunction(int32 IndexOfFunction, FDMXFixtureFun
 
 FDMXOnFixtureTypeChangedDelegate UDMXEntityFixtureType::OnFixtureTypeChangedDelegate;
 
+UDMXEntityFixtureType::UDMXEntityFixtureType()
+{}
+
 UDMXEntityFixtureType* UDMXEntityFixtureType::CreateFixtureTypeInLibrary(FDMXEntityFixtureTypeConstructionParams ConstructionParams, const FString& DesiredName, bool bMarkDMXLibraryDirty)
 {
 	UDMXLibrary* ParentDMXLibrary = ConstructionParams.ParentDMXLibrary;
@@ -184,7 +189,6 @@ UDMXEntityFixtureType* UDMXEntityFixtureType::CreateFixtureTypeInLibrary(FDMXEnt
 			ParentDMXLibrary->PostEditChange();
 		}
 #endif
-
 		OnFixtureTypeChangedDelegate.Broadcast(NewFixtureType);
 
 		return NewFixtureType;
@@ -228,10 +232,24 @@ void UDMXEntityFixtureType::Serialize(FArchive& Ar)
 			for (FDMXFixtureMode& Mode : Modes)
 			{
 				PRAGMA_DISABLE_DEPRECATION_WARNINGS
-				Mode.bFixtureMatrixEnabled = bFixtureMatrixEnabled;
+				Mode.bFixtureMatrixEnabled = bFixtureMatrixEnabled_DEPRECATED;
 				PRAGMA_ENABLE_DEPRECATION_WARNINGS
 			}
 		}
+
+#if WITH_EDITOR
+		if (Ar.CustomVer(FDMXRuntimeMainStreamObjectVersion::GUID) < FDMXRuntimeMainStreamObjectVersion::DMXImportGDTFIsASoftObjectPtr)
+		{
+			PRAGMA_DISABLE_DEPRECATION_WARNINGS
+			if (DMXImport && DMXImport->GetClass() == UDMXImportGDTF::StaticClass())
+			{
+				UDMXImportGDTF* DMXImportGDTF = CastChecked<UDMXImportGDTF>(DMXImport);
+				GDTFSource = DMXImportGDTF;
+				DMXImport = nullptr;
+			}
+			PRAGMA_ENABLE_DEPRECATION_WARNINGS
+		}
+#endif // WITH_EDITOR
 	}
 }
 
@@ -272,17 +290,6 @@ void UDMXEntityFixtureType::PostEditChangeChainProperty(FPropertyChangedChainEve
 	Super::PostEditChangeChainProperty(PropertyChangedChainEvent);
 
 	const FName PropertyName = PropertyChangedChainEvent.GetPropertyName();
-	if (PropertyName == GET_MEMBER_NAME_CHECKED(UDMXEntityFixtureType, Modes))
-	{
-		int32 ChangedModeIndex = PropertyChangedChainEvent.GetArrayIndex(PropertyName.ToString());
-		if (Modes.IsValidIndex(ChangedModeIndex))
-		{
-			// Notify DataType changes (Deprecated 5.0)
-			PRAGMA_DISABLE_DEPRECATION_WARNINGS
-			DataTypeChangeDelegate_DEPRECATED.Broadcast(this, Modes[ChangedModeIndex]);
-			PRAGMA_ENABLE_DEPRECATION_WARNINGS
-		}
-	}
 
 	if (PropertyChangedChainEvent.ChangeType != EPropertyChangeType::Interactive)
 	{
@@ -310,16 +317,19 @@ void UDMXEntityFixtureType::PostEditUndo()
 #if WITH_EDITOR
 void UDMXEntityFixtureType::SetModesFromDMXImport(UDMXImport* DMXImportAsset)
 {
+	// DEPRECATED 5.4
+
 	if (!IsValid(DMXImportAsset))
 	{
 		return;
 	}
 
-	DMXImport = DMXImportAsset;
-	if (UDMXImportGDTFDMXModes* GDTFDMXModes = Cast<UDMXImportGDTFDMXModes>(DMXImport->DMXModes))
+	PRAGMA_DISABLE_DEPRECATION_WARNINGS
+	GDTFSource = DMXImportAsset;
+	if (UDMXImportGDTFDMXModes* GDTFDMXModes = Cast<UDMXImportGDTFDMXModes>(GDTFSource->DMXModes_DEPRECATED))
 	{
 		// Clear existing modes
-		Modes.Empty(DMXImport->DMXModes != nullptr ? GDTFDMXModes->DMXModes.Num() : 0);
+		Modes.Empty(GDTFSource->DMXModes_DEPRECATED != nullptr ? GDTFDMXModes->DMXModes.Num() : 0);
 
 		// Used to map Functions to Attributes
 		const UDMXProtocolSettings* ProtocolSettings = GetDefault<UDMXProtocolSettings>();
@@ -435,6 +445,7 @@ void UDMXEntityFixtureType::SetModesFromDMXImport(UDMXImport* DMXImportAsset)
 			UpdateChannelSpan(ModeIndex);
 		}
 	}
+	PRAGMA_ENABLE_DEPRECATION_WARNINGS
 
 	OnFixtureTypeChangedDelegate.Broadcast(this);
 }
@@ -1168,234 +1179,4 @@ float UDMXEntityFixtureType::BytesToNormalizedValue(EDMXFixtureSignalFormat InSi
 	return Value / FDMXConversions::GetSignalFormatMaxValue(InSignalFormat);
 }
 
-
-#if WITH_EDITOR
-/** DEPRECATED 5.0 */
-FDataTypeChangeDelegate UDMXEntityFixtureType::DataTypeChangeDelegate_DEPRECATED;
-#endif
-
-uint32 UDMXEntityFixtureType::GetDataTypeMaxValue(EDMXFixtureSignalFormat DataType)
-{	
-	// DEPRECATED 5.0
-	switch (DataType)
-	{
-	case EDMXFixtureSignalFormat::E8Bit:
-		return MAX_uint8;
-	case EDMXFixtureSignalFormat::E16Bit:
-		return MAX_uint16;
-	case EDMXFixtureSignalFormat::E24Bit:
-		return 0xFFFFFF;
-	case EDMXFixtureSignalFormat::E32Bit:
-		return MAX_uint32;
-	default:
-		checkNoEntry();
-		return 1;
-	}
-}
-
-uint8 UDMXEntityFixtureType::NumChannelsToOccupy(EDMXFixtureSignalFormat DataType)
-{
-	// DEPRECATED 5.0
-	switch (DataType)
-	{
-	case EDMXFixtureSignalFormat::E8Bit:
-		return 1;
-
-	case EDMXFixtureSignalFormat::E16Bit:
-		return 2;
-
-	case EDMXFixtureSignalFormat::E24Bit:
-		return 3;
-
-	case EDMXFixtureSignalFormat::E32Bit:
-		return 4;
-
-	default:
-		// Unhandled type
-		checkNoEntry();
-		break;
-	}
-	return 1;
-}
-
-uint8 UDMXEntityFixtureType::GetFunctionLastChannel(const FDMXFixtureFunction& Function)
-{
-	// DEPRECATED 5.0
-	return Function.Channel + Function.GetNumChannels() - 1;
-}
-
-bool UDMXEntityFixtureType::IsFunctionInModeRange(const FDMXFixtureFunction& InFunction, const FDMXFixtureMode& InMode, int32 ChannelOffset /*= 0*/)
-{
-	// DEPRECATED 5.0
-	const int32 LastChannel = InFunction.GetLastChannel();
-	const bool bLastChannelExceedsChannelSpan = LastChannel > InMode.ChannelSpan;
-	const bool bLastChannelExceedsUniverseSize = LastChannel + ChannelOffset > DMX_MAX_ADDRESS;
-
-	return !bLastChannelExceedsChannelSpan && !bLastChannelExceedsUniverseSize;
-}
-
-bool UDMXEntityFixtureType::IsFixtureMatrixInModeRange(const FDMXFixtureMatrix& InFixtureMatrix, const FDMXFixtureMode& InMode, int32 ChannelOffset /*= 0*/)
-{
-	// DEPRECATED 5.0
-	const int32 LastChannel = InFixtureMatrix.GetLastChannel();
-	const bool bLastChannelExceedsChannelSpan = LastChannel > InMode.ChannelSpan;
-	const bool bLastChannelExceedsUniverseSize = LastChannel + ChannelOffset > DMX_MAX_ADDRESS;
-
-	return !bLastChannelExceedsChannelSpan && !bLastChannelExceedsUniverseSize;
-}
-
-void UDMXEntityFixtureType::ClampDefaultValue(FDMXFixtureFunction& InFunction)
-{
-	// DEPRECATED 5.0
-	const uint32 SafeDefaultValue = FMath::Min(InFunction.DefaultValue, static_cast<int64>(TNumericLimits<uint32>::Max()));
-	InFunction.DefaultValue = FDMXConversions::ClampValueBySignalFormat(SafeDefaultValue, InFunction.DataType);
-}
-
-uint32 UDMXEntityFixtureType::ClampValueToDataType(EDMXFixtureSignalFormat DataType, uint32 InValue)
-{
-	// DEPRECATED 5.0
-	switch (DataType)
-	{
-	case EDMXFixtureSignalFormat::E8Bit:
-		return FMath::Clamp(InValue, 0u, (uint32)MAX_uint8);
-
-	case EDMXFixtureSignalFormat::E16Bit:
-		return FMath::Clamp(InValue, 0u, (uint32)MAX_uint16);
-
-	case EDMXFixtureSignalFormat::E24Bit:
-		return FMath::Clamp(InValue, 0u, 0xFFFFFFu);
-
-	case EDMXFixtureSignalFormat::E32Bit:
-		return FMath::Clamp(InValue, 0u, MAX_uint32);
-
-	default:
-		break;
-	}
-	return InValue;
-}
-
-#if WITH_EDITOR
-void UDMXEntityFixtureType::SetFunctionSize(FDMXFixtureFunction& InFunction, uint8 Size)
-{
-	// DEPRECATED 5.0
-	EDMXFixtureSignalFormat NewDataType;
-	switch (Size)
-	{
-	case 0:
-	case 1:
-		NewDataType = EDMXFixtureSignalFormat::E8Bit;
-		break;
-	case 2:
-		NewDataType = EDMXFixtureSignalFormat::E16Bit;
-		break;
-	case 3:
-		NewDataType = EDMXFixtureSignalFormat::E24Bit;
-		break;
-	case 4:
-	default:
-		NewDataType = EDMXFixtureSignalFormat::E32Bit;
-		break;
-	}
-
-	InFunction.DataType = NewDataType;
-
-	const uint32 SafeDefaultValue = FMath::Min(InFunction.DefaultValue, static_cast<int64>(TNumericLimits<uint32>::Max()));
-	InFunction.DefaultValue = FDMXConversions::ClampValueBySignalFormat(SafeDefaultValue, InFunction.DataType);
-}
-#endif // WITH_EDITOR
-
-#if WITH_EDITOR
-void UDMXEntityFixtureType::UpdateModeChannelProperties(FDMXFixtureMode& Mode)
-{
-	// DEPRECATED 4.27
-	PRAGMA_DISABLE_DEPRECATION_WARNINGS
-	UpdateChannelSpan(Mode);
-	PRAGMA_ENABLE_DEPRECATION_WARNINGS
-}
-#endif // WITH_EDITOR
-
-#if WITH_EDITOR
-void UDMXEntityFixtureType::UpdateChannelSpan(FDMXFixtureMode& Mode)
-{
-	// DEPRECATED 5.0
-	if (Mode.bAutoChannelSpan)
-	{
-		if (Mode.Functions.Num() == 0 &&
-			Mode.FixtureMatrixConfig.CellAttributes.Num() == 0)
-		{
-			Mode.ChannelSpan = 0;
-		}
-		else
-		{
-			int32 ChannelSpan = 0;
-
-			// Update span from common Functions
-			for (FDMXFixtureFunction& Function : Mode.Functions)
-			{
-				Function.Channel = ChannelSpan + 1;
-
-				switch (Function.DataType)
-				{
-				case EDMXFixtureSignalFormat::E8Bit:
-					ChannelSpan = Function.Channel;
-					break;
-				case EDMXFixtureSignalFormat::E16Bit:
-					ChannelSpan = Function.Channel + 1;
-					break;
-				case EDMXFixtureSignalFormat::E24Bit:
-					ChannelSpan = Function.Channel + 2;
-					break;
-				case EDMXFixtureSignalFormat::E32Bit:
-					ChannelSpan = Function.Channel + 3;
-					break;
-				default:
-					checkNoEntry();
-					break;
-				}
-			}
-
-			// If fixture matrix is enabled, add the channel span of the matrix
-			int32 NumCells = Mode.FixtureMatrixConfig.XCells * Mode.FixtureMatrixConfig.YCells;
-			if (Mode.bFixtureMatrixEnabled && NumCells > 0)
-			{
-				// Add 'empty' channels bewtween normal Functions and the Matrix to the channel span
-				ChannelSpan = FMath::Max(ChannelSpan, Mode.FixtureMatrixConfig.FirstCellChannel);
-
-				ChannelSpan += Mode.FixtureMatrixConfig.GetNumChannels();
-			}
-
-			ChannelSpan = FMath::Max(ChannelSpan, 1);
-			Mode.ChannelSpan = ChannelSpan;
-		}
-
-		// Notify DataType changes (Deprecated 5.0)
-		PRAGMA_DISABLE_DEPRECATION_WARNINGS
-		DataTypeChangeDelegate_DEPRECATED.Broadcast(this, Mode);
-		PRAGMA_ENABLE_DEPRECATION_WARNINGS
-	}
-}
-#endif // WITH_EDITOR
-
-#if WITH_EDITOR
-void UDMXEntityFixtureType::UpdateYCellsFromXCells(FDMXFixtureMode& Mode)
-{
-	// Deprecated 5.0
-	const int32 MaxNumCells = 512;
-
-	Mode.FixtureMatrixConfig.XCells = FMath::Clamp(Mode.FixtureMatrixConfig.XCells, 1, MaxNumCells);
-	Mode.FixtureMatrixConfig.YCells = FMath::Clamp(Mode.FixtureMatrixConfig.YCells, 1, MaxNumCells - Mode.FixtureMatrixConfig.XCells + 1);
-}
-#endif // WITH_EDITOR
-
-#if WITH_EDITOR
-void UDMXEntityFixtureType::UpdateXCellsFromYCells(FDMXFixtureMode& Mode)
-{
-	// Deprecated 5.0
-	const int32 MaxNumCells = 512;
-
-	Mode.FixtureMatrixConfig.YCells = FMath::Clamp(Mode.FixtureMatrixConfig.YCells, 1, MaxNumCells);
-	Mode.FixtureMatrixConfig.XCells = FMath::Clamp(Mode.FixtureMatrixConfig.XCells, 1, MaxNumCells - Mode.FixtureMatrixConfig.YCells + 1);
-}
-#endif // WITH_EDITOR
-
 #undef LOCTEXT_NAMESPACE
diff --git a/Engine/Plugins/VirtualProduction/DMX/DMXEngine/Source/DMXRuntime/Private/Library/DMXImport.cpp b/Engine/Plugins/VirtualProduction/DMX/DMXEngine/Source/DMXRuntime/Private/Library/DMXImport.cpp
--- a/Engine/Plugins/VirtualProduction/DMX/DMXEngine/Source/DMXRuntime/Private/Library/DMXImport.cpp
+++ b/Engine/Plugins/VirtualProduction/DMX/DMXEngine/Source/DMXRuntime/Private/Library/DMXImport.cpp
@@ -5,6 +5,7 @@
 
 namespace DMXImport
 {
+	PRAGMA_DISABLE_DEPRECATION_WARNINGS
 	FDMXColorCIE ParseColorCIE(const FString& InColor)
 	{
 		TArray<FString> ColorCIEArray;
@@ -33,7 +34,9 @@ namespace DMXImport
 
 		return MoveTemp(ColorCIE);
 	}
+	PRAGMA_ENABLE_DEPRECATION_WARNINGS
 
+	PRAGMA_DISABLE_DEPRECATION_WARNINGS
 	FMatrix ParseMatrix(FString&& InMatrixStr)
 	{
 		FMatrix Matrix = FMatrix::Identity;
@@ -72,5 +75,5 @@ namespace DMXImport
 
 		return MoveTemp(Matrix);
 	}
+	PRAGMA_ENABLE_DEPRECATION_WARNINGS
 }
-
diff --git a/Engine/Plugins/VirtualProduction/DMX/DMXEngine/Source/DMXRuntime/Private/Library/DMXImportGDTF.cpp b/Engine/Plugins/VirtualProduction/DMX/DMXEngine/Source/DMXRuntime/Private/Library/DMXImportGDTF.cpp
--- a/Engine/Plugins/VirtualProduction/DMX/DMXEngine/Source/DMXRuntime/Private/Library/DMXImportGDTF.cpp
+++ b/Engine/Plugins/VirtualProduction/DMX/DMXEngine/Source/DMXRuntime/Private/Library/DMXImportGDTF.cpp
@@ -2,13 +2,57 @@
 
 #include "Library/DMXImportGDTF.h"
 
+#include "DMXGDTF.h"
+#include "GDTF/DMXGDTFDescription.h"
 #include "Library/DMXGDTFAssetImportData.h"
 
+#if WITH_EDITORONLY_DATA
 #include "EditorFramework/AssetImportData.h"
+#endif 
+
+UDMXImportGDTF::UDMXImportGDTF()
+{
+	GDTFAssetImportData = NewObject<UDMXGDTFAssetImportData>(this, TEXT("GDTFAssetImportData"), RF_Public);
+}
+
+void UDMXImportGDTF::PostLoad()
+{
+	Super::PostLoad();
+
+	// Upgrade so this object always holds asset data.
+	if (!GDTFAssetImportData)
+	{
+		GDTFAssetImportData = NewObject<UDMXGDTFAssetImportData>(this, TEXT("GDTFAssetImportData"));
+	}
+
+#if WITH_EDITORONLY_DATA
+	PRAGMA_DISABLE_DEPRECATION_WARNINGS
+	if (!SourceFilename_DEPRECATED.IsEmpty())
+	{
+		GDTFAssetImportData->SetSourceFile(SourceFilename_DEPRECATED);
+
+		SourceFilename_DEPRECATED.Empty();
+	}
+	PRAGMA_ENABLE_DEPRECATION_WARNINGS
+#endif
+}
 
+UDMXGDTF* UDMXImportGDTF::LoadGDTF() const
+{
+	UDMXGDTF* GDTF = NewObject<UDMXGDTF>();
+	if (GDTFAssetImportData)
+	{
+		GDTF->InitializeFromData(GDTFAssetImportData->GetRawSourceData());
+	}
+
+	return GDTF;
+}
 
+PRAGMA_DISABLE_DEPRECATION_WARNINGS
 bool UDMXImportGDTFAttributeDefinitions::FindFeature(const FString& InQuery, FDMXImportGDTFFeature& OutFeature) const
 {
+	// DEPRECATED 5.4
+
     if (InQuery.IsEmpty())
     {
         return false;
@@ -37,9 +81,13 @@ bool UDMXImportGDTFAttributeDefinitions::FindFeature(const FString& InQuery, FDM
 
     return false;
 }
+PRAGMA_ENABLE_DEPRECATION_WARNINGS
 
+PRAGMA_DISABLE_DEPRECATION_WARNINGS
 bool UDMXImportGDTFAttributeDefinitions::FindAtributeByName(const FName& InName, FDMXImportGDTFAttribute& OutAttribute) const
-{
+{	
+	// DEPRECATED 5.4
+
     if (InName.IsNone())
     {
         return false;
@@ -56,9 +104,13 @@ bool UDMXImportGDTFAttributeDefinitions::FindAtributeByName(const FName& InName,
 
     return false;
 }
+PRAGMA_ENABLE_DEPRECATION_WARNINGS
 
+PRAGMA_DISABLE_DEPRECATION_WARNINGS
 bool UDMXImportGDTFWheels::FindWeelByName(const FName& InName, FDMXImportGDTFWheel& OutWheel) const
-{
+{	
+	// DEPRECATED 5.4
+
 	if (InName.IsNone())
     {
         return false;
@@ -75,9 +127,13 @@ bool UDMXImportGDTFWheels::FindWeelByName(const FName& InName, FDMXImportGDTFWhe
 
     return false;
 }
+PRAGMA_ENABLE_DEPRECATION_WARNINGS
 
+PRAGMA_DISABLE_DEPRECATION_WARNINGS
 bool UDMXImportGDTFPhysicalDescriptions::FindEmitterByName(const FName& InName, FDMXImportGDTFEmitter& OutEmitter) const
 {
+	// DEPRECATED 5.4
+
 	if (InName.IsNone())
     {
 		return false;
@@ -94,11 +150,15 @@ bool UDMXImportGDTFPhysicalDescriptions::FindEmitterByName(const FName& InName,
 
 	return false;
 }
+PRAGMA_ENABLE_DEPRECATION_WARNINGS
 
+PRAGMA_DISABLE_DEPRECATION_WARNINGS
 FDMXImportGDTFDMXValue::FDMXImportGDTFDMXValue(const FString& InDMXValueStr)
     : Value(0)
     , ValueSize(1)
-{
+{	
+	// DEPRECATED 5.4
+
     if (!InDMXValueStr.IsEmpty())
     {
         if (InDMXValueStr.Equals("None"))
@@ -116,9 +176,13 @@ FDMXImportGDTFDMXValue::FDMXImportGDTFDMXValue(const FString& InDMXValueStr)
         }
     }
 };
+PRAGMA_ENABLE_DEPRECATION_WARNINGS
 
+PRAGMA_DISABLE_DEPRECATION_WARNINGS
 bool FDMXImportGDTFDMXChannel::ParseOffset(const FString& InOffsetStr)
-{
+{	
+	// DEPRECATED 5.4
+
 	if (InOffsetStr.IsEmpty())
 	{
 		return false;
@@ -137,9 +201,13 @@ bool FDMXImportGDTFDMXChannel::ParseOffset(const FString& InOffsetStr)
     
 	return true;
 }
+PRAGMA_ENABLE_DEPRECATION_WARNINGS
 
+PRAGMA_DISABLE_DEPRECATION_WARNINGS
 TArray<FDMXImportGDTFChannelFunction> UDMXImportGDTFDMXModes::GetDMXChannelFunctions(const FDMXImportGDTFDMXMode& InMode)
 {
+	// DEPRECATED 5.4
+
 	TArray<FDMXImportGDTFChannelFunction> Channels;
 	for (const FDMXImportGDTFDMXChannel& ModeChannel : InMode.DMXChannels)
 	{
@@ -147,31 +215,4 @@ TArray<FDMXImportGDTFChannelFunction> UDMXImportGDTFDMXModes::GetDMXChannelFunct
 	}
 	return MoveTemp(Channels);
 };
-
-UDMXImportGDTF::UDMXImportGDTF()
-{
-#if WITH_EDITORONLY_DATA
-	GDTFAssetImportData = NewObject<UDMXGDTFAssetImportData>(this, TEXT("GDTFAssetImportData"), RF_Public);
-#endif
-}
-
-void UDMXImportGDTF::PostLoad()
-{
-	Super::PostLoad();
-
-#if WITH_EDITORONLY_DATA
-	if (!GDTFAssetImportData)
-	{
-		GDTFAssetImportData = NewObject<UDMXGDTFAssetImportData>(this, TEXT("GDTFAssetImportData"));
-	}
-
-	PRAGMA_DISABLE_DEPRECATION_WARNINGS
-	if (!SourceFilename_DEPRECATED.IsEmpty())
-	{
-		GDTFAssetImportData->SetSourceFile(SourceFilename_DEPRECATED);
-
-		SourceFilename_DEPRECATED.Empty();
-	}
-	PRAGMA_ENABLE_DEPRECATION_WARNINGS
-#endif
-}
+PRAGMA_ENABLE_DEPRECATION_WARNINGS
diff --git a/Engine/Plugins/VirtualProduction/DMX/DMXEngine/Source/DMXRuntime/Private/MVR/DMXMVRGeneralSceneDescription.cpp b/Engine/Plugins/VirtualProduction/DMX/DMXEngine/Source/DMXRuntime/Private/MVR/DMXMVRGeneralSceneDescription.cpp
--- a/Engine/Plugins/VirtualProduction/DMX/DMXEngine/Source/DMXRuntime/Private/MVR/DMXMVRGeneralSceneDescription.cpp
+++ b/Engine/Plugins/VirtualProduction/DMX/DMXEngine/Source/DMXRuntime/Private/MVR/DMXMVRGeneralSceneDescription.cpp
@@ -200,7 +200,8 @@ void UDMXMVRGeneralSceneDescription::WriteFixturePatchToGeneralSceneDescription(
 	{
 		MVRFixtureNode->GDTFMode = FixtureType->Modes[ModeIndex].ModeName;
 
-		if (UDMXImportGDTF* GDTF = Cast<UDMXImportGDTF>(FixtureType->DMXImport))
+		UDMXImportGDTF* GDTF = FixtureType->GDTFSource.LoadSynchronous();
+		if (GDTF)
 		{
 			const FString SourceFilename = [GDTF]()
 			{
@@ -210,6 +211,7 @@ void UDMXMVRGeneralSceneDescription::WriteFixturePatchToGeneralSceneDescription(
 				}
 				return FString();
 			}();
+
 			MVRFixtureNode->GDTFSpec = FPaths::GetCleanFilename(SourceFilename);
 
 			bSetGDTFSpec = true;
