DMX - Upgrade to a new GDTF implementation, initial dev version. UDMXImportGDTF now only holds the raw GDTF data, its GDTF description is deprecated. DMXGDTF now holds the GDTF standard in its own res

diff --git a/Engine/Plugins/VirtualProduction/DMX/DMXEngine/Source/DMXRuntime/Private/MVR/DMXMVRSceneActor.cpp b/Engine/Plugins/VirtualProduction/DMX/DMXEngine/Source/DMXRuntime/Private/MVR/DMXMVRSceneActor.cpp
--- a/Engine/Plugins/VirtualProduction/DMX/DMXEngine/Source/DMXRuntime/Private/MVR/DMXMVRSceneActor.cpp
+++ b/Engine/Plugins/VirtualProduction/DMX/DMXEngine/Source/DMXRuntime/Private/MVR/DMXMVRSceneActor.cpp
@@ -314,7 +314,7 @@ TArray<AActor*> ADMXMVRSceneActor::GetActorsSpawnedForGDTF(const UDMXImportGDTF*
 			continue;
 		}
 
-		if (FixturePatch->GetFixtureType()->DMXImport == GDTF)
+		if (FixturePatch->GetFixtureType()->GDTFSource == GDTF)
 		{
 			Result.Add(Actor);
 		}
@@ -349,11 +349,11 @@ void ADMXMVRSceneActor::UpdateGDTFToDefaultActorClasses(const TSharedRef<FDMXMVR
 	for (UDMXEntityFixturePatch* FixturePatch : FixturePatches)
 	{
 		const UDMXEntityFixtureType* FixtureType = FixturePatch->GetFixtureType();
-		if (FixtureType && FixtureType->DMXImport)
+		if (FixtureType && !FixtureType->GDTFSource.IsNull())
 		{
 			const FDMXMVRSceneGDTFToActorClassPair* ExistingGDTFToActorClassPairPtr = Algo::FindByPredicate(GDTFToDefaultActorClasses, [FixtureType](const FDMXMVRSceneGDTFToActorClassPair& GDTFToActorClassPair)
 				{
-					return GDTFToActorClassPair.GDTF == FixtureType->DMXImport;
+					return GDTFToActorClassPair.GDTF == FixtureType->GDTFSource;
 				});
 			if (ExistingGDTFToActorClassPairPtr)
 			{
@@ -362,7 +362,7 @@ void ADMXMVRSceneActor::UpdateGDTFToDefaultActorClasses(const TSharedRef<FDMXMVR
 
 			UClass* ActorClass = MVRFixtureActorLibrary->FindMostAppropriateActorClassForPatch(FixturePatch);
 			FDMXMVRSceneGDTFToActorClassPair GDTFToActorClassPair;
-			GDTFToActorClassPair.GDTF = FixtureType->DMXImport;
+			GDTFToActorClassPair.GDTF = FixtureType->GDTFSource;
 			GDTFToActorClassPair.ActorClass = ActorClass;
 
 			GDTFToDefaultActorClasses.Add(GDTFToActorClassPair);
@@ -376,7 +376,7 @@ void ADMXMVRSceneActor::UpdateGDTFToDefaultActorClasses(const TSharedRef<FDMXMVR
 				{
 					const UDMXEntityFixturePatch* FixturePatch = Actor.IsValid() ? GetFixturePatchFromActor(Actor.Get()) : nullptr;
 					const UDMXEntityFixtureType* FixtureType = FixturePatch ? FixturePatch->GetFixtureType() : nullptr;
-					const UDMXImport* GDTF = FixtureType ? FixtureType->DMXImport : nullptr;
+					const UDMXImportGDTF* GDTF = FixtureType ? FixtureType->GDTFSource.LoadSynchronous() : nullptr;
 					
 					return GDTFToActorClassPair.GDTF == GDTF;
 				});
@@ -495,7 +495,7 @@ void ADMXMVRSceneActor::HandleDefaultActorClassForGDTFChanged()
 		UDMXEntityFixturePatch* FixturePatch = GetFixturePatchFromActor(Actor);
 		if (FixturePatch && 
 			FixturePatch->GetFixtureType() && 
-			FixturePatch->GetFixtureType()->DMXImport == GDTFToDefaultActorClasses[IndexOfChangedElement].GDTF)
+			FixturePatch->GetFixtureType()->GDTFSource == GDTFToDefaultActorClasses[IndexOfChangedElement].GDTF)
 		{
 			ReplaceMVRActor(Actor, Class);
 		}
@@ -517,7 +517,7 @@ AActor* ADMXMVRSceneActor::SpawnMVRActor(const TSubclassOf<AActor>&ActorClass, U
 		return nullptr;
 	}
 
-	if (!FixturePatch->GetFixtureType() || !FixturePatch->GetFixtureType()->DMXImport)
+	if (!FixturePatch->GetFixtureType() || FixturePatch->GetFixtureType()->GDTFSource.IsNull())
 	{
 		UE_LOG(LogDMXRuntime, Warning, TEXT("Cannot spawn fixture for Fixture Patch '%s'. It's Fixture Type does not source from a GDTF."), *FixturePatch->Name);
 		return nullptr;
diff --git a/Engine/Plugins/VirtualProduction/DMX/DMXEngine/Source/DMXRuntime/Public/Library/DMXEntityFixtureType.h b/Engine/Plugins/VirtualProduction/DMX/DMXEngine/Source/DMXRuntime/Public/Library/DMXEntityFixtureType.h
--- a/Engine/Plugins/VirtualProduction/DMX/DMXEngine/Source/DMXRuntime/Public/Library/DMXEntityFixtureType.h
+++ b/Engine/Plugins/VirtualProduction/DMX/DMXEngine/Source/DMXRuntime/Public/Library/DMXEntityFixtureType.h
@@ -10,6 +10,7 @@
 
 #include "DMXEntityFixtureType.generated.h"
 
+class UDMXGDTF;
 class UDMXImport;
 class UDMXImportGDTF;
 
@@ -293,6 +294,8 @@ class DMXRUNTIME_API UDMXEntityFixtureType
 	GENERATED_BODY()
 
 public:
+	UDMXEntityFixtureType();
+
 	/** Creates a new Fixture Type in the DMX Library */
 	UFUNCTION(BlueprintCallable, Category = "DMX")
 	static UDMXEntityFixtureType* CreateFixtureTypeInLibrary(FDMXEntityFixtureTypeConstructionParams ConstructionParams, const FString& DesiredName = TEXT(""), bool bMarkDMXLibraryDirty = true);
@@ -313,17 +316,14 @@ public:
 
 public:
 #if WITH_EDITOR
-	UFUNCTION(BlueprintCallable, Category = "Fixture Settings")
+	UE_DEPRECATED(5.5, "Setting GDTFs this way is not supported. Instead set the GDTFSource and generate the modes via UDMXGDTF.")
+	UFUNCTION(BlueprintCallable, Category = "Fixture Settings", Meta = (DeprecatedFunction, DeprecationMessage = "Setting GDTFs from blueprints was never fully supported and now deprecated. Instead please refer to the Create Fixture Type In DMX Library function."))
 	void SetModesFromDMXImport(UDMXImport* DMXImportAsset);
 #endif // WITH_EDITOR
 
 	/** Returns a delegate that is and should be broadcast whenever a Fixture Type changed */
 	static FDMXOnFixtureTypeChangedDelegate& GetOnFixtureTypeChanged();
 
-	/** The GDTF file from which the Fixture Type was setup */
-	UPROPERTY(EditAnywhere, BlueprintReadOnly, Category = "Fixture Settings")
-	TObjectPtr<UDMXImport> DMXImport;
-
 	/** The Category of the Fixture, useful for Filtering */
 	UPROPERTY(EditAnywhere, BlueprintReadOnly, Category = "Fixture Settings", meta = (DisplayName = "DMX Category"))
 	FDMXFixtureCategory DMXCategory;
@@ -471,7 +471,7 @@ public:
 	/** Adds a new cell attribute to the Mode */
 	void AddCellAttribute(int32 ModeIndex);
 
-	/** Removes a cell attribute to the Mode */
+	/** Removes a cell attribute from the Mode */
 	void RemoveCellAttribute(int32 ModeIndex, int32 CellAttributeIndex);
 
 	/**
@@ -513,59 +513,22 @@ public:
 	static float BytesToFunctionNormalizedValue(const FDMXFixtureFunction& InFunction, const uint8* InBytes);
 	static float BytesToNormalizedValue(EDMXFixtureSignalFormat InSignalFormat, bool bUseLSB, const uint8* InBytes);
 
+	/** The GDTF that initializes this fixture type. When changed, reinitializes with data from the GDTF. */
+	UPROPERTY(EditAnywhere, BlueprintReadOnly, Category = "Fixture Settings", Meta = (DisplayName = "GDTF Source"))
+	TSoftObjectPtr<UDMXImportGDTF> GDTFSource;
 
 
 	//////////////////////////////////////////////////
 	// Deprecated Members
 	
 public:
-	UE_DEPRECATED(5.0, "Deprecated in favor of FDMXConversions::GetSignalFormatMaxValue.")
-	static uint32 GetDataTypeMaxValue(EDMXFixtureSignalFormat DataType);
-
-	UE_DEPRECATED(5.0, "Deprecated in favor of FDMXConversions::SizeOfSignalFormat.")
-	static uint8 NumChannelsToOccupy(EDMXFixtureSignalFormat DataType);
-
-	UE_DEPRECATED(5.0, "Deprecated in favor of FDMXFixtureFunction::GetLastChannel.")
-	static uint8 GetFunctionLastChannel(const FDMXFixtureFunction& Function);
-
-	UE_DEPRECATED(5.0, "Deprecated to reduce redundant code. Instead use FDMXFixtureFunction::GetNumChannels() and FDMXFixtureMode::ChannelSpan")
-	static bool IsFunctionInModeRange(const FDMXFixtureFunction& InFunction, const FDMXFixtureMode& InMode, int32 ChannelOffset = 0);
-
-	UE_DEPRECATED(5.0, "Deprecated since the use of the function and its name were not clear, leading to hard to read, complicated code.")
-	static bool IsFixtureMatrixInModeRange(const FDMXFixtureMatrix& InFixtureMatrix, const FDMXFixtureMode& InMode, int32 ChannelOffset = 0);
-
-	UE_DEPRECATED(5.0, "Deprecated to reduce redundant code. Use FDMXConversions::ClampValueBySignalFormat locally where clamping is required.")
-	static void ClampDefaultValue(FDMXFixtureFunction& InFunction);
-
-	UE_DEPRECATED(5.0, "Deprecated in favor of FDMXConversions now holding all conversions. Use FDMXConversions::ClampValueBySignalFormat instead.")
-	static uint32 ClampValueToDataType(EDMXFixtureSignalFormat DataType, uint32 InValue);
-
-#if WITH_EDITOR
-	UE_DEPRECATED(5.0, "Deprecated because of unclear use. Set via FDMXFixtureFunction::DataType directly instead.")
-	static void SetFunctionSize(FDMXFixtureFunction& InFunction, uint8 Size);
-
-	UE_DEPRECATED(5.0, "Deprecated in favor of the more generic GetOnFixtureTypeChanged which is supported in non-editor builds as well.")
-	static FDataTypeChangeDelegate& GetDataTypeChangeDelegate() { return DataTypeChangeDelegate_DEPRECATED; }
-
-	UE_DEPRECATED(4.27, "Use MakeValid instead.")
-	void UpdateModeChannelProperties(FDMXFixtureMode& Mode);
-
-	UE_DEPRECATED(5.0, "Deprecated in favor of UDMXEntityFixtureType::UpdateChannelSpan(int32 ModeIndex).")
-	void UpdateChannelSpan(FDMXFixtureMode& Mode);
-
-	UE_DEPRECATED(5.0, "Deprecated  in favor of UDMXEntityFixtureType::UpdateYCellsFromXCells(int32 ModeIndex).")
-	void UpdateYCellsFromXCells(FDMXFixtureMode& Mode);
-
-	UE_DEPRECATED(5.0, "Deprecated to in favor of UDMXEntityFixtureType::UpdateXCellsFromYCells(int32 ModeIndex).")
-	void UpdateXCellsFromYCells(FDMXFixtureMode& Mode);
-#endif // WITH_EDITOR
-
-	UPROPERTY(meta = (DeprecatedProperty, DeprecationMessage = "FixtureMatrixEnabled is deprecated. Instead now each Mode has a FixtureMatrixEnabled property."))
-	bool bFixtureMatrixEnabled = false;
-
-private:
-#if WITH_EDITOR
-	/** Editor only data type change delagate */
-	static FDataTypeChangeDelegate DataTypeChangeDelegate_DEPRECATED;
+#if WITH_EDITORONLY_DATA
+	/** DEPRECATED 5.5 - The GDTF from which this Fixture Type was setup. */
+	UE_DEPRECATED(5.5, "Changed to a soft object pointer to reduce the memory footprint of fixture types. Please refer to GDTFSource instead.")
+	UPROPERTY(BlueprintReadOnly, Category = "Fixture Settings", Meta = (DeprecatedProperty, DeprecationMessage = "Changed to a soft object pointer to reduce the memory footprint of fixture types. Please use GDTF Source instead."))
+	TObjectPtr<UDMXImport> DMXImport;
 #endif
+
+	UPROPERTY(Meta = (DeprecatedProperty, DeprecationMessage = "FixtureMatrixEnabled is deprecated. Instead now each Mode has a FixtureMatrixEnabled property."))
+	bool bFixtureMatrixEnabled_DEPRECATED = false;
 };
diff --git a/Engine/Plugins/VirtualProduction/DMX/DMXEngine/Source/DMXRuntime/Public/Library/DMXGDTFAssetImportData.h b/Engine/Plugins/VirtualProduction/DMX/DMXEngine/Source/DMXRuntime/Public/Library/DMXGDTFAssetImportData.h
--- a/Engine/Plugins/VirtualProduction/DMX/DMXEngine/Source/DMXRuntime/Public/Library/DMXGDTFAssetImportData.h
+++ b/Engine/Plugins/VirtualProduction/DMX/DMXEngine/Source/DMXRuntime/Public/Library/DMXGDTFAssetImportData.h
@@ -25,10 +25,10 @@ public:
 #if WITH_EDITOR
 	/** Creates new asset GDTF Asset Import Data from the File. Returns the GDTF Asset Import Data or nullptr if not a valid GDTF File */
 	void SetSourceFile(const FString& InFilePathAndName);
+#endif // WITH_EDITOR 
 
 	/** Returns the source data the asset was generated from */
 	FORCEINLINE const TArray64<uint8>& GetRawSourceData() const { return RawSourceData.ByteArray; }
-#endif // WITH_EDITOR 
 
 protected:
 	//~ Begin UObject interface
@@ -40,9 +40,7 @@ private:
 	UPROPERTY()
 	FString FilePathAndName;
 
-#if WITH_EDITORONLY_DATA
 	/** The raw GDTF zip file as byte array */
 	UPROPERTY()
 	FDMXByteArray64 RawSourceData;
-#endif // WITH_EDITORONLY_DATA
 };
diff --git a/Engine/Plugins/VirtualProduction/DMX/DMXEngine/Source/DMXRuntime/Public/Library/DMXImport.h b/Engine/Plugins/VirtualProduction/DMX/DMXEngine/Source/DMXRuntime/Public/Library/DMXImport.h
--- a/Engine/Plugins/VirtualProduction/DMX/DMXEngine/Source/DMXRuntime/Public/Library/DMXImport.h
+++ b/Engine/Plugins/VirtualProduction/DMX/DMXEngine/Source/DMXRuntime/Public/Library/DMXImport.h
@@ -6,94 +6,101 @@
 #include "UObject/Object.h"
 #include "DMXImport.generated.h"
 
-namespace DMXImport
+namespace UE::DMX::DMXImport::Private
 {
-    template <typename EnumType> static EnumType GetEnumValueFromString(const FString& String)
-    {
-        return static_cast<EnumType>(StaticEnum<EnumType>()->GetValueByName(FName(*String)));
-    }
+	template <typename EnumType> 
+	UE_DEPRECATED(5.5, "Please refer to the new GDTF Description in UDMXDeviceType instead.")
+	static EnumType GetEnumValueFromString(const FString& String)
+	{
+		return static_cast<EnumType>(StaticEnum<EnumType>()->GetValueByName(FName(*String)));
+	}
 
-    DMXRUNTIME_API FDMXColorCIE ParseColorCIE(const FString& InColor);
+	UE_DEPRECATED(5.5, "Please refer to the new GDTF Description in UDMXDeviceType instead.")
+	DMXRUNTIME_API FDMXColorCIE ParseColorCIE(const FString& InColor);
 
-    DMXRUNTIME_API FMatrix ParseMatrix(FString&& InMatrixStr);
+	UE_DEPRECATED(5.5, "Please refer to the new GDTF Description in UDMXDeviceType instead.")
+	DMXRUNTIME_API FMatrix ParseMatrix(FString&& InMatrixStr);
 };
 
 USTRUCT(BlueprintType)
-struct FDMXColorCIE
+struct UE_DEPRECATED(5.5, "FDMXColorCIE is deprecated. Please refer to the new GDTF Description in UDMXDeviceType instead") FDMXColorCIE
 {
-    GENERATED_BODY()
+	GENERATED_BODY()
 
-        FDMXColorCIE()
-        : X(0.f)
-        , Y(0.f)
-        , YY(0)
-    {}
+	FDMXColorCIE()
+		: X(0.f)
+		, Y(0.f)
+		, YY(0)
+	{}
 
-    UPROPERTY(EditAnywhere, Category = Color, meta = (ClampMin = "0", ClampMax = "1.0"))
-    float X;
+	UPROPERTY(EditAnywhere, Category = Color, meta = (ClampMin = "0", ClampMax = "1.0"))
+	float X;
 
-    UPROPERTY(EditAnywhere, Category = Color, meta = (ClampMin = "0", ClampMax = "1.0"))
-    float Y;
+	UPROPERTY(EditAnywhere, Category = Color, meta = (ClampMin = "0", ClampMax = "1.0"))
+	float Y;
 
-    UPROPERTY(EditAnywhere, Category = Color, meta = (ClampMin = "0", ClampMax = "255"))
-    uint8 YY;
+	UPROPERTY(EditAnywhere, Category = Color, meta = (ClampMin = "0", ClampMax = "255"))
+	uint8 YY;
 };
 
-UCLASS(BlueprintType, Blueprintable)
-class DMXRUNTIME_API UDMXImportFixtureType
-    : public UObject
+UCLASS(BlueprintType, Blueprintable, Meta = (DeprecatedNode, DeprecationMessage = "This class is now deprecated in favor of a DMXEntityFixtureType based workflow. Please refer to members of DMXEntityFixtureType instead."))
+class UE_DEPRECATED(5.5, "UDMXImportFixtureType is deprecated. Please refer to the new GDTF Description in UDMXDeviceType instead") DMXRUNTIME_API UDMXImportFixtureType
+	: public UObject
 {
-    GENERATED_BODY()
+	GENERATED_BODY()
 };
 
-UCLASS(BlueprintType, Blueprintable)
-class DMXRUNTIME_API UDMXImportAttributeDefinitions
-    : public UObject
+UCLASS(BlueprintType, Blueprintable, Meta = (DeprecatedNode, DeprecationMessage = "This class is now deprecated in favor of a DMXEntityFixtureType based workflow. Please refer to members of DMXEntityFixtureType instead."))
+class UE_DEPRECATED(5.5, "UDMXImportAttributeDefinitions is deprecated. Please refer to the new GDTF Description in UDMXDeviceType instead") DMXRUNTIME_API UDMXImportAttributeDefinitions
+	: public UObject
 {
-    GENERATED_BODY()
+	GENERATED_BODY()
 };
 
-UCLASS(BlueprintType, Blueprintable)
-class DMXRUNTIME_API UDMXImportWheels
-    : public UObject
+UCLASS(BlueprintType, Blueprintable, Meta = (DeprecatedNode, DeprecationMessage = "This class is now deprecated in favor of a DMXEntityFixtureType based workflow. Please refer to members of DMXEntityFixtureType instead."))
+class UE_DEPRECATED(5.5, "UDMXImportWheels is deprecated. Please refer to the new GDTF Description in UDMXDeviceType instead") DMXRUNTIME_API UDMXImportWheels
+	: public UObject
 {
-    GENERATED_BODY()
+	GENERATED_BODY()
 };
 
-UCLASS(BlueprintType, Blueprintable)
-class DMXRUNTIME_API UDMXImportPhysicalDescriptions
-    : public UObject
+UCLASS(BlueprintType, Blueprintable, Meta = (DeprecatedNode, DeprecationMessage = "This class is now deprecated in favor of a DMXEntityFixtureType based workflow. Please refer to members of DMXEntityFixtureType instead."))
+class UE_DEPRECATED(5.5, "UDMXImportPhysicalDescriptions is deprecated. Please refer to the new GDTF Description in UDMXDeviceType instead") DMXRUNTIME_API UDMXImportPhysicalDescriptions
+	: public UObject
 {
-    GENERATED_BODY()
+	GENERATED_BODY()
 };
 
-UCLASS(BlueprintType, Blueprintable)
-class DMXRUNTIME_API UDMXImportModels
-    : public UObject
+UCLASS(BlueprintType, Blueprintable, Meta = (DeprecatedNode, DeprecationMessage = "This class is now deprecated in favor of a DMXEntityFixtureType based workflow. Please refer to members of DMXEntityFixtureType instead."))
+class UE_DEPRECATED(5.5, "UDMXImportModels is deprecated. Please refer to the new GDTF Description in UDMXDeviceType instead") DMXRUNTIME_API UDMXImportModels
+	: public UObject
 {
-    GENERATED_BODY()
+	GENERATED_BODY()
 };
 
-UCLASS(BlueprintType, Blueprintable)
-class DMXRUNTIME_API UDMXImportGeometries
-    : public UObject
+UCLASS(BlueprintType, Blueprintable, Meta = (DeprecatedNode, DeprecationMessage = "This class is now deprecated in favor of a DMXEntityFixtureType based workflow. Please refer to members of DMXEntityFixtureType instead."))
+class UE_DEPRECATED(5.5, "UDMXImportGeometries is deprecated. Please refer to the new GDTF Description in UDMXDeviceType instead") DMXRUNTIME_API UDMXImportGeometries
+	: public UObject
 {
-    GENERATED_BODY()
+	GENERATED_BODY()
 };
 
-UCLASS(BlueprintType, Blueprintable)
+class UE_DEPRECATED(5.5, "UDMXImportDMXModes is deprecated. Please refer to the new GDTF Description in UDMXDeviceType instead") UDMXImportDMXModes;
+PRAGMA_DISABLE_DEPRECATION_WARNINGS
+UCLASS(BlueprintType, Blueprintable, Meta = (DeprecatedNode, DeprecationMessage = "This class is now deprecated in favor of a DMXEntityFixtureType based workflow. Please refer to members of DMXEntityFixtureType instead."))
 class DMXRUNTIME_API UDMXImportDMXModes
-    : public UObject
+	: public UObject
 {
-    GENERATED_BODY()
+	GENERATED_BODY()
 
 };
+PRAGMA_ENABLE_DEPRECATION_WARNINGS
 
-UCLASS(BlueprintType, Blueprintable)
-class DMXRUNTIME_API UDMXImportProtocols
-    : public UObject
+UCLASS(BlueprintType, Blueprintable, Meta = (DeprecatedNode, DeprecationMessage = "This class is now deprecated in favor of a DMXEntityFixtureType based workflow. Please refer to members of DMXEntityFixtureType instead."))
+class UE_DEPRECATED(5.5, "UDMXImportProtocols is deprecated. Please refer to the new GDTF Description in UDMXDeviceType instead") DMXRUNTIME_API UDMXImportProtocols
+	: public UObject
 {
-    GENERATED_BODY()
+	GENERATED_BODY()
 
 };
 
@@ -105,34 +112,41 @@ class DMXRUNTIME_API UDMXImport
 
 public:
 	template <typename TType>
+	UE_DEPRECATED(5.5, "Instead please use NewObject.")
 	TType* CreateNewObject()
 	{
-        FName NewName = MakeUniqueObjectName(this, TType::StaticClass());
-        return NewObject<TType>(this, TType::StaticClass(), NewName, RF_Public);
+		FName NewName = MakeUniqueObjectName(this, TType::StaticClass());
+		return NewObject<TType>(this, TType::StaticClass(), NewName, RF_Public);
 	}
 
 public:
-	UPROPERTY(VisibleAnywhere, Category = "Fixture Type", meta = (EditInline))
-	TObjectPtr<UDMXImportFixtureType> FixtureType;
+	PRAGMA_DISABLE_DEPRECATION_WARNINGS
+#if WITH_EDITORONLY_DATA
+	UPROPERTY()
+	TObjectPtr<UDMXImportFixtureType> FixtureType_DEPRECATED;
 
-    UPROPERTY(VisibleAnywhere, Category = "Attribute Definitions", meta = (EditInline))
-    TObjectPtr<UDMXImportAttributeDefinitions> AttributeDefinitions;
+	UPROPERTY()
+	TObjectPtr<UDMXImportAttributeDefinitions> AttributeDefinitions_DEPRECATED;
 
-    UPROPERTY(VisibleAnywhere, Category = "Wheels", meta = (EditInline))
-    TObjectPtr<UDMXImportWheels> Wheels;
+	UPROPERTY()
+	TObjectPtr<UDMXImportWheels> Wheels_DEPRECATED;
 
-    UPROPERTY(VisibleAnywhere, Category = "Physical Descriptions", meta = (EditInline))
-    TObjectPtr<UDMXImportPhysicalDescriptions> PhysicalDescriptions;
+	UPROPERTY()
+	TObjectPtr<UDMXImportPhysicalDescriptions> PhysicalDescriptions_DEPRECATED;
 
-    UPROPERTY(VisibleAnywhere, Category = "Models", meta = (EditInline))
-    TObjectPtr<UDMXImportModels> Models;
+	UPROPERTY()
+	TObjectPtr<UDMXImportModels> Models_DEPRECATED;
 
-    UPROPERTY(VisibleAnywhere, Category = "Geometries", meta = (EditInline))
-    TObjectPtr<UDMXImportGeometries> Geometries;
+	UPROPERTY()
+	TObjectPtr<UDMXImportGeometries> Geometries_DEPRECATED;
+#endif // WITH_EDITORONLY_DATA
 
-    UPROPERTY(VisibleAnywhere, Category = "DMXModes", meta = (EditInline))
-    TObjectPtr<UDMXImportDMXModes> DMXModes;
+	UPROPERTY()
+	TObjectPtr<UDMXImportDMXModes> DMXModes_DEPRECATED;
 
-    UPROPERTY(VisibleAnywhere, Category = "Protocols", meta = (EditInline))
-    TObjectPtr<UDMXImportProtocols> Protocols;
+#if WITH_EDITORONLY_DATA
+	UPROPERTY()
+	TObjectPtr<UDMXImportProtocols> Protocols_DEPRECATED;
+#endif // WITH_EDITORONLY_DATA
+	PRAGMA_ENABLE_DEPRECATION_WARNINGS
 };
diff --git a/Engine/Plugins/VirtualProduction/DMX/DMXEngine/Source/DMXRuntime/Public/Library/DMXImportGDTF.h b/Engine/Plugins/VirtualProduction/DMX/DMXEngine/Source/DMXRuntime/Public/Library/DMXImportGDTF.h
--- a/Engine/Plugins/VirtualProduction/DMX/DMXEngine/Source/DMXRuntime/Public/Library/DMXImportGDTF.h
+++ b/Engine/Plugins/VirtualProduction/DMX/DMXEngine/Source/DMXRuntime/Public/Library/DMXImportGDTF.h
@@ -6,21 +6,70 @@
 
 #include "DMXImportGDTF.generated.h"
 
+class UDMXGDTF;
 class UDMXGDTFAssetImportData;
-
 class UTexture2D;
+namespace UE::DMX::GDTF { class FDMXGDTFDescription; }
 
 
-UENUM(BlueprintType)
-enum class EDMXImportGDTFType : uint8
+/** 
+ * A GDTF imported into Unreal Engine. Note this object holds GDTF source data and may be memory heavy.
+ * 
+ * To access the GDTF description in lightweight fashion, please refer to UDMXGDTF (C++ only). 
+ */
+UCLASS(Blueprintable, BlueprintType)
+class DMXRUNTIME_API UDMXImportGDTF
+	: public UDMXImport
+{
+	GENERATED_BODY()
+
+public:
+	/** Constructor */
+	UDMXImportGDTF();
+
+	//~ Begin UObject interface
+	virtual void PostLoad() override;
+	//~ End UObject interface
+
+	/** Loads the GDTF stored in this object */
+	UDMXGDTF* LoadGDTF() const;
+
+	/** DEPRECATED 5.4 */
+	UE_DEPRECATED(5.5, "UDMXImportGDTF::GetDMXModes is deprecated in favor of UDMXGDTF. See also newly added UDMXImportGDTF::LoadGDTF to read out GDTF data of this asset.")
+	PRAGMA_DISABLE_DEPRECATION_WARNINGS // Suppress deprecation warnings for the deprecated UDMXImportGDTFDMXModes return type
+	UFUNCTION(BlueprintPure, Category = "DMXGDTF|Import Data", Meta = (DeprecatedFunction, DeprecationMessage = "This member is deprecated in favor of a DMX Entity Fixture Type based workflow. Please refer to members of DMX Entity Fixture Type instead"))
+	class UDMXImportGDTFDMXModes* GetDMXModes() const { return Cast<UDMXImportGDTFDMXModes>(DMXModes_DEPRECATED); }
+	PRAGMA_ENABLE_DEPRECATION_WARNINGS
+
+#if WITH_EDITORONLY_DATA
+	/** DEPRECATED 5.1 in favor of GDTFAssetImportData */
+	UPROPERTY(Meta = (DeprecatedProperty, DeprecationMessage = "Deprecated in favor of GDTFAssetImportData, see UDMXImportGDTF::GetGDTFAssetImportData."))
+	FString SourceFilename_DEPRECATED;
+#endif
+
+	/** Returns GDTF Asset Import Data for this GDTF */
+	FORCEINLINE UDMXGDTFAssetImportData* GetGDTFAssetImportData() const { return GDTFAssetImportData; }
+ 
+private:
+	/** The Asset Import Data used to generate the GDTF asset or nullptr, if not generated from a GDTF file */
+	UPROPERTY()
+	TObjectPtr<UDMXGDTFAssetImportData> GDTFAssetImportData;
+};
+
+
+///////////////////////////////
+//~ Begin Deprecated Members
+// 
+
+UENUM()
+enum class UE_DEPRECATED(5.5, "FDMXImportGDTFType is deprecated in favor of its corresponding type in the DMXGDTF module. See also UDMXImportGDTF::LoadGDTF.") EDMXImportGDTFType : uint8
 {
     Multiply,
     Override
 };
 
-
-UENUM(BlueprintType)
-enum class EDMXImportGDTFSnap : uint8
+UENUM()
+enum class UE_DEPRECATED(5.5, "FDMXImportGDTFSnap is deprecated in favor of its corresponding type in the DMXGDTF module. See also UDMXImportGDTF::LoadGDTF.") EDMXImportGDTFSnap : uint8
 {
     Yes,
     No,
@@ -28,23 +77,23 @@ enum class EDMXImportGDTFSnap : uint8
     Off
 };
 
-UENUM(BlueprintType)
-enum class EDMXImportGDTFMaster : uint8
+UENUM()
+enum class UE_DEPRECATED(5.5, "FDMXImportGDTFMaster is deprecated in favor of its corresponding type in the DMXGDTF module. See also UDMXImportGDTF::LoadGDTF.") EDMXImportGDTFMaster : uint8
 {
     None,
     Grand,
     Group
 };
 
-UENUM(BlueprintType)
-enum class EDMXImportGDTFDMXInvert : uint8
+UENUM()
+enum class UE_DEPRECATED(5.5, "FDMXImportGDTFDMXInvert is deprecated in favor of its corresponding type in the DMXGDTF module. See also UDMXImportGDTF::LoadGDTF.") EDMXImportGDTFDMXInvert : uint8
 {
     Yes,
     No
 };
 
-UENUM(BlueprintType)
-enum class EDMXImportGDTFLampType : uint8
+UENUM()
+enum class UE_DEPRECATED(5.5, "FDMXImportGDTFLampType is deprecated in favor of its corresponding type in the DMXGDTF module. See also UDMXImportGDTF::LoadGDTF.") EDMXImportGDTFLampType : uint8
 {
     Discharge,
     Tungsten,
@@ -52,16 +101,16 @@ enum class EDMXImportGDTFLampType : uint8
     LED
 };
 
-UENUM(BlueprintType)
-enum class EDMXImportGDTFBeamType : uint8
+UENUM()
+enum class UE_DEPRECATED(5.5, "FDMXImportGDTFBeamType is deprecated in favor of its corresponding type in the DMXGDTF module. See also UDMXImportGDTF::LoadGDTF.") EDMXImportGDTFBeamType : uint8
 {
     Wash,
     Spot,
     None
 };
 
-UENUM(BlueprintType)
-enum class EDMXImportGDTFPrimitiveType : uint8
+UENUM()
+enum class UE_DEPRECATED(5.5, "FDMXImportGDTFPrimitiveType is deprecated in favor of its corresponding type in the DMXGDTF module. See also UDMXImportGDTF::LoadGDTF.") EDMXImportGDTFPrimitiveType : uint8
 {
     Undefined,
     Cube,
@@ -75,8 +124,8 @@ enum class EDMXImportGDTFPrimitiveType : uint8
     Pigtail
 };
 
-UENUM(BlueprintType)
-enum class EDMXImportGDTFPhysicalUnit : uint8
+UENUM()
+enum class UE_DEPRECATED(5.5, "FDMXImportGDTFPhysicalUnit is deprecated in favor of its corresponding type in the DMXGDTF module. See also UDMXImportGDTF::LoadGDTF.") EDMXImportGDTFPhysicalUnit : uint8
 {
     None,
     Percent,
@@ -102,8 +151,8 @@ enum class EDMXImportGDTFPhysicalUnit : uint8
     ColorComponent
 };
 
-UENUM(BlueprintType)
-enum class EDMXImportGDTFMode : uint8
+UENUM()
+enum class UE_DEPRECATED(5.5, "FDMXImportGDTFMode is deprecated in favor of its corresponding type in the DMXGDTF module. See also UDMXImportGDTF::LoadGDTF.") EDMXImportGDTFMode : uint8
 {
     Custom,
     sRGB,
@@ -111,8 +160,8 @@ enum class EDMXImportGDTFMode : uint8
     ANSI
 };
 
-UENUM(BlueprintType)
-enum class EDMXImportGDTFInterpolationTo : uint8
+UENUM()
+enum class UE_DEPRECATED(5.5, "FDMXImportGDTFInterpolationTo is deprecated in favor of its corresponding type in the DMXGDTF module. See also UDMXImportGDTF::LoadGDTF.") EDMXImportGDTFInterpolationTo : uint8
 {
     Linear,
     Step,
@@ -120,7 +169,7 @@ enum class EDMXImportGDTFInterpolationTo : uint8
 };
 
 USTRUCT(BlueprintType)
-struct DMXRUNTIME_API FDMXImportGDTFActivationGroup
+struct UE_DEPRECATED(5.5, "FDMXImportGDTFActivationGroup is deprecated in favor of its corresponding type in the DMXGDTF module. See also UDMXImportGDTF::LoadGDTF.") DMXRUNTIME_API FDMXImportGDTFActivationGroup
 {
     GENERATED_BODY()
 
@@ -129,7 +178,7 @@ struct DMXRUNTIME_API FDMXImportGDTFActivationGroup
 };
 
 USTRUCT(BlueprintType)
-struct DMXRUNTIME_API FDMXImportGDTFFeature
+struct UE_DEPRECATED(5.5, "FDMXImportGDTFFeature is deprecated in favor of its corresponding type in the DMXGDTF module. See also UDMXImportGDTF::LoadGDTF.") DMXRUNTIME_API FDMXImportGDTFFeature
 {
     GENERATED_BODY()
 
@@ -138,7 +187,7 @@ struct DMXRUNTIME_API FDMXImportGDTFFeature
 };
 
 USTRUCT(BlueprintType)
-struct DMXRUNTIME_API FDMXImportGDTFFeatureGroup
+struct UE_DEPRECATED(5.5, "FDMXImportGDTFFeatureGroup is deprecated in favor of its corresponding type in the DMXGDTF module. See also UDMXImportGDTF::LoadGDTF.") DMXRUNTIME_API FDMXImportGDTFFeatureGroup
 {
     GENERATED_BODY()
 
@@ -148,12 +197,14 @@ struct DMXRUNTIME_API FDMXImportGDTFFeatureGroup
     UPROPERTY(EditAnywhere, BlueprintReadOnly, Category = "DMX")
     FString Pretty;
 
+	PRAGMA_DISABLE_DEPRECATION_WARNINGS
     UPROPERTY(EditAnywhere, BlueprintReadOnly, Category = "DMX")
     TArray<FDMXImportGDTFFeature> Features;
+	PRAGMA_ENABLE_DEPRECATION_WARNINGS
 };
 
 USTRUCT(BlueprintType)
-struct DMXRUNTIME_API FDMXImportGDTFAttribute
+struct UE_DEPRECATED(5.5, "FDMXImportGDTFAttribute is deprecated in favor of its corresponding type in the DMXGDTF module. See also UDMXImportGDTF::LoadGDTF.") DMXRUNTIME_API FDMXImportGDTFAttribute
 {
     GENERATED_BODY()
 
@@ -163,6 +214,7 @@ struct DMXRUNTIME_API FDMXImportGDTFAttribute
     UPROPERTY(EditAnywhere, BlueprintReadOnly, Category = "DMX")
     FString Pretty;
 
+	PRAGMA_DISABLE_DEPRECATION_WARNINGS
     UPROPERTY(EditAnywhere, BlueprintReadOnly, Category = "DMX")
     FDMXImportGDTFActivationGroup ActivationGroup;
 
@@ -177,53 +229,60 @@ struct DMXRUNTIME_API FDMXImportGDTFAttribute
 
     UPROPERTY(EditAnywhere, BlueprintReadOnly, Category = "DMX")
     FDMXColorCIE Color;
+	PRAGMA_ENABLE_DEPRECATION_WARNINGS
 };
 
 USTRUCT(BlueprintType)
-struct FDMXImportGDTFFilter
+struct UE_DEPRECATED(5.5, "FDMXImportGDTFFilter is deprecated in favor of its corresponding type in the DMXGDTF module. See also UDMXImportGDTF::LoadGDTF.") FDMXImportGDTFFilter
 {
     GENERATED_BODY()
 
     UPROPERTY(EditAnywhere, BlueprintReadOnly, Category = "DMX")
     FName Name;
 
+	PRAGMA_DISABLE_DEPRECATION_WARNINGS
     UPROPERTY(EditAnywhere, BlueprintReadOnly, Category = "DMX")
     FDMXColorCIE Color;
+	PRAGMA_ENABLE_DEPRECATION_WARNINGS
 };
 
 
 USTRUCT(BlueprintType)
-struct FDMXImportGDTFWheelSlot
+struct UE_DEPRECATED(5.5, "FDMXImportGDTFWheelSlot is deprecated in favor of its corresponding type in the DMXGDTF module. See also UDMXImportGDTF::LoadGDTF.") FDMXImportGDTFWheelSlot
 {
     GENERATED_BODY()
 
     UPROPERTY(EditAnywhere, BlueprintReadOnly, Category = "DMX")
     FName Name;
 
+	PRAGMA_DISABLE_DEPRECATION_WARNINGS
     UPROPERTY(EditAnywhere, BlueprintReadOnly, Category = "DMX")
     FDMXColorCIE Color;
 
-    UPROPERTY(EditAnywhere, BlueprintReadOnly, Category = "DMX")
+	UPROPERTY(EditAnywhere, BlueprintReadOnly, Category = "DMX")
     FDMXImportGDTFFilter Filter;
+	PRAGMA_ENABLE_DEPRECATION_WARNINGS
 
     UPROPERTY(VisibleAnywhere, Category = "Fixture Type")
     TObjectPtr<UTexture2D> MediaFileName = nullptr;
 };
 
 USTRUCT(BlueprintType)
-struct FDMXImportGDTFWheel
+struct UE_DEPRECATED(5.5, "FDMXImportGDTFWheel is deprecated in favor of its corresponding type in the DMXGDTF module. See also UDMXImportGDTF::LoadGDTF.") FDMXImportGDTFWheel
 {
     GENERATED_BODY()
 
     UPROPERTY(EditAnywhere, BlueprintReadOnly, Category = "DMX")
     FName Name;
 
+	PRAGMA_DISABLE_DEPRECATION_WARNINGS
     UPROPERTY(EditAnywhere, BlueprintReadOnly, Category = "DMX")
     TArray<FDMXImportGDTFWheelSlot> Slots;
+	PRAGMA_ENABLE_DEPRECATION_WARNINGS
 };
 
 USTRUCT(BlueprintType)
-struct FDMXImportGDTFMeasurementPoint
+struct UE_DEPRECATED(5.5, "FDMXImportGDTFMeasurementPoint is deprecated in favor of its corresponding type in the DMXGDTF module. See also UDMXImportGDTF::LoadGDTF.") FDMXImportGDTFMeasurementPoint
 {
     GENERATED_BODY()
 
@@ -236,7 +295,7 @@ struct FDMXImportGDTFMeasurementPoint
 
 
 USTRUCT(BlueprintType)
-struct FDMXImportGDTFMeasurement
+struct UE_DEPRECATED(5.5, "FDMXImportGDTFMeasurement is deprecated in favor of its corresponding type in the DMXGDTF module. See also UDMXImportGDTF::LoadGDTF.") FDMXImportGDTFMeasurement
 {
     GENERATED_BODY()
 
@@ -249,23 +308,27 @@ struct FDMXImportGDTFMeasurement
     UPROPERTY(EditAnywhere, BlueprintReadOnly, Category = "DMX")
     float Transmission = 0.f;
 
+	PRAGMA_DISABLE_DEPRECATION_WARNINGS
     UPROPERTY(EditAnywhere, BlueprintReadOnly, Category = "DMX")
     EDMXImportGDTFInterpolationTo InterpolationTo = EDMXImportGDTFInterpolationTo::Linear;
 
     UPROPERTY(EditAnywhere, BlueprintReadOnly, Category = "DMX")
     TArray<FDMXImportGDTFMeasurementPoint> MeasurementPoints;
+	PRAGMA_ENABLE_DEPRECATION_WARNINGS
 };
 
 USTRUCT(BlueprintType)
-struct FDMXImportGDTFEmitter
+struct UE_DEPRECATED(5.5, "FDMXImportGDTFEmitter is deprecated in favor of its corresponding type in the DMXGDTF module. See also UDMXImportGDTF::LoadGDTF.") FDMXImportGDTFEmitter
 {
     GENERATED_BODY()
 
     UPROPERTY(EditAnywhere, BlueprintReadOnly, Category = "DMX")
     FName Name;
 
+	PRAGMA_DISABLE_DEPRECATION_WARNINGS
     UPROPERTY(EditAnywhere, BlueprintReadOnly, Category = "DMX")
     FDMXColorCIE Color;
+	PRAGMA_ENABLE_DEPRECATION_WARNINGS
 
     UPROPERTY(EditAnywhere, BlueprintReadOnly, Category = "DMX")
     float DominantWaveLength = 0.f;
@@ -273,21 +336,26 @@ struct FDMXImportGDTFEmitter
     UPROPERTY(EditAnywhere, BlueprintReadOnly, Category = "DMX")
     FString DiodePart;
 
+	PRAGMA_DISABLE_DEPRECATION_WARNINGS
     UPROPERTY(EditAnywhere, BlueprintReadOnly, Category = "DMX")
     FDMXImportGDTFMeasurement Measurement;
+	PRAGMA_ENABLE_DEPRECATION_WARNINGS
 };
 
 USTRUCT(BlueprintType)
-struct FDMXImportGDTFColorSpace
+struct UE_DEPRECATED(5.5, "FDMXImportGDTFColorSpace is deprecated in favor of its corresponding type in the DMXGDTF module. See also UDMXImportGDTF::LoadGDTF.") FDMXImportGDTFColorSpace
 {
     GENERATED_BODY()
 
+	PRAGMA_DISABLE_DEPRECATION_WARNINGS
     UPROPERTY(EditAnywhere, BlueprintReadOnly, Category = "DMX")
     EDMXImportGDTFMode Mode = EDMXImportGDTFMode::sRGB;
+	PRAGMA_ENABLE_DEPRECATION_WARNINGS
 
     UPROPERTY(EditAnywhere, BlueprintReadOnly, Category = "DMX")
     FString Description;
 
+	PRAGMA_DISABLE_DEPRECATION_WARNINGS
     UPROPERTY(EditAnywhere, BlueprintReadOnly, Category = "DMX")
     FDMXColorCIE Red;
 
@@ -299,23 +367,24 @@ struct FDMXImportGDTFColorSpace
 
     UPROPERTY(EditAnywhere, BlueprintReadOnly, Category = "DMX")
     FDMXColorCIE WhitePoint;
+	PRAGMA_ENABLE_DEPRECATION_WARNINGS
 };
 
 USTRUCT(BlueprintType)
-struct FDMXImportGDTFDMXProfiles
+struct UE_DEPRECATED(5.5, "FDMXImportGDTFDMXProfiles is deprecated in favor of its corresponding type in the DMXGDTF module. See also UDMXImportGDTF::LoadGDTF.") FDMXImportGDTFDMXProfiles
 {
     GENERATED_BODY()
 };
 
 USTRUCT(BlueprintType)
-struct FDMXImportGDTFCRIs
+struct UE_DEPRECATED(5.5, "FDMXImportGDTFCRIs is deprecated in favor of its corresponding type in the DMXGDTF module. See also UDMXImportGDTF::LoadGDTF.") FDMXImportGDTFCRIs
 {
     GENERATED_BODY()
 };
 
 
 USTRUCT(BlueprintType)
-struct FDMXImportGDTFModel
+struct UE_DEPRECATED(5.5, "FDMXImportGDTFModel is deprecated in favor of its corresponding type in the DMXGDTF module. See also UDMXImportGDTF::LoadGDTF.") FDMXImportGDTFModel
 {
     GENERATED_BODY()
 
@@ -331,17 +400,18 @@ struct FDMXImportGDTFModel
     UPROPERTY(EditAnywhere, BlueprintReadOnly, Category = "DMX")
     float Height = 0.f;
 
+	PRAGMA_DISABLE_DEPRECATION_WARNINGS
     UPROPERTY(EditAnywhere, BlueprintReadOnly, Category = "DMX")
     EDMXImportGDTFPrimitiveType PrimitiveType = EDMXImportGDTFPrimitiveType::Undefined;
+	PRAGMA_ENABLE_DEPRECATION_WARNINGS
 };
 
 USTRUCT(BlueprintType)
-struct FDMXImportGDTFGeometryBase
+struct UE_DEPRECATED(5.5, "FDMXImportGDTFGeometryBase is deprecated in favor of its corresponding type in the DMXGDTF module. See also UDMXImportGDTF::LoadGDTF.") FDMXImportGDTFGeometryBase
 {
 	GENERATED_BODY()
 
 	FDMXImportGDTFGeometryBase()
-		: Position(FMatrix::Identity)
 	{}
 
     UPROPERTY(EditAnywhere, BlueprintReadOnly, Category = "DMX")
@@ -351,17 +421,27 @@ struct FDMXImportGDTFGeometryBase
     FName Model;
 
     UPROPERTY(EditAnywhere, BlueprintReadOnly, Category = "DMX")
-    FMatrix Position;
+    FMatrix Position = FMatrix::Identity;
 };
 
 USTRUCT(BlueprintType)
-struct FDMXImportGDTFBeam
-    : public FDMXImportGDTFGeometryBase
+struct UE_DEPRECATED(5.5, "FDMXImportGDTFBeam is deprecated in favor of its corresponding type in the DMXGDTF module. See also UDMXImportGDTF::LoadGDTF.") FDMXImportGDTFBeam
 {
     GENERATED_BODY()
 
+	UPROPERTY(EditAnywhere, BlueprintReadOnly, Category = "DMX")
+	FName Name;
+
+	UPROPERTY(EditAnywhere, BlueprintReadOnly, Category = "DMX")
+	FName Model;
+
+	UPROPERTY(EditAnywhere, BlueprintReadOnly, Category = "DMX")
+	FMatrix Position = FMatrix::Identity;
+
+	PRAGMA_DISABLE_DEPRECATION_WARNINGS
     UPROPERTY(EditAnywhere, BlueprintReadOnly, Category = "DMX")
     EDMXImportGDTFLampType LampType = EDMXImportGDTFLampType::Discharge;
+	PRAGMA_ENABLE_DEPRECATION_WARNINGS
 
     UPROPERTY(EditAnywhere, BlueprintReadOnly, Category = "DMX")
     float PowerConsumption = 0.f;
@@ -381,70 +461,132 @@ struct FDMXImportGDTFBeam
     UPROPERTY(EditAnywhere, BlueprintReadOnly, Category = "DMX")
     float BeamRadius = 0.f;
 
+	PRAGMA_DISABLE_DEPRECATION_WARNINGS
     UPROPERTY(EditAnywhere, BlueprintReadOnly, Category = "DMX")
     EDMXImportGDTFBeamType BeamType = EDMXImportGDTFBeamType::Wash;
 
     UPROPERTY(EditAnywhere, BlueprintReadOnly, Category = "DMX")
     uint8 ColorRenderingIndex = 0;
+	PRAGMA_ENABLE_DEPRECATION_WARNINGS
 };
 
 USTRUCT(BlueprintType)
-struct FDMXImportGDTFTypeAxis
-    : public FDMXImportGDTFGeometryBase
+struct UE_DEPRECATED(5.5, "FDMXImportGDTFTypeAxis is deprecated in favor of its corresponding type in the DMXGDTF module. See also UDMXImportGDTF::LoadGDTF.") FDMXImportGDTFTypeAxis
 {
     GENERATED_BODY()
 
+	UPROPERTY(EditAnywhere, BlueprintReadOnly, Category = "DMX")
+	FName Name;
+
+	UPROPERTY(EditAnywhere, BlueprintReadOnly, Category = "DMX")
+	FName Model;
+
+	UPROPERTY(EditAnywhere, BlueprintReadOnly, Category = "DMX")
+	FMatrix Position = FMatrix::Identity;
+
+	PRAGMA_DISABLE_DEPRECATION_WARNINGS
     UPROPERTY(EditAnywhere, BlueprintReadOnly, Category = "DMX")
     TArray<FDMXImportGDTFBeam> Beams;
+	PRAGMA_ENABLE_DEPRECATION_WARNINGS
 };
 
 USTRUCT(BlueprintType)
-struct FDMXImportGDTFGeneralAxis
-    : public FDMXImportGDTFGeometryBase
+struct UE_DEPRECATED(5.5, "FDMXImportGDTFGeneralAxis is deprecated in favor of its corresponding type in the DMXGDTF module. See also UDMXImportGDTF::LoadGDTF.") FDMXImportGDTFGeneralAxis
 {
     GENERATED_BODY()
 
+	UPROPERTY(EditAnywhere, BlueprintReadOnly, Category = "DMX")
+	FName Name;
+
+	UPROPERTY(EditAnywhere, BlueprintReadOnly, Category = "DMX")
+	FName Model;
+
+	UPROPERTY(EditAnywhere, BlueprintReadOnly, Category = "DMX")
+	FMatrix Position = FMatrix::Identity;
+
+	PRAGMA_DISABLE_DEPRECATION_WARNINGS
     UPROPERTY(EditAnywhere, BlueprintReadOnly, Category = "DMX")
     TArray<FDMXImportGDTFTypeAxis> Axis;
+	PRAGMA_ENABLE_DEPRECATION_WARNINGS
 };
 
 USTRUCT(BlueprintType)
-struct FDMXImportGDTFTypeGeometry
-    : public FDMXImportGDTFGeometryBase
+struct UE_DEPRECATED(5.5, "FDMXImportGDTFTypeGeometry is deprecated in favor of its corresponding type in the DMXGDTF module. See also UDMXImportGDTF::LoadGDTF.") FDMXImportGDTFTypeGeometry
 {
     GENERATED_BODY()
+
+	UPROPERTY(EditAnywhere, BlueprintReadOnly, Category = "DMX")
+	FName Name;
+
+	UPROPERTY(EditAnywhere, BlueprintReadOnly, Category = "DMX")
+	FName Model;
+
+	UPROPERTY(EditAnywhere, BlueprintReadOnly, Category = "DMX")
+	FMatrix Position = FMatrix::Identity;
 };
 
 USTRUCT(BlueprintType)
-struct FDMXImportGDTFFilterBeam
-    : public FDMXImportGDTFGeometryBase
+struct UE_DEPRECATED(5.5, "FDMXImportGDTFFilterBeam is deprecated in favor of its corresponding type in the DMXGDTF module. See also UDMXImportGDTF::LoadGDTF.") FDMXImportGDTFFilterBeam
 {
     GENERATED_BODY()
+
+	UPROPERTY(EditAnywhere, BlueprintReadOnly, Category = "DMX")
+	FName Name;
+
+	UPROPERTY(EditAnywhere, BlueprintReadOnly, Category = "DMX")
+	FName Model;
+
+	UPROPERTY(EditAnywhere, BlueprintReadOnly, Category = "DMX")
+	FMatrix Position = FMatrix::Identity;
 };
 
 USTRUCT(BlueprintType)
-struct FDMXImportGDTFFilterColor
-    : public FDMXImportGDTFGeometryBase
+struct UE_DEPRECATED(5.5, "FDMXImportGDTFFilterColor is deprecated in favor of its corresponding type in the DMXGDTF module. See also UDMXImportGDTF::LoadGDTF.") FDMXImportGDTFFilterColor
 {
     GENERATED_BODY()
+
+	UPROPERTY(EditAnywhere, BlueprintReadOnly, Category = "DMX")
+	FName Name;
+
+	UPROPERTY(EditAnywhere, BlueprintReadOnly, Category = "DMX")
+	FName Model;
+
+	UPROPERTY(EditAnywhere, BlueprintReadOnly, Category = "DMX")
+	FMatrix Position = FMatrix::Identity;
 };
 
 USTRUCT(BlueprintType)
-struct FDMXImportGDTFFilterGobo
-    : public FDMXImportGDTFGeometryBase
+struct UE_DEPRECATED(5.5, "FDMXImportGDTFFilterGobo is deprecated in favor of its corresponding type in the DMXGDTF module. See also UDMXImportGDTF::LoadGDTF.") FDMXImportGDTFFilterGobo
 {
     GENERATED_BODY()
+
+	UPROPERTY(EditAnywhere, BlueprintReadOnly, Category = "DMX")
+	FName Name;
+
+	UPROPERTY(EditAnywhere, BlueprintReadOnly, Category = "DMX")
+	FName Model;
+
+	UPROPERTY(EditAnywhere, BlueprintReadOnly, Category = "DMX")
+	FMatrix Position = FMatrix::Identity;
 };
 
 USTRUCT(BlueprintType)
-struct FDMXImportGDTFFilterShaper
-    : public FDMXImportGDTFGeometryBase
+struct UE_DEPRECATED(5.5, "FDMXImportGDTFFilterShaper is deprecated in favor of its corresponding type in the DMXGDTF module. See also UDMXImportGDTF::LoadGDTF.") FDMXImportGDTFFilterShaper
 {
     GENERATED_BODY()
+
+	UPROPERTY(EditAnywhere, BlueprintReadOnly, Category = "DMX")
+	FName Name;
+
+	UPROPERTY(EditAnywhere, BlueprintReadOnly, Category = "DMX")
+	FName Model;
+
+	UPROPERTY(EditAnywhere, BlueprintReadOnly, Category = "DMX")
+	FMatrix Position = FMatrix::Identity;
 };
 
 USTRUCT(BlueprintType)
-struct FDMXImportGDTFBreak
+struct UE_DEPRECATED(5.5, "FDMXImportGDTFBreak is deprecated in favor of its corresponding type in the DMXGDTF module. See also UDMXImportGDTF::LoadGDTF.") FDMXImportGDTFBreak
 {
     GENERATED_BODY()
 
@@ -456,21 +598,40 @@ struct FDMXImportGDTFBreak
 };
 
 USTRUCT(BlueprintType)
-struct FDMXImportGDTFGeometryReference
-    : public FDMXImportGDTFGeometryBase
+struct UE_DEPRECATED(5.5, "FDMXImportGDTFGeometryReference is deprecated in favor of its corresponding type in the DMXGDTF module. See also UDMXImportGDTF::LoadGDTF.") FDMXImportGDTFGeometryReference
 {
     GENERATED_BODY()
 
+	UPROPERTY(EditAnywhere, BlueprintReadOnly, Category = "DMX")
+	FName Name;
+
+	UPROPERTY(EditAnywhere, BlueprintReadOnly, Category = "DMX")
+	FName Model;
+
+	UPROPERTY(EditAnywhere, BlueprintReadOnly, Category = "DMX")
+	FMatrix Position = FMatrix::Identity;
+
+	PRAGMA_DISABLE_DEPRECATION_WARNINGS
     UPROPERTY(EditAnywhere, BlueprintReadOnly, Category = "DMX")
     TArray<FDMXImportGDTFBreak> Breaks;
+	PRAGMA_ENABLE_DEPRECATION_WARNINGS
 };
 
 USTRUCT(BlueprintType)
-struct FDMXImportGDTFGeneralGeometry
-    : public FDMXImportGDTFGeometryBase
+struct UE_DEPRECATED(5.5, "FDMXImportGDTFGeneralGeometry is deprecated in favor of its corresponding type in the DMXGDTF module. See also UDMXImportGDTF::LoadGDTF.") FDMXImportGDTFGeneralGeometry
 {
     GENERATED_BODY()
 
+	UPROPERTY(EditAnywhere, BlueprintReadOnly, Category = "DMX")
+	FName Name;
+
+	UPROPERTY(EditAnywhere, BlueprintReadOnly, Category = "DMX")
+	FName Model;
+
+	UPROPERTY(EditAnywhere, BlueprintReadOnly, Category = "DMX")
+	FMatrix Position = FMatrix::Identity;
+
+	PRAGMA_DISABLE_DEPRECATION_WARNINGS
     UPROPERTY(EditAnywhere, BlueprintReadOnly, Category = "DMX")
     FDMXImportGDTFGeneralAxis Axis;
 
@@ -491,10 +652,11 @@ struct FDMXImportGDTFGeneralGeometry
 
     UPROPERTY(EditAnywhere, BlueprintReadOnly, Category = "DMX")
     FDMXImportGDTFGeometryReference GeometryReference;
+	PRAGMA_ENABLE_DEPRECATION_WARNINGS
 };
 
 USTRUCT(BlueprintType)
-struct DMXRUNTIME_API FDMXImportGDTFDMXValue
+struct UE_DEPRECATED(5.5, "FDMXImportGDTFDMXValue is deprecated in favor of its corresponding type in the DMXGDTF module. See also UDMXImportGDTF::LoadGDTF.") DMXRUNTIME_API FDMXImportGDTFDMXValue
 {
     GENERATED_BODY()
 
@@ -514,15 +676,17 @@ struct DMXRUNTIME_API FDMXImportGDTFDMXValue
 };
 
 USTRUCT(BlueprintType)
-struct FDMXImportGDTFChannelSet
+struct UE_DEPRECATED(5.5, "FDMXImportGDTFChannelSet is deprecated in favor of its corresponding type in the DMXGDTF module. See also UDMXImportGDTF::LoadGDTF.") FDMXImportGDTFChannelSet
 {
     GENERATED_BODY()
 
     UPROPERTY(EditAnywhere, BlueprintReadOnly, Category = "DMX")
     FString Name;
 
+	PRAGMA_DISABLE_DEPRECATION_WARNINGS
     UPROPERTY(EditAnywhere, BlueprintReadOnly, Category = "DMX")
     FDMXImportGDTFDMXValue DMXFrom;
+	PRAGMA_ENABLE_DEPRECATION_WARNINGS
 
     UPROPERTY(EditAnywhere, BlueprintReadOnly, Category = "DMX")
     float PhysicalFrom = 0.f;
@@ -535,24 +699,28 @@ struct FDMXImportGDTFChannelSet
 };
 
 USTRUCT(BlueprintType)
-struct FDMXImportGDTFChannelFunction
+struct UE_DEPRECATED(5.5, "FDMXImportGDTFChannelFunction is deprecated in favor of its corresponding type in the DMXGDTF module. See also UDMXImportGDTF::LoadGDTF.") FDMXImportGDTFChannelFunction
 {
     GENERATED_BODY()
 
     UPROPERTY(EditAnywhere, BlueprintReadOnly, Category = "DMX")
     FName Name;
 
+	PRAGMA_DISABLE_DEPRECATION_WARNINGS
     UPROPERTY(EditAnywhere, BlueprintReadOnly, Category = "DMX")
     FDMXImportGDTFAttribute Attribute;
+	PRAGMA_ENABLE_DEPRECATION_WARNINGS
 
     UPROPERTY(EditAnywhere, BlueprintReadOnly, Category = "DMX")
     FString OriginalAttribute;
 
+	PRAGMA_DISABLE_DEPRECATION_WARNINGS
     UPROPERTY(EditAnywhere, BlueprintReadOnly, Category = "DMX")
     FDMXImportGDTFDMXValue DMXFrom;
 
     UPROPERTY(EditAnywhere, BlueprintReadOnly, Category = "DMX")
     FDMXImportGDTFDMXValue DMXValue;
+	PRAGMA_ENABLE_DEPRECATION_WARNINGS
 
     UPROPERTY(EditAnywhere, BlueprintReadOnly, Category = "DMX")
     float PhysicalFrom = 0.f;
@@ -563,6 +731,7 @@ struct FDMXImportGDTFChannelFunction
     UPROPERTY(EditAnywhere, BlueprintReadOnly, Category = "DMX")
     float RealFade = 0.f;
 
+	PRAGMA_DISABLE_DEPRECATION_WARNINGS
     UPROPERTY(EditAnywhere, BlueprintReadOnly, Category = "DMX")
     FDMXImportGDTFWheel Wheel;
 
@@ -574,10 +743,12 @@ struct FDMXImportGDTFChannelFunction
 
     UPROPERTY(EditAnywhere, BlueprintReadOnly, Category = "DMX")
     EDMXImportGDTFDMXInvert DMXInvert = EDMXImportGDTFDMXInvert::No;
-
+	PRAGMA_ENABLE_DEPRECATION_WARNINGS
+		
     UPROPERTY(EditAnywhere, BlueprintReadOnly, Category = "DMX")
     FString ModeMaster;
 
+	PRAGMA_DISABLE_DEPRECATION_WARNINGS
     UPROPERTY(EditAnywhere, BlueprintReadOnly, Category = "DMX")
     FDMXImportGDTFDMXValue ModeFrom;
 
@@ -586,13 +757,15 @@ struct FDMXImportGDTFChannelFunction
 
     UPROPERTY(EditAnywhere, BlueprintReadOnly, Category = "DMX")
     TArray<FDMXImportGDTFChannelSet> ChannelSets;
+	PRAGMA_ENABLE_DEPRECATION_WARNINGS
 };
 
 USTRUCT(BlueprintType)
-struct FDMXImportGDTFLogicalChannel
+struct UE_DEPRECATED(5.5, "FDMXImportGDTFLogicalChannel is deprecated in favor of its corresponding type in the DMXGDTF module. See also UDMXImportGDTF::LoadGDTF.") FDMXImportGDTFLogicalChannel
 {
     GENERATED_BODY()
 
+	PRAGMA_DISABLE_DEPRECATION_WARNINGS
     UPROPERTY(EditAnywhere, BlueprintReadOnly, Category = "DMX")
     FDMXImportGDTFAttribute Attribute;
 
@@ -601,6 +774,7 @@ struct FDMXImportGDTFLogicalChannel
 
     UPROPERTY(EditAnywhere, BlueprintReadOnly, Category = "DMX")
     EDMXImportGDTFMaster Master = EDMXImportGDTFMaster::None;
+	PRAGMA_ENABLE_DEPRECATION_WARNINGS
 
     UPROPERTY(EditAnywhere, BlueprintReadOnly, Category = "DMX")
     float MibFade = 0.f;
@@ -608,12 +782,14 @@ struct FDMXImportGDTFLogicalChannel
     UPROPERTY(EditAnywhere, BlueprintReadOnly, Category = "DMX")
     float DMXChangeTimeLimit = 0.f;
 
+	PRAGMA_DISABLE_DEPRECATION_WARNINGS
     UPROPERTY(EditAnywhere, BlueprintReadOnly, Category = "DMX")
     FDMXImportGDTFChannelFunction ChannelFunction;
+	PRAGMA_ENABLE_DEPRECATION_WARNINGS
 };
 
 USTRUCT(BlueprintType)
-struct DMXRUNTIME_API FDMXImportGDTFDMXChannel
+struct UE_DEPRECATED(5.5, "FDMXImportGDTFDMXChannel is deprecated in favor of its corresponding type in the DMXGDTF module. See also UDMXImportGDTF::LoadGDTF.") DMXRUNTIME_API FDMXImportGDTFDMXChannel
 {
     GENERATED_BODY()
 
@@ -626,21 +802,25 @@ struct DMXRUNTIME_API FDMXImportGDTFDMXChannel
     UPROPERTY(EditAnywhere, BlueprintReadOnly, BlueprintReadOnly, Category = "DMX")
     TArray<int32> Offset;
 
+	PRAGMA_DISABLE_DEPRECATION_WARNINGS
     UPROPERTY(EditAnywhere, BlueprintReadOnly, BlueprintReadOnly, Category = "DMX")
     FDMXImportGDTFDMXValue Default;
 
     UPROPERTY(EditAnywhere, BlueprintReadOnly, BlueprintReadOnly, Category = "DMX")
     FDMXImportGDTFDMXValue Highlight;
+	PRAGMA_ENABLE_DEPRECATION_WARNINGS
 
     UPROPERTY(EditAnywhere, BlueprintReadOnly, BlueprintReadOnly, Category = "DMX")
     FName Geometry;
 
+	PRAGMA_DISABLE_DEPRECATION_WARNINGS
     UPROPERTY(EditAnywhere, BlueprintReadOnly, BlueprintReadOnly, Category = "DMX")
     FDMXImportGDTFLogicalChannel LogicalChannel;
+	PRAGMA_ENABLE_DEPRECATION_WARNINGS
 };
 
 USTRUCT(BlueprintType)
-struct FDMXImportGDTFRelation
+struct UE_DEPRECATED(5.5, "FDMXImportGDTFRelation is deprecated in favor of its corresponding type in the DMXGDTF module. See also UDMXImportGDTF::LoadGDTF.") FDMXImportGDTFRelation
 {
     GENERATED_BODY()
 
@@ -653,12 +833,14 @@ struct FDMXImportGDTFRelation
     UPROPERTY(EditAnywhere, BlueprintReadOnly, BlueprintReadOnly, Category = "DMX")
     FString Follower;
 
+	PRAGMA_DISABLE_DEPRECATION_WARNINGS
     UPROPERTY(EditAnywhere, BlueprintReadOnly, BlueprintReadOnly, Category = "DMX")
     EDMXImportGDTFType Type = EDMXImportGDTFType::Multiply;
+	PRAGMA_ENABLE_DEPRECATION_WARNINGS
 };
 
 USTRUCT(BlueprintType)
-struct FDMXImportGDTFFTMacro
+struct UE_DEPRECATED(5.5, "FDMXImportGDTFFTMacro is deprecated in favor of its corresponding type in the DMXGDTF module. See also UDMXImportGDTF::LoadGDTF.") FDMXImportGDTFFTMacro
 {
     GENERATED_BODY()
 
@@ -667,7 +849,7 @@ struct FDMXImportGDTFFTMacro
 };
 
 USTRUCT(BlueprintType)
-struct FDMXImportGDTFDMXMode
+struct UE_DEPRECATED(5.5, "FDMXImportGDTFDMXMode is deprecated in favor of its corresponding type in the DMXGDTF module. See also UDMXImportGDTF::LoadGDTF.") FDMXImportGDTFDMXMode
 {
     GENERATED_BODY()
 
@@ -677,6 +859,7 @@ struct FDMXImportGDTFDMXMode
     UPROPERTY(EditAnywhere, BlueprintReadOnly, BlueprintReadOnly, Category = "DMX")
     FName Geometry;
 
+	PRAGMA_DISABLE_DEPRECATION_WARNINGS
     UPROPERTY(EditAnywhere, BlueprintReadOnly, BlueprintReadOnly, Category = "DMX")
     TArray<FDMXImportGDTFDMXChannel> DMXChannels;
 
@@ -685,14 +868,19 @@ struct FDMXImportGDTFDMXMode
 
     UPROPERTY(EditAnywhere, BlueprintReadOnly, BlueprintReadOnly, Category = "DMX")
     TArray<FDMXImportGDTFFTMacro> FTMacros;
+	PRAGMA_ENABLE_DEPRECATION_WARNINGS
 };
 
-UCLASS(BlueprintType, Blueprintable)
+class UE_DEPRECATED(5.5, "UDMXImportGDTFFixtureType is deprecated in favor of its corresponding type in the DMXGDTF module. See also UDMXImportGDTF::LoadGDTF.") UDMXImportGDTFFixtureType;
+PRAGMA_DISABLE_DEPRECATION_WARNINGS
+UCLASS(BlueprintType, Blueprintable, Meta = (DeprecatedNode, DeprecationMessage = "This class is now deprecated in favor of a DMXEntityFixtureType based workflow. Please refer to members of DMXEntityFixtureType instead."))
 class DMXRUNTIME_API UDMXImportGDTFFixtureType
     : public UDMXImportFixtureType
 {
     GENERATED_BODY()
 
+PRAGMA_ENABLE_DEPRECATION_WARNINGS
+
 public:
     UPROPERTY(EditAnywhere, BlueprintReadOnly, Category = "Fixture Type")
     FName Name;
@@ -719,7 +907,9 @@ public:
     FString RefFT;
 };
 
-UCLASS(BlueprintType, Blueprintable)
+class UE_DEPRECATED(5.5, "UDMXImportGDTFAttributeDefinitions is deprecated in favor of its corresponding type in the DMXGDTF module. See also UDMXImportGDTF::LoadGDTF.") UDMXImportGDTFAttributeDefinitions;
+PRAGMA_DISABLE_DEPRECATION_WARNINGS
+UCLASS(BlueprintType, Blueprintable, Meta = (DeprecatedNode, DeprecationMessage = "This class is now deprecated in favor of a DMXEntityFixtureType based workflow. Please refer to members of DMXEntityFixtureType instead."))
 class DMXRUNTIME_API UDMXImportGDTFAttributeDefinitions
     : public UDMXImportAttributeDefinitions
 {
@@ -739,9 +929,13 @@ public:
 
     UPROPERTY(EditAnywhere, BlueprintReadOnly, Category = "DMX")
     TArray<FDMXImportGDTFAttribute> Attributes;
+
+PRAGMA_ENABLE_DEPRECATION_WARNINGS
 };
 
-UCLASS(BlueprintType, Blueprintable)
+class UE_DEPRECATED(5.5, "UDMXImportGDTFWheels is deprecated in favor of its corresponding type in the DMXGDTF module. See also UDMXImportGDTF::LoadGDTF.") UDMXImportGDTFWheels;
+PRAGMA_DISABLE_DEPRECATION_WARNINGS
+UCLASS(BlueprintType, Blueprintable, Meta = (DeprecatedNode, DeprecationMessage = "This class is now deprecated in favor of a DMXEntityFixtureType based workflow. Please refer to members of DMXEntityFixtureType instead."))
 class DMXRUNTIME_API UDMXImportGDTFWheels
     : public UDMXImportWheels
 {
@@ -753,9 +947,13 @@ public:
 public:
     UPROPERTY(EditAnywhere, BlueprintReadOnly, Category = "DMX")
     TArray<FDMXImportGDTFWheel> Wheels;
+
+PRAGMA_ENABLE_DEPRECATION_WARNINGS
 };
 
-UCLASS(BlueprintType, Blueprintable)
+class UE_DEPRECATED(5.5, "UDMXImportGDTFPhysicalDescriptions is deprecated in favor of its corresponding type in the DMXGDTF module. See also UDMXImportGDTF::LoadGDTF.") UDMXImportGDTFPhysicalDescription;
+PRAGMA_DISABLE_DEPRECATION_WARNINGS
+UCLASS(BlueprintType, Blueprintable, Meta = (DeprecatedNode, DeprecationMessage = "This class is now deprecated in favor of a DMXEntityFixtureType based workflow. Please refer to members of DMXEntityFixtureType instead."))
 class DMXRUNTIME_API UDMXImportGDTFPhysicalDescriptions
     : public UDMXImportPhysicalDescriptions
 {
@@ -776,10 +974,13 @@ public:
 
     UPROPERTY(EditAnywhere, BlueprintReadOnly, Category = "DMX")
     FDMXImportGDTFCRIs CRIs;
-};
 
+PRAGMA_ENABLE_DEPRECATION_WARNINGS
+};
 
-UCLASS(BlueprintType, Blueprintable)
+class UE_DEPRECATED(5.5, "UDMXImportGDTFModels is deprecated in favor of its corresponding type in the DMXGDTF module. See also UDMXImportGDTF::LoadGDTF.") UDMXImportGDTFModels;
+PRAGMA_DISABLE_DEPRECATION_WARNINGS
+UCLASS(BlueprintType, Blueprintable, Meta = (DeprecatedNode, DeprecationMessage = "This class is now deprecated in favor of a DMXEntityFixtureType based workflow. Please refer to members of DMXEntityFixtureType instead."))
 class DMXRUNTIME_API UDMXImportGDTFModels
     : public UDMXImportModels
 {
@@ -788,9 +989,13 @@ class DMXRUNTIME_API UDMXImportGDTFModels
 public:
     UPROPERTY(EditAnywhere, BlueprintReadOnly, Category = "DMX")
     TArray<FDMXImportGDTFModel> Models;
+
+PRAGMA_ENABLE_DEPRECATION_WARNINGS
 };
 
-UCLASS(BlueprintType, Blueprintable)
+class UE_DEPRECATED(5.5, "UDMXImportGDTFGeometries is deprecated in favor of its corresponding type in the DMXGDTF module. See also UDMXImportGDTF::LoadGDTF.") UDMXImportGDTFGeometries;
+PRAGMA_DISABLE_DEPRECATION_WARNINGS
+UCLASS(BlueprintType, Blueprintable, Meta = (DeprecatedNode, DeprecationMessage = "This class is now deprecated in favor of a DMXEntityFixtureType based workflow. Please refer to members of DMXEntityFixtureType instead."))
 class DMXRUNTIME_API UDMXImportGDTFGeometries
     : public UDMXImportGeometries
 {
@@ -799,9 +1004,13 @@ class DMXRUNTIME_API UDMXImportGDTFGeometries
 public:
     UPROPERTY(EditAnywhere, BlueprintReadOnly, Category = "DMX")
     TArray<FDMXImportGDTFGeneralGeometry> GeneralGeometry;
+
+PRAGMA_ENABLE_DEPRECATION_WARNINGS
 };
 
-UCLASS(BlueprintType, Blueprintable)
+class UE_DEPRECATED(5.5, "UDMXImportGDTFDMXModes is deprecated in favor of its corresponding type in the DMXGDTF module. See also UDMXImportGDTF::LoadGDTF.") UDMXImportGDTFDMXModes;
+PRAGMA_DISABLE_DEPRECATION_WARNINGS
+UCLASS(BlueprintType, Blueprintable, Meta = (DeprecatedNode, DeprecationMessage = "This class is now deprecated in favor of a DMXEntityFixtureType based workflow. Please refer to members of DMXEntityFixtureType instead."))
 class DMXRUNTIME_API UDMXImportGDTFDMXModes
     : public UDMXImportDMXModes
 {
@@ -811,52 +1020,27 @@ public:
     UPROPERTY(EditAnywhere, BlueprintReadOnly, Category = "DMX")
     TArray<FDMXImportGDTFDMXMode> DMXModes;
 
-public:
 	UFUNCTION(BlueprintPure, Category = "DMXGDTF|Import Data")
 	TArray<FDMXImportGDTFChannelFunction> GetDMXChannelFunctions(const FDMXImportGDTFDMXMode& InMode);
+
+PRAGMA_ENABLE_DEPRECATION_WARNINGS
 };
 
-UCLASS(BlueprintType, Blueprintable)
+class UE_DEPRECATED(5.5, "UDMXImportGDTFProtocols is deprecated in favor of its corresponding type in the DMXGDTF module. See also UDMXImportGDTF::LoadGDTF.") UDMXImportGDTFProtocols;
+PRAGMA_DISABLE_DEPRECATION_WARNINGS
+UCLASS(BlueprintType, Blueprintable, Meta = (DeprecatedNode, DeprecationMessage = "This class is now deprecated in favor of a DMXEntityFixtureType based workflow. Please refer to members of DMXEntityFixtureType instead."))
 class DMXRUNTIME_API UDMXImportGDTFProtocols
     : public UDMXImportProtocols
 {
     GENERATED_BODY()
 
+PRAGMA_ENABLE_DEPRECATION_WARNINGS
+
 public:
     UPROPERTY(EditAnywhere, BlueprintReadOnly, Category = "DMX")
     TArray<FName> Protocols;
 };
 
-UCLASS(BlueprintType, Blueprintable)
-class DMXRUNTIME_API UDMXImportGDTF
-	: public UDMXImport
-{
-    GENERATED_BODY()
+//~ End Deprecated Types
+///////////////////////////////
 
-public:
-	/** Constructor */
-	UDMXImportGDTF();
-
-	//~ Begin UObject interface
-	virtual void PostLoad() override;
-	//~ End UObject interface
-
-	UFUNCTION(BlueprintPure, Category = "DMXGDTF|Import Data")
-	UDMXImportGDTFDMXModes* GetDMXModes() const { return Cast<UDMXImportGDTFDMXModes>(DMXModes); }
-
-	/** DEPRECATED 5.1 in favor of AssetImportData */
-	UPROPERTY(Meta = (DeprecatedProperty, DeprecationMessage = "Deprecated in favor of GDTFAssetImportData, see UDMXImportGDTF::GetGDTFAssetImportData."))
-	FString SourceFilename_DEPRECATED;
-
-#if WITH_EDITORONLY_DATA
-	/** Returns GDTF Asset Import Data for this GDTF */
-	FORCEINLINE UDMXGDTFAssetImportData* GetGDTFAssetImportData() const { return GDTFAssetImportData; }
-#endif 
-
-private:
-#if WITH_EDITORONLY_DATA
-	/** The Asset Import Data used to generate the GDTF asset or nullptr, if not generated from a GDTF file */
-	UPROPERTY()
-	TObjectPtr<UDMXGDTFAssetImportData> GDTFAssetImportData;
-#endif
-};
diff --git a/Engine/Plugins/VirtualProduction/DMX/DMXGDTF/Source/DMXGDTF/Internal/GDTF/Geometries/DMXGDTFAddress.h b/Engine/Plugins/VirtualProduction/DMX/DMXGDTF/Source/DMXGDTF/Internal/GDTF/Geometries/DMXGDTFAddress.h
new file mode 100644
--- /dev/null
+++ b/Engine/Plugins/VirtualProduction/DMX/DMXGDTF/Source/DMXGDTF/Internal/GDTF/Geometries/DMXGDTFAddress.h
@@ -0,0 +1,31 @@
+// Copyright Epic Games, Inc. All Rights Reserved.
+
+#pragma once
+
+#include "GenericPlatform/GenericPlatform.h"
+
+namespace UE::DMX::GDTF
+{
+	/** The container for DMX Addresses */
+	struct FDMXGDTFAddress
+	{
+		FDMXGDTFAddress() = default;
+
+		FDMXGDTFAddress(int64 InAbsoluteAddress)
+			: AbsoluteAddress(InAbsoluteAddress)
+		{}
+
+		FDMXGDTFAddress(int32 InUniverse, int32 InChannel)
+		{
+			AbsoluteAddress = InUniverse * UniverseSize + InChannel - 1;
+		}
+
+		int32 GetUniverse() const { return AbsoluteAddress / UniverseSize; }
+
+		int32 GetChannel() const { return AbsoluteAddress % UniverseSize + 1; }
+
+		int64 AbsoluteAddress = 512;
+
+		static const uint16 UniverseSize = 512;
+	};
+}
diff --git a/Engine/Plugins/VirtualProduction/DMX/DMXGDTF/Source/DMXGDTF/Internal/GDTF/Geometries/DMXGDTFBreak.h b/Engine/Plugins/VirtualProduction/DMX/DMXGDTF/Source/DMXGDTF/Internal/GDTF/Geometries/DMXGDTFBreak.h
new file mode 100644
--- /dev/null
+++ b/Engine/Plugins/VirtualProduction/DMX/DMXGDTF/Source/DMXGDTF/Internal/GDTF/Geometries/DMXGDTFBreak.h
@@ -0,0 +1,33 @@
+// Copyright Epic Games, Inc. All Rights Reserved.
+
+#pragma once
+
+#include "GDTF/DMXGDTFNode.h"
+#include "GDTF/Geometries/DMXGDTFAddress.h"
+
+namespace UE::DMX::GDTF
+{
+	class FDMXGDTFGeometryReference;
+
+	/** This XML node specifies the DMX offset for the DMX channel of the referenced geometry (XML node <Break>) */
+	class DMXGDTF_API FDMXGDTFBreak
+		: public FDMXGDTFNode
+	{
+	public:
+		FDMXGDTFBreak(const TSharedRef<FDMXGDTFGeometryReference>& InGeometryReference);
+
+		//~ Begin DMXGDTFGeneralGeometryNode interface
+		virtual const TCHAR* GetXmlTag() const override { return TEXT("Break"); }
+		virtual void Initialize(const FXmlNode& XmlNode) override;
+		//~ End DMXGDTFGeneralGeometryNode interface
+
+		/** DMX offset; Default value:1 (Means no offset for the corresponding DMX Channel) */
+		FDMXGDTFAddress DMXAddress;
+
+		/** Defines the unique number of the DMX Break for which the Offset is given. Size: 1 byte; Default value 1 */
+		uint8 DMXBreak = 1;
+
+		/** The outer geometry reference */
+		const TWeakPtr<FDMXGDTFGeometryReference> OuterGeometryReference;
+	};
+}
diff --git a/Engine/Plugins/VirtualProduction/DMX/DMXGDTF/Source/DMXGDTF/Internal/GDTF/Protocols/RDM/DMXGDTFProtocolRDM.h b/Engine/Plugins/VirtualProduction/DMX/DMXGDTF/Source/DMXGDTF/Internal/GDTF/Protocols/RDM/DMXGDTFProtocolRDM.h
new file mode 100644
--- /dev/null
+++ b/Engine/Plugins/VirtualProduction/DMX/DMXGDTF/Source/DMXGDTF/Internal/GDTF/Protocols/RDM/DMXGDTFProtocolRDM.h
@@ -0,0 +1,37 @@
+// Copyright Epic Games, Inc. All Rights Reserved.
+
+#pragma once
+
+#include "GDTF/DMXGDTFNode.h"
+
+namespace UE::DMX::GDTF { class FDMXGDTFProtocols; }
+
+namespace UE::DMX::GDTF::RDM
+{
+	class FDMXGDTFSoftwareVersionID;
+
+	/** This section defines the overall ProtocolRDM of the device (XML node <ProtocolRDM>). */
+	class DMXGDTF_API FDMXGDTFProtocolRDM
+		: public FDMXGDTFNode
+	{
+	public:
+		FDMXGDTFProtocolRDM(const TSharedRef<FDMXGDTFProtocols>& InProtocols);
+
+		//~ Begin FDMXGDTFNode interface
+		virtual const TCHAR* GetXmlTag() const override { return TEXT("FTRDM"); }
+		virtual void Initialize(const FXmlNode& XmlNode) override;
+		//~ End FDMXGDTFNode interface
+
+		/** Manufacturer ESTA ID */
+		uint32 ManufacturerID = 0;	
+
+		/** Unique device model ID */
+		uint32 DeviceModelID = 0;
+
+		/** As children the FTRDM has a list of SoftwareVersionID. */
+		TArray<TSharedPtr<FDMXGDTFSoftwareVersionID>> SoftwareVersionIDArray;
+
+		/** The outer protocols */
+		const TWeakPtr<FDMXGDTFProtocols> OuterProtocols;
+	};
+}
diff --git a/Engine/Plugins/VirtualProduction/DMX/DMXGDTF/Source/DMXGDTF/Private/GDTF/DMXModes/DMXGDTFDMXValue.cpp b/Engine/Plugins/VirtualProduction/DMX/DMXGDTF/Source/DMXGDTF/Private/GDTF/DMXModes/DMXGDTFDMXValue.cpp
new file mode 100644
--- /dev/null
+++ b/Engine/Plugins/VirtualProduction/DMX/DMXGDTF/Source/DMXGDTF/Private/GDTF/DMXModes/DMXGDTFDMXValue.cpp
@@ -0,0 +1,92 @@
+// Copyright Epic Games, Inc. All Rights Reserved.
+
+#include "GDTF/DMXModes/DMXGDTFDMXValue.h"
+
+namespace UE::DMX::GDTF
+{
+	bool FDMXGDTFDMXValue::ToInt(uint32& OutInteger) const
+	{
+		if (Value.IsEmpty() || Value == TEXT("None"))
+		{
+			return false;
+		}
+
+		TArray<FString> Substrings;
+		Value.ParseIntoArray(Substrings, TEXT("/"));
+
+		// As per GDTF specs
+		if (Substrings.Num() == 2)
+		{
+			uint32 DMXValue;
+			if (!LexTryParseString(DMXValue, *Substrings[0]))
+			{
+				return false;
+			}
+
+			if (Substrings[1].EndsWith(TEXT("s")))
+			{
+				Substrings[1].RemoveFromEnd(TEXT("s"));
+
+				uint8 ByteShiftValue;
+				if (LexTryParseString(ByteShiftValue, *Substrings[1]))
+				{
+					if (ByteShiftValue == 1)
+					{
+						// No shifting needed
+						OutInteger = DMXValue;
+						return true;
+					}
+					else if (ByteShiftValue > 4)
+					{
+						return false;
+					}
+
+					OutInteger = DMXValue << (ByteShiftValue - 1) * 8;
+					return true;
+				}
+			}
+			else
+			{
+				uint8 ByteMirroringValue;
+				if (LexTryParseString(ByteMirroringValue, *Substrings[1]))
+				{
+					if (ByteMirroringValue == 1)
+					{
+						// No mirroring needed
+						OutInteger = DMXValue;
+						return true;
+					}
+					else if (ByteMirroringValue > 4)
+					{
+						return false;
+					}
+					else
+					{
+						constexpr uint32 Masks[] = { 0x000000FF, 0x0000FF00, 0x00FF0000, 0xFF000000 };
+
+						if (ByteMirroringValue == 2)
+						{
+							OutInteger = ((DMXValue & 0x000000FF) << 8) | ((DMXValue & 0x0000FF00) >> 8);
+						}
+						else if (ByteMirroringValue == 3)
+						{
+#if PLATFORM_LITTLE_ENDIAN
+							OutInteger = ((DMXValue & 0x0000FF00) << 16) | ((DMXValue & 0xFF000000) >> 16);
+#else
+							OutInteger = ((DMXValue & 0x000000FF) << 16) | ((DMXValue & 0x00FF0000) >> 16);
+#endif
+						}
+						else
+						{
+							OutInteger = ((DMXValue & 0x000000FF) << 24) | ((DMXValue & 0x0000FF00) << 8) | ((DMXValue & 0x00FF0000) >> 8) | ((DMXValue & 0xFF000000) >> 24);
+						}
+					}
+
+					return true;
+				}
+			}
+		}
+
+		return false;
+	}
+}
diff --git a/Engine/Plugins/VirtualProduction/DMX/DMXGDTF/Source/DMXGDTF/Private/GDTF/Geometries/DMXGDTFBreak.cpp b/Engine/Plugins/VirtualProduction/DMX/DMXGDTF/Source/DMXGDTF/Private/GDTF/Geometries/DMXGDTFBreak.cpp
new file mode 100644
--- /dev/null
+++ b/Engine/Plugins/VirtualProduction/DMX/DMXGDTF/Source/DMXGDTF/Private/GDTF/Geometries/DMXGDTFBreak.cpp
@@ -0,0 +1,19 @@
+// Copyright Epic Games, Inc. All Rights Reserved.
+
+#include "GDTF/Geometries/DMXGDTFBreak.h"
+
+#include "Serialization/DMXGDTFNodeInitializer.h"
+
+namespace UE::DMX::GDTF
+{
+	FDMXGDTFBreak::FDMXGDTFBreak(const TSharedRef<FDMXGDTFGeometryReference>& InGeometryReference)
+		: OuterGeometryReference(InGeometryReference)
+	{}
+
+	void FDMXGDTFBreak::Initialize(const FXmlNode& XmlNode)
+	{
+		FDMXGDTFNodeInitializer(SharedThis(this), XmlNode)
+			.GetAttribute(TEXT("DMXAddress"), DMXAddress)
+			.GetAttribute(TEXT("DMXBreak"), DMXBreak);
+	}
+}
diff --git a/Engine/Plugins/VirtualProduction/DMX/DMXGDTF/Source/DMXGDTF/Private/GDTF/Geometries/DMXGDTFGeometryCollectBase.cpp b/Engine/Plugins/VirtualProduction/DMX/DMXGDTF/Source/DMXGDTF/Private/GDTF/Geometries/DMXGDTFGeometryCollectBase.cpp
new file mode 100644
--- /dev/null
+++ b/Engine/Plugins/VirtualProduction/DMX/DMXGDTF/Source/DMXGDTF/Private/GDTF/Geometries/DMXGDTFGeometryCollectBase.cpp
@@ -0,0 +1,93 @@
+// Copyright Epic Games, Inc. All Rights Reserved.
+
+#include "GDTF/Geometries/DMXGDTFGeometryCollectBase.h"
+
+#include "GDTF/Geometries/DMXGDTFAxisGeometry.h"
+#include "GDTF/Geometries/DMXGDTFBeamGeometry.h"
+#include "GDTF/Geometries/DMXGDTFDisplayGeometry.h"
+#include "GDTF/Geometries/DMXGDTFFilterBeamGeometry.h"
+#include "GDTF/Geometries/DMXGDTFFilterColorGeometry.h"
+#include "GDTF/Geometries/DMXGDTFFilterGoboGeometry.h"
+#include "GDTF/Geometries/DMXGDTFFilterShaperGeometry.h"
+#include "GDTF/Geometries/DMXGDTFGeometry.h"
+#include "GDTF/Geometries/DMXGDTFGeometryReference.h"
+#include "GDTF/Geometries/DMXGDTFInventoryGeometry.h"
+#include "GDTF/Geometries/DMXGDTFLaserGeometry.h"
+#include "GDTF/Geometries/DMXGDTFMagnetGeometry.h"
+#include "GDTF/Geometries/DMXGDTFMagnetGeometry.h"
+#include "GDTF/Geometries/DMXGDTFMediaServerCameraGeometry.h"
+#include "GDTF/Geometries/DMXGDTFMediaServerLayerGeometry.h"
+#include "GDTF/Geometries/DMXGDTFMediaServerMasterGeometry.h"
+#include "GDTF/Geometries/DMXGDTFStructureGeometry.h"
+#include "GDTF/Geometries/DMXGDTFSupportGeometry.h"
+#include "GDTF/Geometries/DMXGDTFWiringObjectGeometry.h"
+#include "Serialization/DMXGDTFNodeInitializer.h"
+
+namespace UE::DMX::GDTF
+{
+	const TCHAR* FDMXGDTFGeometryCollectBase::GetXmlTag() const
+	{
+		ensureMsgf(0, TEXT("Unexpected call to FDMXGDTFGeometryCollectBase::GetXmlTag in abstract FDMXGDTFGeometryCollectBase."));
+		return TEXT("Invalid");
+	}
+
+	void FDMXGDTFGeometryCollectBase::Initialize(const FXmlNode& XmlNode)
+	{
+		FDMXGDTFNodeInitializer(SharedThis(this), XmlNode)
+			.CreateChildren(TEXT("Geometry"), GeometryArray)
+			.CreateChildren(TEXT("Axis"), AxisArray)
+			.CreateChildren(TEXT("FilterBeam"), FilterBeamArray)
+			.CreateChildren(TEXT("FilterColor"), FilterColorArray)
+			.CreateChildren(TEXT("FilterGobo"), FilterGoboArray)
+			.CreateChildren(TEXT("FilterShaper"), FilterShaperArray)
+			.CreateChildren(TEXT("Beam"), BeamArray)
+			.CreateChildren(TEXT("MediaServerLayer"), MediaServerLayerArray)
+			.CreateChildren(TEXT("MediaServerCamera"), MediaServerCameraArray)
+			.CreateChildren(TEXT("MediaServerMaster"), MediaServerMasterArray)
+			.CreateChildren(TEXT("Display"), DisplayArray)
+			.CreateChildren(TEXT("GeometryReference"), GeometryReferenceArray)
+			.CreateChildren(TEXT("Laser"), LaserArray)
+			.CreateChildren(TEXT("WiringObject"), WiringObjectArray)
+			.CreateChildren(TEXT("Inventory"), InventoryArray)
+			.CreateChildren(TEXT("Structure"), StructureArray)
+			.CreateChildren(TEXT("Support"), SupportArray)
+			.CreateChildren(TEXT("Magnet"), MagnetArray);
+	}
+
+	void FDMXGDTFGeometryCollectBase::FindGeometryByName(const TCHAR* InName, TSharedPtr<FDMXGDTFGeometry>& OutGeometry, TSharedPtr<FDMXGDTFGeometryReference>& OutGeometryReference) const
+	{
+		const TSharedPtr<FDMXGDTFGeometryReference>* GeometryReferencePtr = Algo::FindBy(GeometryReferenceArray, InName, &FDMXGDTFGeometryReference::Name);
+		if (GeometryReferencePtr)
+		{
+			OutGeometryReference = *GeometryReferencePtr;
+			return;
+		}
+
+		// Helper to find geometry nodes in arrays of different types
+		auto FindInArrayLambda = [InName](auto InArray, TSharedPtr<FDMXGDTFGeometry>& OutGeometry) -> bool
+			{
+				auto const* GeometryPtr = Algo::FindBy(InArray, InName, &FDMXGDTFGeometry::Name);
+				OutGeometry = GeometryPtr ? *GeometryPtr : nullptr;
+				return GeometryPtr != nullptr;
+			};
+
+		TSharedPtr<FDMXGDTFGeometry> Geometry;
+		if (FindInArrayLambda(GeometryArray, Geometry)) { OutGeometry = Geometry; }
+		else if (FindInArrayLambda(AxisArray, Geometry)) { OutGeometry = Geometry; }
+		else if (FindInArrayLambda(FilterBeamArray, Geometry)) { OutGeometry = Geometry; }
+		else if (FindInArrayLambda(FilterColorArray, Geometry)) { OutGeometry = Geometry; }
+		else if (FindInArrayLambda(FilterGoboArray, Geometry)) { OutGeometry = Geometry; }
+		else if (FindInArrayLambda(FilterShaperArray, Geometry)) { OutGeometry = Geometry; }
+		else if (FindInArrayLambda(BeamArray, Geometry)) { OutGeometry = Geometry; }
+		else if (FindInArrayLambda(MediaServerLayerArray, Geometry)) { OutGeometry = Geometry; }
+		else if (FindInArrayLambda(MediaServerCameraArray, Geometry)) { OutGeometry = Geometry; }
+		else if (FindInArrayLambda(MediaServerMasterArray, Geometry)) { OutGeometry = Geometry; }
+		else if (FindInArrayLambda(DisplayArray, Geometry)) { OutGeometry = Geometry; }
+		else if (FindInArrayLambda(LaserArray, Geometry)) { OutGeometry = Geometry; }
+		else if (FindInArrayLambda(WiringObjectArray, Geometry)) { OutGeometry = Geometry; }
+		else if (FindInArrayLambda(InventoryArray, Geometry)) { OutGeometry = Geometry; }
+		else if (FindInArrayLambda(StructureArray, Geometry)) { OutGeometry = Geometry; }
+		else if (FindInArrayLambda(SupportArray, Geometry)) { OutGeometry = Geometry; }
+		else if (FindInArrayLambda(MagnetArray, Geometry)) { OutGeometry = Geometry; }
+	}
+}
diff --git a/Engine/Plugins/VirtualProduction/DMX/DMXGDTF/Source/DMXGDTF/Private/GDTF/Protocols/RDM/DMXGDTFProtocolRDM.cpp b/Engine/Plugins/VirtualProduction/DMX/DMXGDTF/Source/DMXGDTF/Private/GDTF/Protocols/RDM/DMXGDTFProtocolRDM.cpp
new file mode 100644
--- /dev/null
+++ b/Engine/Plugins/VirtualProduction/DMX/DMXGDTF/Source/DMXGDTF/Private/GDTF/Protocols/RDM/DMXGDTFProtocolRDM.cpp
@@ -0,0 +1,21 @@
+// Copyright Epic Games, Inc. All Rights Reserved.
+
+#include "GDTF/Protocols/RDM/DMXGDTFProtocolRDM.h"
+
+#include "GDTF/Protocols/RDM/DMXGDTFSoftwareVersionID.h"
+#include "Serialization/DMXGDTFNodeInitializer.h"
+
+namespace UE::DMX::GDTF::RDM
+{
+	FDMXGDTFProtocolRDM::FDMXGDTFProtocolRDM(const TSharedRef<FDMXGDTFProtocols>& InProtocols)
+		: OuterProtocols(InProtocols)
+	{}
+
+	void FDMXGDTFProtocolRDM::Initialize(const FXmlNode& XmlNode)
+	{
+		FDMXGDTFNodeInitializer(SharedThis(this), XmlNode)
+			.GetAttribute(TEXT("ManufacturerID"), ManufacturerID, &FParse::HexNumber)
+			.GetAttribute(TEXT("DeviceModelID"), DeviceModelID, &FParse::HexNumber)
+			.CreateChildren(TEXT("SoftwareVersionID"), SoftwareVersionIDArray);
+	}
+}
