DMX - Fix text conflict from LOCTEXT macro for namespace "DMXGDTFImporter" and key "InvalidGDTFError"

#rnx

diff --git a/Engine/Plugins/VirtualProduction/DMX/DMXEngine/Source/DMXEditor/Private/Factories/DMXGDTFImporter.cpp b/Engine/Plugins/VirtualProduction/DMX/DMXEngine/Source/DMXEditor/Private/Factories/DMXGDTFImporter.cpp
--- a/Engine/Plugins/VirtualProduction/DMX/DMXEngine/Source/DMXEditor/Private/Factories/DMXGDTFImporter.cpp
+++ b/Engine/Plugins/VirtualProduction/DMX/DMXEngine/Source/DMXEditor/Private/Factories/DMXGDTFImporter.cpp
@@ -12,7 +12,6 @@
 #include "Misc/Paths.h"
 #include "XmlFile.h"
 
-
 #define LOCTEXT_NAMESPACE "DMXGDTFImporter"
 
 namespace UE::DMX
@@ -57,7 +56,7 @@ namespace UE::DMX
 		TArray64<uint8> DescriptionXmlData;
 		if (!Zip->GetFileContent(DescriptionXmlFilename, DescriptionXmlData))
 		{
-			OutErrorReason = LOCTEXT("InvalidGDTFError", "Cannot read GDTF. Cannot find Description.xml.");
+			OutErrorReason = LOCTEXT("MissingDescriptionXmlGDTFError", "Cannot read GDTF. Cannot find Description.xml.");
 			return nullptr;
 		}
 
