DMX - DMX Libraries are now only loaded on demand, avoiding memory overhead

diff --git a/Engine/Plugins/VirtualProduction/DMX/DMXEngine/Config/BaseDMXEngine.ini b/Engine/Plugins/VirtualProduction/DMX/DMXEngine/Config/BaseDMXEngine.ini
--- a/Engine/Plugins/VirtualProduction/DMX/DMXEngine/Config/BaseDMXEngine.ini
+++ b/Engine/Plugins/VirtualProduction/DMX/DMXEngine/Config/BaseDMXEngine.ini
@@ -1,3 +1,6 @@
 ; UE4.26 to UE4.27
 [CoreRedirects]
 +EnumRedirects=(OldName="EDMXCommunicationTypes", NewName="EDMXCommunicationType") ; Only used for DMX Libraries that used deprecated UDMXEntitiyController, to upgrade them to 4.27
+
+; UE5.4 to UE5.5
++FunctionRedirects=(OldName="GetAllDMXLibraries",NewName="/Script/DMXEngine.DMXSubsystem.LoadDMXLibrariesSynchronous")
diff --git a/Engine/Plugins/VirtualProduction/DMX/DMXEngine/Source/DMXEditor/Private/DMXEditorModule.cpp b/Engine/Plugins/VirtualProduction/DMX/DMXEngine/Source/DMXEditor/Private/DMXEditorModule.cpp
--- a/Engine/Plugins/VirtualProduction/DMX/DMXEngine/Source/DMXEditor/Private/DMXEditorModule.cpp
+++ b/Engine/Plugins/VirtualProduction/DMX/DMXEngine/Source/DMXEditor/Private/DMXEditorModule.cpp
@@ -393,7 +393,7 @@ TSharedRef<SDockTab> FDMXEditorModule::OnSpawnPatchToolTab(const FSpawnTabArgs&
 		.Label(LOCTEXT("PatchToolTitle", "DMX Patch Tool"))
 		.TabRole(ETabRole::NomadTab)
 		[
-			SNew(SDMXPatchTool)
+			SNew(UE::DMX::SDMXPatchTool)
 		];
 }
 
diff --git a/Engine/Plugins/VirtualProduction/DMX/DMXEngine/Source/DMXEditor/Private/DMXEditorUtils.cpp b/Engine/Plugins/VirtualProduction/DMX/DMXEngine/Source/DMXEditor/Private/DMXEditorUtils.cpp
--- a/Engine/Plugins/VirtualProduction/DMX/DMXEngine/Source/DMXEditor/Private/DMXEditorUtils.cpp
+++ b/Engine/Plugins/VirtualProduction/DMX/DMXEngine/Source/DMXEditor/Private/DMXEditorUtils.cpp
@@ -2,34 +2,32 @@
 
 #include "DMXEditorUtils.h"
 
+#include "AssetRegistry/AssetData.h"
+#include "AssetRegistry/AssetRegistryModule.h"
+#include "Dialogs/Dialogs.h"
 #include "DMXEditorLog.h"
 #include "DMXRuntimeUtils.h"
 #include "DMXSubsystem.h"
+#include "Exporters/Exporter.h"
+#include "Factories.h"
+#include "HAL/PlatformApplicationMisc.h"
+#include "Internationalization/Regex.h"
 #include "IO/DMXInputPort.h"
 #include "IO/DMXOutputPort.h"
 #include "IO/DMXPortManager.h"
-#include "Library/DMXLibrary.h"
-#include "Library/DMXEntityFixtureType.h"
 #include "Library/DMXEntityFixturePatch.h"
+#include "Library/DMXEntityFixtureType.h"
+#include "Library/DMXLibrary.h"
 #include "MVR/DMXMVRGeneralSceneDescription.h"
 #include "MVR/Types/DMXMVRFixtureNode.h"
-
-#include "AssetRegistry/AssetRegistryModule.h"
-#include "Factories.h"
 #include "PackageTools.h"
 #include "ScopedTransaction.h"
 #include "UnrealExporter.h"
-#include "Dialogs/Dialogs.h"
-#include "Exporters/Exporter.h"
-#include "HAL/PlatformApplicationMisc.h"
-#include "Internationalization/Regex.h"
 #include "UObject/Package.h"
 
-
 #define LOCTEXT_NAMESPACE "FDMXEditorUtils"
 
-
-// Text object factory for pasting DMX Entities
+/** Text object factory for pasting DMX Entities */
 struct FDMXEntityObjectTextFactory : public FCustomizableTextObjectFactory
 {
 	/** Entities instantiated */
@@ -1110,12 +1108,12 @@ void FDMXEditorUtils::ClearFixturePatchCachedData()
 	UDMXSubsystem* Subsystem = UDMXSubsystem::GetDMXSubsystem_Callable();
 	if (Subsystem && Subsystem->IsValidLowLevel())
 	{
-		TArray<UDMXLibrary*> DMXLibraries = Subsystem->GetAllDMXLibraries();
-		for (UDMXLibrary* Library : DMXLibraries)
+		TArray<TSoftObjectPtr<UDMXLibrary>> DMXLibraries = Subsystem->GetDMXLibraries();
+		for (const TSoftObjectPtr<UDMXLibrary>& Library : DMXLibraries)
 		{
-			if (Library != nullptr && Library->IsValidLowLevel())
+			if (Library.IsValid())
 			{
-				Library->ForEachEntityOfType<UDMXEntityFixturePatch>([](UDMXEntityFixturePatch* Patch) {
+				Library.Get()->ForEachEntityOfType<UDMXEntityFixturePatch>([](UDMXEntityFixturePatch* Patch) {
 					Patch->RebuildCache();
 				});
 			}
diff --git a/Engine/Plugins/VirtualProduction/DMX/DMXEngine/Source/DMXEditor/Private/Widgets/PatchTool/DMXPatchToolItem.h b/Engine/Plugins/VirtualProduction/DMX/DMXEngine/Source/DMXEditor/Private/Widgets/PatchTool/DMXPatchToolItem.h
new file mode 100644
--- /dev/null
+++ b/Engine/Plugins/VirtualProduction/DMX/DMXEngine/Source/DMXEditor/Private/Widgets/PatchTool/DMXPatchToolItem.h
@@ -0,0 +1,22 @@
+// Copyright Epic Games, Inc. All Rights Reserved.
+
+#pragma once
+
+#include "Templates/SharedPointer.h"
+#include "UObject/SoftObjectPtr.h"
+
+class UDMXLibrary;
+
+namespace UE::DMX
+{
+	class FDMXPatchToolItem
+		: public TSharedFromThis<FDMXPatchToolItem>
+	{
+	public:
+		FDMXPatchToolItem(TSoftObjectPtr<UDMXLibrary> InDMXLibrary)
+			: SoftDMXLibrary(InDMXLibrary)
+		{}
+
+		const TSoftObjectPtr<UDMXLibrary> SoftDMXLibrary;
+	};
+}
diff --git a/Engine/Plugins/VirtualProduction/DMX/DMXEngine/Source/DMXEditor/Private/Widgets/PatchTool/SDMXPatchTool.cpp b/Engine/Plugins/VirtualProduction/DMX/DMXEngine/Source/DMXEditor/Private/Widgets/PatchTool/SDMXPatchTool.cpp
--- a/Engine/Plugins/VirtualProduction/DMX/DMXEngine/Source/DMXEditor/Private/Widgets/PatchTool/SDMXPatchTool.cpp
+++ b/Engine/Plugins/VirtualProduction/DMX/DMXEngine/Source/DMXEditor/Private/Widgets/PatchTool/SDMXPatchTool.cpp
@@ -2,456 +2,419 @@
 
 #include "SDMXPatchTool.h"
 
+#include "Algo/Transform.h"
+#include "DMXPatchToolItem.h"
 #include "DMXSubsystem.h"
-#include "Game/DMXComponent.h"
-#include "Library/DMXLibrary.h"
-#include "Library/DMXEntityFixturePatch.h"
-
 #include "Editor.h"
-#include "GameFramework/Actor.h"
 #include "Engine/Selection.h"
+#include "Game/DMXComponent.h"
+#include "GameFramework/Actor.h"
+#include "Library/DMXEntityFixturePatch.h"
+#include "Library/DMXLibrary.h"
 #include "Styling/CoreStyle.h"
 #include "Widgets/Input/SButton.h"
 #include "Widgets/Input/SComboBox.h"
 #include "Widgets/Layout/SGridPanel.h"
 
-
 #define LOCTEXT_NAMESPACE "SDMXPatchTool"
 
-SDMXPatchTool::SDMXPatchTool()
-	: AnalyticsProvider("PatchTool")
-{}
-
-SDMXPatchTool::~SDMXPatchTool()
+namespace UE::DMX
 {
-	// Unbind from library changes
-	if (PreviouslySelectedLibrary.IsValid())
+	SDMXPatchTool::SDMXPatchTool()
+		: AnalyticsProvider("PatchTool")
+	{}
+
+	SDMXPatchTool::~SDMXPatchTool()
 	{
-		PreviouslySelectedLibrary->GetOnEntitiesAdded().RemoveAll(this);
-		PreviouslySelectedLibrary->GetOnEntitiesRemoved().RemoveAll(this);
+		// Unbind from library changes
+		if (PreviouslySelectedLibrary.IsValid())
+		{
+			PreviouslySelectedLibrary->GetOnEntitiesAdded().RemoveAll(this);
+			PreviouslySelectedLibrary->GetOnEntitiesRemoved().RemoveAll(this);
+		}
 	}
-}
-
-void SDMXPatchTool::Construct(const FArguments& InArgs)
-{
-	UDMXSubsystem* Subsystem = UDMXSubsystem::GetDMXSubsystem_Pure();
-	check(Subsystem);
 
-	LibrarySource = ObjectPtrWrap(Subsystem->GetAllDMXLibraries());
+	void SDMXPatchTool::Construct(const FArguments& InArgs)
+	{
+		UDMXSubsystem* Subsystem = UDMXSubsystem::GetDMXSubsystem_Pure();
+		check(Subsystem);
 
-	ChildSlot
-	[
-		SNew(SVerticalBox)
+		Algo::Transform(Subsystem->GetDMXLibraries(), LibrarySource,
+			[](const TSoftObjectPtr<UDMXLibrary>& SoftDMXLibrary)
+			{
+				return MakeShared<FDMXPatchToolItem>(SoftDMXLibrary);
+			});
 
-		+ SVerticalBox::Slot()
-		.HAlign(HAlign_Fill)
-		.VAlign(VAlign_Top)
-		.AutoHeight()
+		ChildSlot
 		[
-			SNew(SGridPanel)
+			SNew(SVerticalBox)
 
-			// Library Selection Label
-			+ SGridPanel::Slot(0, 0)
-			.Padding(4.f)
-			.HAlign(HAlign_Left)
-			.VAlign(VAlign_Center)
+			+ SVerticalBox::Slot()
+			.HAlign(HAlign_Fill)
+			.VAlign(VAlign_Top)
+			.AutoHeight()
 			[
-				SNew(SBox)
-				.MinDesiredWidth(160.f)
-				.MaxDesiredWidth(160.f)
+				SNew(SGridPanel)
+
+				// Library Selection Label
+				+ SGridPanel::Slot(0, 0)
+				.Padding(4.f)
+				.HAlign(HAlign_Left)
+				.VAlign(VAlign_Center)
 				[
-					SNew(STextBlock)
-					.Text(LOCTEXT("DMXLibraryComboboxLabel", "DMX Library"))
-					.Font(FCoreStyle::GetDefaultFontStyle("Regular", 14))
+					SNew(SBox)
+					.MinDesiredWidth(160.f)
+					.MaxDesiredWidth(160.f)
+					[
+						SNew(STextBlock)
+						.Text(LOCTEXT("DMXLibraryComboboxLabel", "DMX Library"))
+						.Font(FCoreStyle::GetDefaultFontStyle("Regular", 14))
+					]
 				]
-			]
 
-			// Library Selection Combo Box
-			+ SGridPanel::Slot(1, 0)
-			.Padding(4.f)
-			.HAlign(HAlign_Fill)
-			.VAlign(VAlign_Center)
-			[
-				SNew(SHorizontalBox)
+				// Library Selection Combo Box
+				+ SGridPanel::Slot(1, 0)
+				.Padding(4.f)
+				.HAlign(HAlign_Fill)
+				.VAlign(VAlign_Center)
+				[
+					SNew(SHorizontalBox)
 
-				+ SHorizontalBox::Slot()
+					+ SHorizontalBox::Slot()
+					.HAlign(HAlign_Left)
+					.AutoWidth()
+					.VAlign(VAlign_Center)
+					[
+						SAssignNew(LibraryComboBox, SComboBox<TSharedPtr<FDMXPatchToolItem>>)
+						.OnGenerateWidget(this, &SDMXPatchTool::GenerateLibraryComboBoxEntry)
+						.OnSelectionChanged(this, &SDMXPatchTool::OnLibrarySelected)
+						.OptionsSource(&ObjectPtrDecay(LibrarySource))
+						[
+							SAssignNew(SelectedLibraryTextBlock, STextBlock)
+						]
+					]
+				]
+
+				// Patch Selection Label
+				+ SGridPanel::Slot(0, 1)
+				.Padding(4.f)
 				.HAlign(HAlign_Left)
-				.AutoWidth()
 				.VAlign(VAlign_Center)
 				[
-					SAssignNew(LibraryComboBox, SComboBox<UDMXLibrary*>)
-					.OnGenerateWidget(this, &SDMXPatchTool::GenerateLibraryComboBoxEntry)
-					.OnSelectionChanged(this, &SDMXPatchTool::OnLibrarySelected)
-					.OptionsSource(&ObjectPtrDecay(LibrarySource))
+					SNew(SBox)
+					.MinDesiredWidth(160.f)
+					.MaxDesiredWidth(160.f)
 					[
-						SAssignNew(SelectedLibraryTextBlock, STextBlock)
+						SNew(STextBlock)
+						.Text(LOCTEXT("DMXFixturePatchComboboxLabel", "Fixture Patch"))
+						.Font(FCoreStyle::GetDefaultFontStyle("Regular", 14))
 					]
 				]
 
-			]
-
-			// Patch Selection Label
-			+ SGridPanel::Slot(0, 1)
-			.Padding(4.f)
-			.HAlign(HAlign_Left)
-			.VAlign(VAlign_Center)
-			[
-				SNew(SBox)
-				.MinDesiredWidth(160.f)
-				.MaxDesiredWidth(160.f)
+				// Patch Selection Combo Box
+				+ SGridPanel::Slot(1, 1)
+				.Padding(4.f)
+				.HAlign(HAlign_Fill)
+				.VAlign(VAlign_Center)
 				[
-					SNew(STextBlock)
-					.Text(LOCTEXT("DMXFixturePatchComboboxLabel", "Fixture Patch"))
-					.Font(FCoreStyle::GetDefaultFontStyle("Regular", 14))
+					SNew(SHorizontalBox)
+
+					+ SHorizontalBox::Slot()
+					.HAlign(HAlign_Left)
+					.AutoWidth()
+					.VAlign(VAlign_Center)
+					[
+						SAssignNew(FixturePatchComboBox, SComboBox<UDMXEntityFixturePatch*>)
+						.OnGenerateWidget(this, &SDMXPatchTool::GenerateFixturePatchComboBoxEntry)
+						.OnSelectionChanged(this, &SDMXPatchTool::OnFixturePatchSelected)
+						.OptionsSource(&ObjectPtrDecay(FixturePatchSource))
+						[
+							SAssignNew(SelectedFixturePatchTextBlock, STextBlock)
+						]
+					]
 				]
 			]
 
-			// Patch Selection Combo Box
-			+ SGridPanel::Slot(1, 1)
-			.Padding(4.f)
+			+ SVerticalBox::Slot()
 			.HAlign(HAlign_Fill)
-			.VAlign(VAlign_Center)
+			.VAlign(VAlign_Top)
+			.AutoHeight()
 			[
 				SNew(SHorizontalBox)
 
 				+ SHorizontalBox::Slot()
-				.HAlign(HAlign_Left)
-				.AutoWidth()
-				.VAlign(VAlign_Center)
+				.HAlign(HAlign_Fill)
+				.VAlign(VAlign_Top)
+				.Padding(8.f)
 				[
-					SAssignNew(FixturePatchComboBox, SComboBox<UDMXEntityFixturePatch*>)
-					.OnGenerateWidget(this, &SDMXPatchTool::GenerateFixturePatchComboBoxEntry)
-					.OnSelectionChanged(this, &SDMXPatchTool::OnFixturePatchSelected)
-					.OptionsSource(&ObjectPtrDecay(FixturePatchSource))
+					SNew(SButton)
+					.OnClicked(this, &SDMXPatchTool::OnAddressIncrementalClicked)
 					[
-						SAssignNew(SelectedFixturePatchTextBlock, STextBlock)
+						SNew(STextBlock)
+						.Text(LOCTEXT("AddressIncrementalButtonText", "Address incremental"))
+						.Font(FCoreStyle::GetDefaultFontStyle("Regular", 14))
 					]
 				]
-			]
-		]
 
-		+ SVerticalBox::Slot()
-		.HAlign(HAlign_Fill)
-		.VAlign(VAlign_Top)
-		.AutoHeight()
-		[
-			SNew(SHorizontalBox)
-
-			+SHorizontalBox::Slot()
-			.HAlign(HAlign_Fill)
-			.VAlign(VAlign_Top)
-			.Padding(8.f)
-			[
-				SNew(SButton)
-				.OnClicked(this, &SDMXPatchTool::OnAddressIncrementalClicked)
+				+ SHorizontalBox::Slot()
+				.HAlign(HAlign_Fill)
+				.VAlign(VAlign_Top)
+				.Padding(8.f)
 				[
-					SNew(STextBlock)
-					.Text(LOCTEXT("AddressIncrementalButtonText", "Address incremental"))
-					.Font(FCoreStyle::GetDefaultFontStyle("Regular", 14))
+					SNew(SButton)
+					.OnClicked(this, &SDMXPatchTool::OnAddressSameClicked)
+					[
+						SNew(STextBlock)
+						.Text(LOCTEXT("AddressSameButtonText", "Address same"))
+						.Font(FCoreStyle::GetDefaultFontStyle("Regular", 14))
+					]
 				]
 			]
 
-			+SHorizontalBox::Slot()
+			+ SVerticalBox::Slot()
 			.HAlign(HAlign_Fill)
 			.VAlign(VAlign_Top)
 			.Padding(8.f)
 			[
 				SNew(SButton)
-				.OnClicked(this, &SDMXPatchTool::OnAddressSameClicked)
+				.OnClicked(this, &SDMXPatchTool::OnAddressAndRenameClicked)
 				[
 					SNew(STextBlock)
-					.Text(LOCTEXT("AddressSameButtonText", "Address same"))
+					.Text(LOCTEXT("AddressAndRenameButtonText", "Address and Rename"))
 					.Font(FCoreStyle::GetDefaultFontStyle("Regular", 14))
 				]
 			]
-		]
-
-		+ SVerticalBox::Slot()
-		.HAlign(HAlign_Fill)
-		.VAlign(VAlign_Top)
-		.Padding(8.f)
-		[
-			SNew(SButton)
-			.OnClicked(this, &SDMXPatchTool::OnAddressAndRenameClicked)
-			[
-				SNew(STextBlock)
-				.Text(LOCTEXT("AddressAndRenameButtonText", "Address and Rename"))
-				.Font(FCoreStyle::GetDefaultFontStyle("Regular", 14))
-			]
-		]
-	];
-
-	// Bind to library changes
-	UDMXSubsystem* DMXSubsystem = UDMXSubsystem::GetDMXSubsystem_Pure();
-	check(DMXSubsystem);
-
-	DMXSubsystem->OnAllDMXLibraryAssetsLoaded.AddSP(this, &SDMXPatchTool::OnAllDMXLibraryAssetsLoaded);
-	DMXSubsystem->OnDMXLibraryAssetAdded.AddSP(this, &SDMXPatchTool::OnDMXLibraryAssetAdded);
-	DMXSubsystem->OnDMXLibraryAssetRemoved.AddSP(this, &SDMXPatchTool::OnDMXLibraryAssetRemoved);
-
-	// Make an initial selection
-	UpdateLibrarySelection();
-	UpdateFixturePatchSelection();
-}
-
-void SDMXPatchTool::AddReferencedObjects(FReferenceCollector& Collector)
-{
-	Collector.AddReferencedObjects(LibrarySource);
-	Collector.AddReferencedObjects(FixturePatchSource);
-}
-
-void SDMXPatchTool::UpdateLibrarySelection()
-{
-	check(LibraryComboBox.IsValid());
-	check(FixturePatchComboBox.IsValid());
+		];
 
-	if (LibrarySource.Num() > 0)
-	{
-		LibraryComboBox->SetSelectedItem(LibrarySource[0]);
+		// Make an initial selection
+		UpdateLibrarySelection();
+		UpdateFixturePatchSelection(DMXLibrary);
 	}
-	else
+
+	void SDMXPatchTool::AddReferencedObjects(FReferenceCollector& Collector)
 	{
-		SelectedLibraryTextBlock->SetText(LOCTEXT("NoLibraryAvailable", "No DMX library available"));
+		Collector.AddReferencedObject(DMXLibrary);
+		Collector.AddReferencedObjects(FixturePatchSource);
 	}
-}
 
-void SDMXPatchTool::UpdateFixturePatchSelection()
-{
-	UDMXLibrary* SelectedLibrary = LibraryComboBox->GetSelectedItem();
-	if (IsValid(SelectedLibrary))
+	void SDMXPatchTool::UpdateLibrarySelection()
 	{
-		FixturePatchSource = ObjectPtrWrap(SelectedLibrary->GetEntitiesTypeCast<UDMXEntityFixturePatch>());
+		check(LibraryComboBox.IsValid());
+		check(FixturePatchComboBox.IsValid());
 
-		if (FixturePatchSource.Num() > 0)
+		if (LibrarySource.Num() > 0)
 		{
-			FixturePatchComboBox->SetSelectedItem(FixturePatchSource[0]);
+			LibraryComboBox->SetSelectedItem(LibrarySource[0]);
 		}
 		else
 		{
-			SelectedFixturePatchTextBlock->SetText(LOCTEXT("NoFixturePatchAvailable", "No Fixture Patch available in Library"));
+			SelectedLibraryTextBlock->SetText(LOCTEXT("NoLibraryAvailable", "No DMX library available"));
 		}
 	}
-}
 
-FReply SDMXPatchTool::OnAddressIncrementalClicked()
-{
-	UDMXEntityFixturePatch* SelectedFixturePatch = FixturePatchComboBox->GetSelectedItem();
-	if (SelectedFixturePatch)
+	void SDMXPatchTool::UpdateFixturePatchSelection(UDMXLibrary* InDMXLibrary)
 	{
-		int32 IndexOfPatch = FixturePatchSource.IndexOfByKey(SelectedFixturePatch);
-
-		for (FSelectionIterator It(GEditor->GetSelectedActorIterator()); It; ++It)
+		if (InDMXLibrary)
 		{
-			if (AActor* Actor = Cast<AActor>(*It))
-			{
-				for (UDMXComponent* Component : TInlineComponentArray<UDMXComponent*>(Actor))
-				{
-					if (FixturePatchSource.IsValidIndex(IndexOfPatch))
-					{
-						UDMXEntityFixturePatch* FixturePatch = FixturePatchSource[IndexOfPatch];
-						Component->SetFixturePatch(FixturePatch);
+			FixturePatchSource = ObjectPtrWrap(InDMXLibrary->GetEntitiesTypeCast<UDMXEntityFixturePatch>());
 
-						IndexOfPatch++;
-					}
-					else
-					{
-						return FReply::Handled();
-					}
-				}
+			if (FixturePatchSource.Num() > 0)
+			{
+				FixturePatchComboBox->SetSelectedItem(FixturePatchSource[0]);
 			}
+			else
+			{
+				SelectedFixturePatchTextBlock->SetText(LOCTEXT("NoFixturePatchAvailable", "No Fixture Patch available in Library"));
+			}
+		}
+		else
+		{
+			FixturePatchSource.Reset();
 		}
 	}
 
-	return FReply::Handled();
-}
-
-FReply SDMXPatchTool::OnAddressSameClicked()
-{
-	UDMXEntityFixturePatch* SelectedFixturePatch = FixturePatchComboBox->GetSelectedItem();
-	if (SelectedFixturePatch)
+	FReply SDMXPatchTool::OnAddressIncrementalClicked()
 	{
-		int32 IndexOfPatch = FixturePatchSource.IndexOfByKey(SelectedFixturePatch);
-
-		for (FSelectionIterator It(GEditor->GetSelectedActorIterator()); It; ++It)
+		UDMXEntityFixturePatch* SelectedFixturePatch = FixturePatchComboBox->GetSelectedItem();
+		if (SelectedFixturePatch)
 		{
-			if (AActor* Actor = Cast<AActor>(*It))
+			int32 IndexOfPatch = FixturePatchSource.IndexOfByKey(SelectedFixturePatch);
+
+			for (FSelectionIterator It(GEditor->GetSelectedActorIterator()); It; ++It)
 			{
-				for (UDMXComponent* Component : TInlineComponentArray<UDMXComponent*>(Actor))
+				if (AActor* Actor = Cast<AActor>(*It))
 				{
-					UDMXEntityFixturePatch* FixturePatch = FixturePatchSource[IndexOfPatch];
-					Component->SetFixturePatch(FixturePatch);
+					for (UDMXComponent* Component : TInlineComponentArray<UDMXComponent*>(Actor))
+					{
+						if (FixturePatchSource.IsValidIndex(IndexOfPatch))
+						{
+							UDMXEntityFixturePatch* FixturePatch = FixturePatchSource[IndexOfPatch];
+							Component->SetFixturePatch(FixturePatch);
+
+							IndexOfPatch++;
+						}
+						else
+						{
+							return FReply::Handled();
+						}
+					}
 				}
 			}
 		}
-	}
 
-	return FReply::Handled();
-}
+		return FReply::Handled();
+	}
 
-FReply SDMXPatchTool::OnAddressAndRenameClicked()
-{
-	UDMXEntityFixturePatch* SelectedFixturePatch = FixturePatchComboBox->GetSelectedItem();
-	if (SelectedFixturePatch)
+	FReply SDMXPatchTool::OnAddressSameClicked()
 	{
-		int32 IndexOfPatch = FixturePatchSource.IndexOfByKey(SelectedFixturePatch);
-
-		for (FSelectionIterator It(GEditor->GetSelectedActorIterator()); It; ++It)
+		UDMXEntityFixturePatch* SelectedFixturePatch = FixturePatchComboBox->GetSelectedItem();
+		if (SelectedFixturePatch)
 		{
-			if (AActor* Actor = Cast<AActor>(*It))
+			int32 IndexOfPatch = FixturePatchSource.IndexOfByKey(SelectedFixturePatch);
+
+			for (FSelectionIterator It(GEditor->GetSelectedActorIterator()); It; ++It)
 			{
-				for (UDMXComponent* Component : TInlineComponentArray<UDMXComponent*>(Actor))
+				if (AActor* Actor = Cast<AActor>(*It))
 				{
-					if (FixturePatchSource.IsValidIndex(IndexOfPatch))
+					for (UDMXComponent* Component : TInlineComponentArray<UDMXComponent*>(Actor))
 					{
 						UDMXEntityFixturePatch* FixturePatch = FixturePatchSource[IndexOfPatch];
 						Component->SetFixturePatch(FixturePatch);
-
-						// Rename
-						Actor->SetActorLabel(FixturePatch->Name);
-
-						IndexOfPatch++;
-					}
-					else
-					{
-						return FReply::Handled();
 					}
 				}
 			}
 		}
-	}
-
-	return FReply::Handled();
-}
 
-TSharedRef<SWidget> SDMXPatchTool::GenerateLibraryComboBoxEntry(UDMXLibrary* LibraryToAdd)
-{
-	FText LibraryName = FText::FromString(LibraryToAdd->GetName());
-
-	return 
-		SNew(STextBlock)
-		.Text(LibraryName);
-}
-
-void SDMXPatchTool::OnLibrarySelected(UDMXLibrary* SelectedLibrary, ESelectInfo::Type SelectInfo)
-{
-	// Ignore unchanged selections
-	if (PreviouslySelectedLibrary.Get() && PreviouslySelectedLibrary.Get() == SelectedLibrary)
-	{
-		return;
+		return FReply::Handled();
 	}
 
-	// Unbind from previously selected library changes
-	if (PreviouslySelectedLibrary.IsValid())
+	FReply SDMXPatchTool::OnAddressAndRenameClicked()
 	{
-		PreviouslySelectedLibrary->GetOnEntitiesAdded().RemoveAll(this);
-		PreviouslySelectedLibrary->GetOnEntitiesRemoved().RemoveAll(this);
+		UDMXEntityFixturePatch* SelectedFixturePatch = FixturePatchComboBox->GetSelectedItem();
+		if (SelectedFixturePatch)
+		{
+			int32 IndexOfPatch = FixturePatchSource.IndexOfByKey(SelectedFixturePatch);
+
+			for (FSelectionIterator It(GEditor->GetSelectedActorIterator()); It; ++It)
+			{
+				if (AActor* Actor = Cast<AActor>(*It))
+				{
+					for (UDMXComponent* Component : TInlineComponentArray<UDMXComponent*>(Actor))
+					{
+						if (FixturePatchSource.IsValidIndex(IndexOfPatch))
+						{
+							UDMXEntityFixturePatch* FixturePatch = FixturePatchSource[IndexOfPatch];
+							Component->SetFixturePatch(FixturePatch);
+
+							// Rename
+							Actor->SetActorLabel(FixturePatch->Name);
+
+							IndexOfPatch++;
+						}
+						else
+						{
+							return FReply::Handled();
+						}
+					}
+				}
+			}
+		}
+
+		return FReply::Handled();
 	}
-	PreviouslySelectedLibrary = SelectedLibrary;
 
-	if (IsValid(SelectedLibrary))
+	TSharedRef<SWidget> SDMXPatchTool::GenerateLibraryComboBoxEntry(TSharedPtr<FDMXPatchToolItem> ItemToAdd)
 	{
-		// Bind to library edits
-		SelectedLibrary->GetOnEntitiesAdded().AddSP(this, &SDMXPatchTool::OnEntitiesAddedOrRemoved);
-		SelectedLibrary->GetOnEntitiesRemoved().AddSP(this, &SDMXPatchTool::OnEntitiesAddedOrRemoved);
+		FText LibraryName = FText::FromString(ItemToAdd->SoftDMXLibrary.GetAssetName());
 
-		SelectedLibraryTextBlock->SetText(FText::FromString(SelectedLibrary->GetName()));
+		return
+			SNew(STextBlock)
+			.Text(LibraryName);
 	}
-	else
+
+	void SDMXPatchTool::OnLibrarySelected(TSharedPtr<FDMXPatchToolItem> SelectedItem, ESelectInfo::Type SelectInfo)
 	{
-		SelectedLibraryTextBlock->SetText(LOCTEXT("NoLibraryAvailable", "No DMX library available"));
-	}
+		// Ignore unchanged selections
+		if (PreviouslySelectedLibrary.IsValid() && PreviouslySelectedLibrary.Get() == SelectedItem->SoftDMXLibrary)
+		{
+			return;
+		}
 
-	UpdateFixturePatchSelection();
-}
+		// Unbind from previously selected library changes
+		if (PreviouslySelectedLibrary.IsValid())
+		{
+			PreviouslySelectedLibrary->GetOnEntitiesAdded().RemoveAll(this);
+			PreviouslySelectedLibrary->GetOnEntitiesRemoved().RemoveAll(this);
+		}
+		PreviouslySelectedLibrary = DMXLibrary;
 
-TSharedRef<SWidget> SDMXPatchTool::GenerateFixturePatchComboBoxEntry(UDMXEntityFixturePatch* FixturePatchToAdd)
-{
-	FText FixturePatchName = FText::FromString(FixturePatchToAdd->Name);
+		DMXLibrary = Cast<UDMXLibrary>(SelectedItem->SoftDMXLibrary.LoadSynchronous());
+		if (DMXLibrary)
+		{
+			// Bind to library edits
+			DMXLibrary->GetOnEntitiesAdded().AddSP(this, &SDMXPatchTool::OnEntitiesAddedOrRemoved);
+			DMXLibrary->GetOnEntitiesRemoved().AddSP(this, &SDMXPatchTool::OnEntitiesAddedOrRemoved);
 
-	return
-		SNew(STextBlock)
-		.Text(FixturePatchName);
-}
+			SelectedLibraryTextBlock->SetText(FText::FromString(DMXLibrary->GetName()));
+		}
+		else
+		{
+			SelectedLibraryTextBlock->SetText(LOCTEXT("NoLibraryAvailable", "No DMX Library available"));
+		}
 
-void SDMXPatchTool::OnFixturePatchSelected(UDMXEntityFixturePatch* SelectedFixturePatch, ESelectInfo::Type SelectInfo)
-{
-	if (IsValid(SelectedFixturePatch))
-	{
-		SelectedFixturePatchTextBlock->SetText(FText::FromString(SelectedFixturePatch->Name));
-	}
-	else
-	{
-		SelectedFixturePatchTextBlock->SetText(LOCTEXT("NoFixturePatchAvailable", "No Fixture Patch available in Library"));
+		UpdateFixturePatchSelection(DMXLibrary);
 	}
-}
 
-void SDMXPatchTool::OnEntitiesAddedOrRemoved(UDMXLibrary* Library, TArray<UDMXEntity*> Entities)
-{
-	check(FixturePatchComboBox.IsValid());
+	TSharedRef<SWidget> SDMXPatchTool::GenerateFixturePatchComboBoxEntry(UDMXEntityFixturePatch* FixturePatchToAdd)
+	{
+		FText FixturePatchName = FText::FromString(FixturePatchToAdd->Name);
 
-	UDMXEntityFixturePatch* PreviouslySelectedFixturePatch = FixturePatchComboBox->GetSelectedItem();
+		return
+			SNew(STextBlock)
+			.Text(FixturePatchName);
+	}
 
-	if (Library == PreviouslySelectedLibrary.Get())
+	void SDMXPatchTool::OnFixturePatchSelected(UDMXEntityFixturePatch* SelectedFixturePatch, ESelectInfo::Type SelectInfo)
 	{
-		FixturePatchSource.Reset();
-
-		if (IsValid(Library))
+		if (IsValid(SelectedFixturePatch))
 		{
-			FixturePatchSource = ObjectPtrWrap(Library->GetEntitiesTypeCast<UDMXEntityFixturePatch>());
-
-			FixturePatchComboBox->RefreshOptions();
-
-			if (IsValid(PreviouslySelectedFixturePatch) && FixturePatchSource.Contains(PreviouslySelectedFixturePatch))
-			{
-				// Restore the previous selection
-				FixturePatchComboBox->SetSelectedItem(PreviouslySelectedFixturePatch);
-			}
-			else
-			{
-				// Update the selection
-				UpdateFixturePatchSelection();
-			}
+			SelectedFixturePatchTextBlock->SetText(FText::FromString(SelectedFixturePatch->Name));
+		}
+		else
+		{
+			SelectedFixturePatchTextBlock->SetText(LOCTEXT("NoFixturePatchAvailable", "No Fixture Patch available in Library"));
 		}
 	}
-}
-
-void SDMXPatchTool::OnAllDMXLibraryAssetsLoaded()
-{
-	// Rebuild all when the libraries are fully loaded (in case this allready is visible) 
-
-	UDMXSubsystem* Subsystem = UDMXSubsystem::GetDMXSubsystem_Pure();
-	check(Subsystem);
-
-	LibrarySource = ObjectPtrWrap(Subsystem->GetAllDMXLibraries());
-
-	UpdateLibrarySelection();
-	UpdateFixturePatchSelection();
-}
-
-void SDMXPatchTool::OnDMXLibraryAssetAdded(UDMXLibrary* DMXLibrary)
-{
-	// Rebuild all when a library was added
-
-	UDMXSubsystem* Subsystem = UDMXSubsystem::GetDMXSubsystem_Pure();
-	check(Subsystem);
 
-	LibrarySource = ObjectPtrWrap(Subsystem->GetAllDMXLibraries());
+	void SDMXPatchTool::OnEntitiesAddedOrRemoved(UDMXLibrary* ChangedDMXLibrary, TArray<UDMXEntity*> Entities)
+	{
+		UpdateLibrarySelection();
 
-	UpdateLibrarySelection();
-	UpdateFixturePatchSelection();
-}
+		UDMXEntityFixturePatch* PreviouslySelectedFixturePatch = FixturePatchComboBox->GetSelectedItem();
 
-void SDMXPatchTool::OnDMXLibraryAssetRemoved(UDMXLibrary* DMXLibrary)
-{
-	// Rebuild all when a library was removed
+		if (ChangedDMXLibrary == DMXLibrary)
+		{
+			FixturePatchSource.Reset();
 
-	UDMXSubsystem* Subsystem = UDMXSubsystem::GetDMXSubsystem_Pure();
-	check(Subsystem);
+			if (IsValid(DMXLibrary))
+			{
+				FixturePatchSource = ObjectPtrWrap(DMXLibrary->GetEntitiesTypeCast<UDMXEntityFixturePatch>());
 
-	LibrarySource = ObjectPtrWrap(Subsystem->GetAllDMXLibraries());
+				FixturePatchComboBox->RefreshOptions();
 
-	UpdateLibrarySelection();
-	UpdateFixturePatchSelection();
+				if (IsValid(PreviouslySelectedFixturePatch) && FixturePatchSource.Contains(PreviouslySelectedFixturePatch))
+				{
+					// Restore the previous selection
+					FixturePatchComboBox->SetSelectedItem(PreviouslySelectedFixturePatch);
+				}
+				else
+				{
+					// Update the selection
+					UpdateFixturePatchSelection(DMXLibrary);
+				}
+			}
+		}
+	}
 }
 
 #undef LOCTEXT_NAMESPACE
diff --git a/Engine/Plugins/VirtualProduction/DMX/DMXEngine/Source/DMXEditor/Private/Widgets/PatchTool/SDMXPatchTool.h b/Engine/Plugins/VirtualProduction/DMX/DMXEngine/Source/DMXEditor/Private/Widgets/PatchTool/SDMXPatchTool.h
--- a/Engine/Plugins/VirtualProduction/DMX/DMXEngine/Source/DMXEditor/Private/Widgets/PatchTool/SDMXPatchTool.h
+++ b/Engine/Plugins/VirtualProduction/DMX/DMXEngine/Source/DMXEditor/Private/Widgets/PatchTool/SDMXPatchTool.h
@@ -3,110 +3,105 @@
 #pragma once
 
 #include "Analytics/DMXEditorToolAnalyticsProvider.h"
-#include "CoreMinimal.h"
 #include "UObject/GCObject.h"
 #include "Widgets/DeclarativeSyntaxSupport.h"
 #include "Widgets/SCompoundWidget.h"
 
+template<typename OptionType> class SComboBox;
+class STextBlock;
 class UDMXEntity;
 class UDMXEntityFixturePatch;
 class UDMXLibrary;
 
-template<typename OptionType> class SComboBox;
-class STextBlock;
-
-
-/**
- * A Monitor for DMX activity in a range of DMX Universes
- */
-class SDMXPatchTool
-	: public SCompoundWidget
-	, public FGCObject
+namespace UE::DMX
 {
-public:
-	SLATE_BEGIN_ARGS(SDMXPatchTool)
-	{}
-
-	SLATE_END_ARGS()
-
-	SDMXPatchTool();
-	virtual ~SDMXPatchTool();
+	class FDMXPatchToolItem;
+
+	/**
+	 * A Monitor for DMX activity in a range of DMX Universes
+	 */
+	class SDMXPatchTool
+		: public SCompoundWidget
+		, public FGCObject
+	{
+	public:
+		SLATE_BEGIN_ARGS(SDMXPatchTool)
+			{}
 
-	/** Constructs the widget */
-	void Construct(const FArguments& InArgs);
+		SLATE_END_ARGS()
 
-protected:
-	// ~Begin FGCObject interface
-	virtual void AddReferencedObjects(FReferenceCollector& Collector) override;
-	virtual FString GetReferencerName() const override
-	{
-		return TEXT("SDMXPatchTool");
-	}
-	// ~End FGCObject interface
+		SDMXPatchTool();
+		virtual ~SDMXPatchTool();
 
-	/** Updates the selected library. Useful on initialization or when assets changed */
-	void UpdateLibrarySelection();
+		/** Constructs the widget */
+		void Construct(const FArguments& InArgs);
 
-	/** Updates the selected fixture patch. Useful when the library changed */
-	void UpdateFixturePatchSelection();
+	private:
+		// ~Begin FGCObject interface
+		virtual void AddReferencedObjects(FReferenceCollector& Collector) override;
+		virtual FString GetReferencerName() const override
+		{
+			return TEXT("SDMXPatchTool");
+		}
+		// ~End FGCObject interface
 
-	/** Called when the Address Incremental Button was clicked */
-	FReply OnAddressIncrementalClicked();
+		/** Updates the selected library. Useful on initialization or when assets changed */
+		void UpdateLibrarySelection();
 
-	/** Called when the Address Same Button was clicked */
-	FReply OnAddressSameClicked();
+		/** Updates the available fixture patches from the provided DMX Library */
+		void UpdateFixturePatchSelection(UDMXLibrary* InDMXLibrary);
 
-	/** Called when the Address And Rename Button was clicked */
-	FReply OnAddressAndRenameClicked();
+		/** Called when the Address Incremental Button was clicked */
+		FReply OnAddressIncrementalClicked();
 
-private:
-	/** Generates an entry in the library combo box */
-	TSharedRef<SWidget> GenerateLibraryComboBoxEntry(UDMXLibrary* LibraryToAdd);
+		/** Called when the Address Same Button was clicked */
+		FReply OnAddressSameClicked();
 
-	/** Called when a dmx library was slected */
-	void OnLibrarySelected(UDMXLibrary* SelectedLibrary, ESelectInfo::Type SelectInfo);
+		/** Called when the Address And Rename Button was clicked */
+		FReply OnAddressAndRenameClicked();
 
-	/** Combobox to select a library */
-	TSharedPtr<SComboBox<UDMXLibrary*>> LibraryComboBox;
+		/** Generates an entry in the library combo box */
+		TSharedRef<SWidget> GenerateLibraryComboBoxEntry(TSharedPtr<FDMXPatchToolItem> LibraryToAdd);
 
-	/** Text block showing the selected library */
-	TSharedPtr<STextBlock> SelectedLibraryTextBlock;
+		/** Called when a dmx library was slected */
+		void OnLibrarySelected(TSharedPtr<FDMXPatchToolItem> SelectedLibrary, ESelectInfo::Type SelectInfo);
 
-	/** Source for the library combo box */
-	TArray<TObjectPtr<UDMXLibrary>> LibrarySource;
+		/** Combobox to select a library */
+		TSharedPtr<SComboBox<TSharedPtr<FDMXPatchToolItem>>> LibraryComboBox;
 
-private:
-	/** Generates an entry in the library combo box */
-	TSharedRef<SWidget> GenerateFixturePatchComboBoxEntry(UDMXEntityFixturePatch* FixturePatchToAdd);
+		/** Text block showing the selected library */
+		TSharedPtr<STextBlock> SelectedLibraryTextBlock;
 
-	/** Called when a fixture patch was slected */
-	void OnFixturePatchSelected(UDMXEntityFixturePatch* SelectedFixturePatch, ESelectInfo::Type SelectInfo);
+		/** Source for the library combo box */
+		TArray<TSharedPtr<FDMXPatchToolItem>> LibrarySource;
 
-	/** Combobox to select a patch within the library */
-	TSharedPtr<SComboBox<UDMXEntityFixturePatch*>> FixturePatchComboBox;
+	private:
+		/** Generates an entry in the library combo box */
+		TSharedRef<SWidget> GenerateFixturePatchComboBoxEntry(UDMXEntityFixturePatch* FixturePatchToAdd);
 
-	/** Text block showing the selected fixture patch */
-	TSharedPtr<STextBlock> SelectedFixturePatchTextBlock;
+		/** Called when a fixture patch was slected */
+		void OnFixturePatchSelected(UDMXEntityFixturePatch* SelectedFixturePatch, ESelectInfo::Type SelectInfo);
 
-	/** Source for the fixture patch combo box */
-	TArray<TObjectPtr<UDMXEntityFixturePatch>> FixturePatchSource;
+		/** Combobox to select a patch within the library */
+		TSharedPtr<SComboBox<UDMXEntityFixturePatch*>> FixturePatchComboBox;
 
-private:
-	/** Called when the library was edited */
-	void OnEntitiesAddedOrRemoved(UDMXLibrary* Library, TArray<UDMXEntity*> Entities);
+		/** Text block showing the selected fixture patch */
+		TSharedPtr<STextBlock> SelectedFixturePatchTextBlock;
 
-	/** Called when the asset registry finished loading files */
-	void OnAllDMXLibraryAssetsLoaded();
+		/** Source for the fixture patch combo box */
+		TArray<TObjectPtr<UDMXEntityFixturePatch>> FixturePatchSource;
 
-	/** Called when a dmx library asset was added */
-	void OnDMXLibraryAssetAdded(UDMXLibrary* DMXLibrary);
+	private:
+		/** Called when the library was edited */
+		void OnEntitiesAddedOrRemoved(UDMXLibrary* Library, TArray<UDMXEntity*> Entities);
 
-	/** Called when a dmx library asset was removed */
-	void OnDMXLibraryAssetRemoved(UDMXLibrary* DMXLibrary);
+		/** The currently selected DMX Library. Useful to GC */
+		TObjectPtr<UDMXLibrary> DMXLibrary;
 
-	/** The previously selected library, to unbind from library changes */
-	TWeakObjectPtr<UDMXLibrary> PreviouslySelectedLibrary;
+		/** The previously selected library, to unbind from library changes */
+		TWeakObjectPtr<UDMXLibrary> PreviouslySelectedLibrary;
 
-	/** The analytics provider for this tool */
-	UE::DMX::FDMXEditorToolAnalyticsProvider AnalyticsProvider;
-};
+		/** The analytics provider for this tool */
+		FDMXEditorToolAnalyticsProvider AnalyticsProvider;
+	};
+}
diff --git a/Engine/Plugins/VirtualProduction/DMX/DMXEngine/Source/DMXRuntime/Private/DMXSubsystem.cpp b/Engine/Plugins/VirtualProduction/DMX/DMXEngine/Source/DMXRuntime/Private/DMXSubsystem.cpp
--- a/Engine/Plugins/VirtualProduction/DMX/DMXEngine/Source/DMXRuntime/Private/DMXSubsystem.cpp
+++ b/Engine/Plugins/VirtualProduction/DMX/DMXEngine/Source/DMXRuntime/Private/DMXSubsystem.cpp
@@ -2,7 +2,6 @@
 
 #include "DMXSubsystem.h"
 
-
 #include "AssetRegistry/AssetData.h"
 #include "Async/Async.h"
 #include "DMXAttribute.h"
@@ -16,16 +15,16 @@
 #include "EngineAnalytics.h"
 #include "EngineUtils.h"
 #include "Interfaces/IDMXProtocol.h"
-#include "IO/DMXTrace.h"
-#include "IO/DMXPortManager.h"
 #include "IO/DMXInputPort.h"
 #include "IO/DMXOutputPort.h"
-#include "Library/DMXLibrary.h"
-#include "Library/DMXEntityReference.h"
+#include "IO/DMXPortManager.h"
+#include "IO/DMXTrace.h"
 #include "Library/DMXEntity.h"
 #include "Library/DMXEntityController.h"
-#include "Library/DMXEntityFixtureType.h"
 #include "Library/DMXEntityFixturePatch.h"
+#include "Library/DMXEntityFixtureType.h"
+#include "Library/DMXEntityReference.h"
+#include "Library/DMXLibrary.h"
 #include "UObject/UObjectIterator.h"
 
 #if WITH_EDITOR
@@ -120,7 +119,6 @@ namespace
 #endif // WITH_EDITOR
 }
 
-
 void UDMXSubsystem::SendDMX(UDMXEntityFixturePatch* FixturePatch, TMap<FDMXAttributeName, int32> AttributeMap, EDMXSendResult& OutResult)
 {
 	OutResult = EDMXSendResult::Success;
@@ -658,9 +656,51 @@ UDMXEntityController* UDMXSubsystem::GetControllerByName(const UDMXLibrary* DMXL
 	return nullptr;
 }
 
-const TArray<UDMXLibrary*>& UDMXSubsystem::GetAllDMXLibraries()
+TArray<UDMXLibrary*> UDMXSubsystem::GetAllDMXLibraries()
+{
+	// DEPRECATED 5.5
+	return LoadDMXLibrariesSynchronous();
+}
+
+TArray<UDMXLibrary*> UDMXSubsystem::LoadDMXLibrariesSynchronous() const
 {
-	return LoadedDMXLibraries;
+	TArray<TSoftObjectPtr<UDMXLibrary>> SoftDMXLibraries = GetDMXLibraries();
+	TArray<UDMXLibrary*> DMXLibraries;
+	Algo::TransformIf(SoftDMXLibraries, DMXLibraries,
+		[](const TSoftObjectPtr<UDMXLibrary>& DMXLibrary)
+		{
+			return DMXLibrary.IsValid();
+		},
+		[](const TSoftObjectPtr<UDMXLibrary>& DMXLibrary)
+		{
+			return DMXLibrary.LoadSynchronous();
+		});
+
+	return DMXLibraries;
+}
+
+TArray<TSoftObjectPtr<UDMXLibrary>> UDMXSubsystem::GetDMXLibraries() const
+{
+	constexpr bool bHasBlueprintClasses = true;
+	UObjectLibrary* LibraryOfDMXLibraries = UObjectLibrary::CreateLibrary(UDMXLibrary::StaticClass(), bHasBlueprintClasses, GIsEditor);
+	LibraryOfDMXLibraries->LoadAssetDataFromPath(TEXT("/Game"));
+	LibraryOfDMXLibraries->LoadAssetsFromAssetData();
+
+	TArray<FAssetData> AssetDatas;
+	LibraryOfDMXLibraries->GetAssetDataList(AssetDatas);
+
+	TArray<TSoftObjectPtr<UDMXLibrary>> DMXLibraries;
+	Algo::TransformIf(AssetDatas, DMXLibraries,
+		[](const FAssetData& AssetData)
+		{
+			return AssetData.IsValid();
+		},
+		[](const FAssetData& AssetData)
+		{
+			return TSoftObjectPtr<UDMXLibrary>(AssetData.ToSoftObjectPath());
+		});
+
+	return DMXLibraries;
 }
 
 FORCEINLINE EDMXFixtureSignalFormat SignalFormatFromBytesNum(uint32 InBytesNum)
@@ -759,67 +799,3 @@ float UDMXSubsystem::GetNormalizedAttributeValue(UDMXEntityFixturePatch* InFixtu
 
 	return -1.0f;
 }
-
-void UDMXSubsystem::Initialize(FSubsystemCollectionBase& Collection)
-{
-	// Load all available dmx libraries
-	constexpr bool bHasBlueprintClasses = true;
-	UObjectLibrary* LibraryOfDMXLibraries = UObjectLibrary::CreateLibrary(UDMXLibrary::StaticClass(), bHasBlueprintClasses, GIsEditor);
-	LibraryOfDMXLibraries->LoadAssetDataFromPath(TEXT("/Game"));
-	LibraryOfDMXLibraries->LoadAssetsFromAssetData();
-
-	TArray<FAssetData> AssetDatas;
-	LibraryOfDMXLibraries->GetAssetDataList(AssetDatas);
-
-	for (const FAssetData& AssetData : AssetDatas)
-	{
-		UDMXLibrary* DMXLibrary = Cast<UDMXLibrary>(AssetData.ToSoftObjectPath().TryLoad());
-
-		if (DMXLibrary)
-		{
-			LoadedDMXLibraries.Add(DMXLibrary);
-		}
-		else
-		{
-			UE_LOG(DMXSubsystemLog, Warning, TEXT("Failed to load DMXLibrary %s. See previous errors for causes."), *AssetData.AssetName.ToString());
-		}
-	}
-	OnAllDMXLibraryAssetsLoaded.Broadcast();
-
-#if WITH_EDITOR
-	// Handle adding/removing new libraries
-	IAssetRegistry& AssetRegistry = FModuleManager::LoadModuleChecked<FAssetRegistryModule>(FName("AssetRegistry")).Get();
-	AssetRegistry.OnAssetAdded().AddUObject(this, &UDMXSubsystem::OnAssetRegistryAddedAsset);
-	AssetRegistry.OnAssetRemoved().AddUObject(this, &UDMXSubsystem::OnAssetRegistryRemovedAsset);
-#endif
-}
-
-#if WITH_EDITOR
-void UDMXSubsystem::OnAssetRegistryAddedAsset(const FAssetData& Asset)
-{
-	if (Asset.AssetClassPath == UDMXLibrary::StaticClass()->GetClassPathName())
-	{
-		UObject* AssetObject = Asset.GetAsset();
-		if (UDMXLibrary* Library = Cast<UDMXLibrary>(AssetObject))
-		{
-			LoadedDMXLibraries.AddUnique(Library);
-			OnDMXLibraryAssetAdded.Broadcast(Library);
-		}
-	}
-}
-#endif // WITH_EDITOR
-
-#if WITH_EDITOR
-void UDMXSubsystem::OnAssetRegistryRemovedAsset(const FAssetData& Asset)
-{
-	if (Asset.AssetClassPath == UDMXLibrary::StaticClass()->GetClassPathName())
-	{
-		UObject* AssetObject = Asset.GetAsset();
-		if (UDMXLibrary* Library = Cast<UDMXLibrary>(AssetObject))
-		{
-			LoadedDMXLibraries.Remove(Library);
-			OnDMXLibraryAssetRemoved.Broadcast(Library);
-		}
-	}
-}
-#endif // WITH_EDITOR
diff --git a/Engine/Plugins/VirtualProduction/DMX/DMXEngine/Source/DMXRuntime/Private/Sequencer/MovieSceneDMXLibrarySection.cpp b/Engine/Plugins/VirtualProduction/DMX/DMXEngine/Source/DMXRuntime/Private/Sequencer/MovieSceneDMXLibrarySection.cpp
--- a/Engine/Plugins/VirtualProduction/DMX/DMXEngine/Source/DMXRuntime/Private/Sequencer/MovieSceneDMXLibrarySection.cpp
+++ b/Engine/Plugins/VirtualProduction/DMX/DMXEngine/Source/DMXRuntime/Private/Sequencer/MovieSceneDMXLibrarySection.cpp
@@ -2,24 +2,23 @@
 
 #include "Sequencer/MovieSceneDMXLibrarySection.h"
 
+#include "Channels/MovieSceneChannelEditorData.h"
+#include "Channels/MovieSceneChannelProxy.h"
 #include "DMXConversions.h"
 #include "DMXProtocolCommon.h"
 #include "DMXRuntimeLog.h"
 #include "DMXRuntimeMainStreamObjectVersion.h"
 #include "DMXSubsystem.h"
+#include "Evaluation/Blending/MovieSceneBlendType.h"
 #include "Interfaces/IDMXProtocol.h"
 #include "IO/DMXInputPort.h"
 #include "IO/DMXOutputPort.h"
 #include "IO/DMXPortManager.h"
 #include "IO/DMXTrace.h"
-#include "Library/DMXLibrary.h"
 #include "Library/DMXEntityFixturePatch.h"
 #include "Library/DMXEntityFixtureType.h"
-
-#include "Channels/MovieSceneChannelProxy.h"
-#include "Channels/MovieSceneChannelEditorData.h"
-#include "Evaluation/Blending/MovieSceneBlendType.h"
-
+#include "Library/DMXLibrary.h"
+#include "UObject/Package.h"
 
 DECLARE_LOG_CATEGORY_CLASS(MovieSceneDMXLibrarySectionLog, Log, All);
 
@@ -159,7 +158,7 @@ FDMXCachedFunctionChannelInfo::FDMXCachedFunctionChannelInfo(const TArray<FDMXFi
 	
 	// Valid patch
 	UDMXEntityFixturePatch* FixturePatch = FixturePatchChannel.Reference.GetFixturePatch();
-	if (FixturePatch == nullptr || !FixturePatch->IsValidLowLevelFast())
+	if (!FixturePatch || !FixturePatch->IsValidLowLevelFast())
 	{
 		UE_LOG(MovieSceneDMXLibrarySectionLog, Error, TEXT("%hs: A Fixture Patch is null."), __FUNCTION__);
 		return;
@@ -173,15 +172,15 @@ FDMXCachedFunctionChannelInfo::FDMXCachedFunctionChannelInfo(const TArray<FDMXFi
 
 	// Try to access the active mode
 	const UDMXEntityFixtureType* FixtureType = FixturePatch->GetFixtureType();
-	if (FixtureType == nullptr || !FixtureType->IsValidLowLevelFast())
+	if (!FixtureType || !FixtureType->IsValidLowLevelFast())
 	{
-		UE_LOG(MovieSceneDMXLibrarySectionLog, Error, TEXT("%S: Patch %s has invalid Fixture Type template."), __FUNCTION__, *FixturePatch->GetDisplayName());
+		UE_LOG(MovieSceneDMXLibrarySectionLog, Error, TEXT("%hs: Patch %s has invalid Fixture Type template."), __FUNCTION__, *FixturePatch->GetDisplayName());
 		return;
 	}
 
 	if (FixturePatchChannel.ActiveMode >= FixtureType->Modes.Num())
 	{
-		UE_LOG(MovieSceneDMXLibrarySectionLog, Error, TEXT("%S: Patch track %s ActiveMode is invalid."), __FUNCTION__, *FixturePatch->GetDisplayName());
+		UE_LOG(MovieSceneDMXLibrarySectionLog, Error, TEXT("%hs: Patch track %s ActiveMode is invalid."), __FUNCTION__, *FixturePatch->GetDisplayName());
 		return;
 	}
 
@@ -203,8 +202,8 @@ FDMXCachedFunctionChannelInfo::FDMXCachedFunctionChannelInfo(const TArray<FDMXFi
 		const bool bMissingFunction = !AttributeNameChannelMap.Contains(FunctionChannel.AttributeName) || !CellAttributePtr;
 		if (!CellAttributePtr || bMissingFunction)
 		{
-			UE_LOG(MovieSceneDMXLibrarySectionLog, Warning, TEXT("%S: Function with attribute %s from %s doesn't have a counterpart Fixture Function."), __FUNCTION__, *FunctionChannel.AttributeName.ToString(), *FixturePatch->GetDisplayName());
-			UE_LOG(MovieSceneDMXLibrarySectionLog, Warning, TEXT("%S: Further attributes may be missing. Warnings ommited to avoid overflowing the log."), __FUNCTION__);
+			UE_LOG(MovieSceneDMXLibrarySectionLog, Warning, TEXT("%hs: Function with attribute %s from %s doesn't have a counterpart Fixture Function."), __FUNCTION__, *FunctionChannel.AttributeName.ToString(), *FixturePatch->GetDisplayName());
+			UE_LOG(MovieSceneDMXLibrarySectionLog, Warning, TEXT("%hs: Further attributes may be missing. Warnings ommited to avoid overflowing the log."), __FUNCTION__);
 
 			return;
 		}
@@ -226,7 +225,7 @@ FDMXCachedFunctionChannelInfo::FDMXCachedFunctionChannelInfo(const TArray<FDMXFi
 
 		if (!FunctionPtr)
 		{
-			UE_LOG(MovieSceneDMXLibrarySectionLog, Warning, TEXT("%S: Function with attribute %s from %s doesn't have a counterpart Fixture Function."), __FUNCTION__, *FunctionChannel.AttributeName.ToString(), *FixturePatch->GetDisplayName());
+			UE_LOG(MovieSceneDMXLibrarySectionLog, Warning, TEXT("%hs: Function with attribute %s from %s doesn't have a counterpart Fixture Function."), __FUNCTION__, *FunctionChannel.AttributeName.ToString(), *FixturePatch->GetDisplayName());
 
 			return;
 		}
@@ -611,6 +610,8 @@ FDMXFixturePatchChannel* UMovieSceneDMXLibrarySection::GetPatchChannel(UDMXEntit
 
 void UMovieSceneDMXLibrarySection::RebuildPlaybackCache() const
 {
+	
+
 	CachedOutputPorts.Reset();
 	CachedChannelsToEvaluate.Reset();
 	CachedChannelsToInitialize.Reset();
@@ -672,7 +673,6 @@ void UMovieSceneDMXLibrarySection::EvaluateAndSendDMX(const FFrameTime& FrameTim
 			if (FixtureFunctionChannelPtr->Channel.Evaluate(FrameTime, ChannelValue))
 			{
 				TArray<uint8> ByteArr;
-
 				if (bUseNormalizedValues)
 				{
 					DMXSubsystem->NormalizedValueToBytes(ChannelValue, InfoForChannelToEvaluate.GetSignalFormat(), ByteArr, InfoForChannelToEvaluate.ShouldUseLSBMode());
@@ -757,23 +757,47 @@ void UMovieSceneDMXLibrarySection::SendDMXForChannelsToInitialize() const
 
 void UMovieSceneDMXLibrarySection::UpdateChannelProxy(bool bResetDefaultChannelValues /*= false*/)
 {
+	// Rebuild caches for the relevant DMX Libraries in this section. 
+	// By this the DMXFixture digital twins receives the initial DMX values regardless of loading order during engine startup.
+	UDMXSubsystem* Subsystem = UDMXSubsystem::GetDMXSubsystem_Callable();
+	TArray<TSoftObjectPtr<UDMXLibrary>> DMXLibraries = Subsystem->GetDMXLibraries();
+	TArray<TSoftObjectPtr<UDMXLibrary>> RelevantDMXLibraries;
+
+	Algo::TransformIf(DMXLibraries, RelevantDMXLibraries,
+		[this](const TSoftObjectPtr<UDMXLibrary>& DMXLibrary)
+		{
+			return DMXLibrary.IsValid() && Algo::FindBy(FixturePatchChannels, DMXLibrary, &FDMXFixturePatchChannel::DMXLibrary) != nullptr;
+		},
+		[](const TSoftObjectPtr<UDMXLibrary>& DMXLibrary)
+		{
+			return DMXLibrary;
+		});
+
+	for (const TSoftObjectPtr<UDMXLibrary>& DMXLibrary : RelevantDMXLibraries)
+	{
+		DMXLibrary.Get()->ForEachEntityOfType<UDMXEntityFixturePatch>([](UDMXEntityFixturePatch* Patch)
+			{
+				Patch->RebuildCache();
+			});
+	}
+
 	FMovieSceneChannelProxyData ChannelProxyData;
 	TArray<int32> InvalidPatchChannelIndices;
-
+	
 	int32 PatchChannelIndex = 0; // Safer because the ranged for ensures the array length isn't changed
 	for (FDMXFixturePatchChannel& PatchChannel : FixturePatchChannels)
 	{
 		if (!IsValid(PatchChannel.DMXLibrary))
 		{
-			UE_LOG(MovieSceneDMXLibrarySectionLog, Warning, TEXT("%S: Missing library for sequence section."), __FUNCTION__);
+			UE_LOG(MovieSceneDMXLibrarySectionLog, Warning, TEXT("%hs: Missing library for sequence section."), __FUNCTION__);
 		}
 
 		PatchChannel.UpdateNumberOfChannels(bResetDefaultChannelValues);
 
-		const UDMXEntityFixturePatch* Patch = PatchChannel.Reference.GetFixturePatch();
-		if (Patch == nullptr || !Patch->IsValidLowLevelFast())
+		UDMXEntityFixturePatch* Patch = PatchChannel.Reference.GetFixturePatch();
+		if (!IsValid(Patch))
 		{
-			UE_LOG(MovieSceneDMXLibrarySectionLog, Warning, TEXT("%S: Ignoring null Patch. Presumably the library changed and patches were removed or changed. This is not supported."), __FUNCTION__);
+			UE_LOG(MovieSceneDMXLibrarySectionLog, Warning, TEXT("%hs: Ignoring null Patch. Presumably the library changed and patches were removed or changed. This is not supported."), __FUNCTION__);
 			InvalidPatchChannelIndices.Add(PatchChannelIndex);
 			continue;
 		}
@@ -786,7 +810,7 @@ void UMovieSceneDMXLibrarySection::UpdateChannelProxy(bool bResetDefaultChannelV
 			// This will give users the impression that the Patch isn't added, but it is,
 			// which prevents the user from adding it again. So, to mitigate that, we remove
 			// the Patch from the track section.
-			UE_LOG(MovieSceneDMXLibrarySectionLog, Warning, TEXT("%S: Ignoring patch without functions %s"), __FUNCTION__, *Patch->GetDisplayName());
+			UE_LOG(MovieSceneDMXLibrarySectionLog, Warning, TEXT("%hs: Ignoring patch without functions %s"), __FUNCTION__, *Patch->GetDisplayName());
 			InvalidPatchChannelIndices.Add(PatchChannelIndex);
 			continue;
 		}
diff --git a/Engine/Plugins/VirtualProduction/DMX/DMXEngine/Source/DMXRuntime/Public/DMXSubsystem.h b/Engine/Plugins/VirtualProduction/DMX/DMXEngine/Source/DMXRuntime/Public/DMXSubsystem.h
--- a/Engine/Plugins/VirtualProduction/DMX/DMXEngine/Source/DMXRuntime/Public/DMXSubsystem.h
+++ b/Engine/Plugins/VirtualProduction/DMX/DMXEngine/Source/DMXRuntime/Public/DMXSubsystem.h
@@ -2,18 +2,17 @@
 
 #pragma once
 
-#include "CoreMinimal.h"
-#include "UObject/NoExportTypes.h"
 #include "DMXProtocolTypes.h"
 #include "DMXTypes.h"
 #include "IO/DMXInputPortReference.h"
 #include "IO/DMXOutputPortReference.h"
 #include "Library/DMXEntityReference.h"
-#include "Library/DMXEntityFixtureType.h"
 #include "Subsystems/EngineSubsystem.h"
+
 #include "DMXSubsystem.generated.h"
 
 class IDMXProtocol;
+class UDMXEntityFixtureType;
 class UDMXEntityFixturePatch;
 class UDMXLibrary;
 class UDMXModulator;
@@ -36,10 +35,12 @@ class DMXRUNTIME_API UDMXSubsystem
 public:
 	/**  Send DMX using function names and integer values. */
 	UFUNCTION(BlueprintCallable, Category = "DMX", meta = (DeprecatedFunction, DeprecationMessage = "Deprecated 4.27. Use DMXEntityFixurePatch::SendDMX instead"))
+	UE_DEPRECATED(5.5, "Use UDMXEntityFixurePatch::SendDMX instead.")
 	void SendDMX(UDMXEntityFixturePatch* FixturePatch, TMap<FDMXAttributeName, int32> AttributeMap, EDMXSendResult& OutResult);
 
 	/**  DEPRECATED 4.27 */
 	UFUNCTION(BlueprintCallable, Category = "DMX", meta = (DeprecatedFunction, DeprecationMessage = "Deprecated 4.27. Use SendDMXToOutputPort instead."))
+	UE_DEPRECATED(5.5, "Use UDMXSubsystem::SendDMXToOutputPort instead.")
 	void SendDMXRaw(FDMXProtocolName SelectedProtocol, int32 RemoteUniverse, TMap<int32, uint8> AddressValueMap, EDMXSendResult& OutResult);
 
 	/**  Sends DMX Values over the Output Port */
@@ -48,14 +49,15 @@ public:
 
 	/**  DEPRECATED 4.27 */
 	UFUNCTION(BlueprintCallable, Category = "DMX", meta = (DeprecatedFunction, DeprecationMessage = "Deprecated 4.27. Use GetDMXDataFromInputPort or GetDMXDataFromOutputPort instead."))
+	UE_DEPRECATED(5.5, "Use UDMXSubsystem::GetDMXDataFromInputPort or GetDMXDataFromOutputPort instead.")
 	void GetRawBuffer(FDMXProtocolName SelectedProtocol, int32 RemoteUniverse, TArray<uint8>& DMXBuffer);
 	
 	/**  Gets accumulated latest DMX Values from the Input Port (all that's been received since Begin Play) */
-	UFUNCTION(BlueprintCallable, Category = "DMX")
+	UFUNCTION(BlueprintCallable, Category = "DMX", Meta = (DisplayName = "Get DMX Data From Input Port"))
 	static void GetDMXDataFromInputPort(FDMXInputPortReference InputPortReference, TArray<uint8>& DMXData, int32 LocalUniverse = 1);
 	
 	/**  Gets accumulated latest DMX Values from the Output Port  (all that's been sent since Begin Play)*/
-	UFUNCTION(BlueprintCallable, Category = "DMX")
+	UFUNCTION(BlueprintCallable, Category = "DMX", Meta = (DisplayName = "Get DMX Data From Output Port"))
 	static void GetDMXDataFromOutputPort(FDMXOutputPortReference OutputPortReference, TArray<uint8>& DMXData, int32 LocalUniverse = 1);
 
 	/**  Return reference to array of Fixture Patch objects of a given type. */
@@ -96,19 +98,29 @@ public:
 
 	/**  DEPRECATED 4.27 */
 	UFUNCTION(BlueprintCallable, Category = "DMX", meta = (DeprecatedFunction, DeprecationMessage = "Deprecated 4.27. Controllers are removed in favor of Ports."))
+	UE_DEPRECATED(5.5, "Controllers are removed in favor of Ports.")
 	void GetAllUniversesInController(const UDMXLibrary* DMXLibrary, FString ControllerName, TArray<int32>& OutResult);
 
 	/**  DEPRECATED 4.27 */
 	UFUNCTION(BlueprintCallable, Category = "DMX", meta = (DeprecatedFunction, DeprecationMessage = "Deprecated 4.27. Controllers are removed in favor of Ports."))
+	UE_DEPRECATED(5.5, "Controllers are removed in favor of Ports.")
 	TArray<UDMXEntityController*> GetAllControllersInLibrary(const UDMXLibrary* DMXLibrary);
 
 	/**  DEPRECATED 4.27 */
 	UFUNCTION(BlueprintCallable, Category = "DMX", meta = (AutoCreateRefTerm = "Name"), meta = (DeprecatedFunction, DeprecationMessage = "Deprecated 4.27. Controllers are removed in favor of Ports."))
+	UE_DEPRECATED(5.5, "Controllers are removed in favor of Ports.")
 	UDMXEntityController* GetControllerByName(const UDMXLibrary* DMXLibrary, const FString& Name);
 
-	/**  Return reference to array of DMX Library objects. */
+	UE_DEPRECATED(5.5, "Renamed to UDMXSubsystem::LoadAllDMXLibrariesSynchronous. See also UDMXSubsystem::GetDMXLibraries to get soft object ptrs for all DMX Libraries.")
+	TArray<UDMXLibrary*> GetAllDMXLibraries();
+
+	/**  Loads all DMX Libraries in this project synchronous. */
+	UFUNCTION(BlueprintCallable, Category = "DMX", Meta = (DisplayName = "Load DMX Libraries Synchronous"))
+	TArray<UDMXLibrary*> LoadDMXLibrariesSynchronous() const;
+
+	/**  Returns a Soft Object Ptr to all DMX Library in the project. */
 	UFUNCTION(BlueprintCallable, Category = "DMX")
-	const TArray<UDMXLibrary*>& GetAllDMXLibraries();
+	TArray<TSoftObjectPtr<UDMXLibrary>> GetDMXLibraries() const;
 
 	/**
 	 * Return integer given an array of bytes. Up to the first 4 bytes in the array will be used for the conversion.
@@ -280,41 +292,18 @@ public:
 	UFUNCTION(BlueprintCallable, Category = "DMX")
 	void PixelMappingDistributionSort(EDMXPixelMappingDistribution InDistribution, int32 InNumXPanels, int32 InNumYPanels, const TArray<int32>& InUnorderedList, TArray<int32>& OutSortedList);
 
-public:
-	//~ USubsystem interface begin
-	virtual void Initialize(FSubsystemCollectionBase& Collection) override;
-	//~ USubsystem interface end
-
-private:
-	/**
-	 * Stores DelegateHandles for each Protocol's UniverseInputUpdate event.
-	 * That way we can unbind them when this subsystem is being destroyed and prevent crashes.
-	 */
-	TMap<FName, FDelegateHandle> UniverseInputBufferUpdatedHandles;
-
-public:
-	/** Delegate broadcast when all dmx library assets were loaded */
-	FSimpleMulticastDelegate OnAllDMXLibraryAssetsLoaded;
+	public:
+		/** Delegate broadcast when all dmx library assets were loaded */
+		UE_DEPRECATED(5.5, "DMX Libraries are no longer loaded by default and this delegate is no longer raised. Instead please use UDMXSubsystem::GetDMXLibraries or UDMXSubsystem::LoadDMXLibrariesSynchronous.")
+		FSimpleMulticastDelegate OnAllDMXLibraryAssetsLoaded;
 
 #if WITH_EDITOR
-	/** Delegate broadcast when a dmx library asset was added */
-	FDMXOnDMXLibraryAssetDelegate OnDMXLibraryAssetAdded;
-
-	/** Delegate broadcast when a dmx library asset was removed */
-	FDMXOnDMXLibraryAssetDelegate OnDMXLibraryAssetRemoved;
-
-private:
-	/** Called when asset registry added an asset */
-	UFUNCTION()
-	void OnAssetRegistryAddedAsset(const FAssetData& Asset);
-
-	/** Called when asset registry removed an asset */
-	UFUNCTION()
-	void OnAssetRegistryRemovedAsset(const FAssetData& Asset);
-#endif 
-
-private:
-	/** Strongly references all libraries at all times */
-	UPROPERTY(Transient)
-	TArray<TObjectPtr<UDMXLibrary>> LoadedDMXLibraries;
+		/** Delegate broadcast when a dmx library asset was added */
+		UE_DEPRECATED(5.5, "DMX Libraries are no longer loaded by default and this delegate is no longer raised. Instead please refer to the asset subsystem directly. See IAssetRegistry::OnAssetAdded.")
+		FDMXOnDMXLibraryAssetDelegate OnDMXLibraryAssetAdded;
+
+		/** Delegate broadcast when a dmx library asset was removed */
+		UE_DEPRECATED(5.5, "DMX Libraries are no longer loaded by default and this delegate is no longer raised. Instead please refer to the asset subsystem directly. See IAssetRegistry::OnAssetRemoved.")
+		FDMXOnDMXLibraryAssetDelegate OnDMXLibraryAssetRemoved;
+#endif // WITH_EDITOR
 };
diff --git a/Engine/Plugins/VirtualProduction/DMX/DMXEngine/Source/DMXRuntime/Public/Sequencer/MovieSceneDMXLibrarySection.h b/Engine/Plugins/VirtualProduction/DMX/DMXEngine/Source/DMXRuntime/Public/Sequencer/MovieSceneDMXLibrarySection.h
--- a/Engine/Plugins/VirtualProduction/DMX/DMXEngine/Source/DMXRuntime/Public/Sequencer/MovieSceneDMXLibrarySection.h
+++ b/Engine/Plugins/VirtualProduction/DMX/DMXEngine/Source/DMXRuntime/Public/Sequencer/MovieSceneDMXLibrarySection.h
@@ -2,7 +2,6 @@
 
 #pragma once
 
-#include "CoreMinimal.h"
 #include "DMXProtocolCommon.h"
 #include "MovieSceneSection.h"
 #include "Channels/MovieSceneFloatChannel.h"
@@ -13,20 +12,14 @@
 enum class EDMXFixtureSignalFormat : uint8;
 class UDMXEntityFixturePatch;
 
-
 USTRUCT()
 struct FDMXFixtureFunctionChannel
 {
 	GENERATED_BODY()
 
-	FDMXFixtureFunctionChannel()
-		: CellCoordinate(FIntPoint(-1, -1))
-		, DefaultValue(0)
-		, bEnabled(true)
-	{}
-
 	/** Attribute Name of the function */
-	FName AttributeName;
+	UPROPERTY()
+	FName AttributeName = NAME_None;
 
 	/** Function animation curve. */
 	UPROPERTY()
@@ -36,18 +29,19 @@ struct FDMXFixtureFunctionChannel
 	 * For Cell Functions the coordinate of the cell
 	 * For common functions -1.
 	 */
-	FIntPoint CellCoordinate;
+	UPROPERTY()
+	FIntPoint CellCoordinate = FIntPoint(-1, -1);
 
 	/** Default value to use when this Function is disabled in the track. */
 	UPROPERTY()
-	uint32 DefaultValue;
+	uint32 DefaultValue = 0;
 
 	/**
 	 * Whether or not to display this Function in the Patch's group
 	 * If false, the Function's default value is sent to DMX protocols.
 	 */
 	UPROPERTY()
-	bool bEnabled;
+	bool bEnabled = true;
 
 	/** True if the function is a cell function */
 	bool IsCellFunction() const { return CellCoordinate.X != -1 && CellCoordinate.Y != -1; }
@@ -58,11 +52,6 @@ struct FDMXFixturePatchChannel
 {
 	GENERATED_BODY()
 
-	FDMXFixturePatchChannel()
-		: DMXLibrary(nullptr)
-		, ActiveMode(INDEX_NONE)
-	{}
-
 	/** The outer library of the channel */
 	UPROPERTY()
 	TObjectPtr<UDMXLibrary> DMXLibrary;
@@ -80,7 +69,7 @@ struct FDMXFixturePatchChannel
 	 * simply by the user changing the active mode in the DMX Library.
 	 */
 	UPROPERTY()
-	int32 ActiveMode;
+	int32 ActiveMode = INDEX_NONE;
 
 	void SetFixturePatch(UDMXEntityFixturePatch* InPatch);
 
