DMX: Fixture Types created in engine now convert to GDTF when a DMX Library is exported as MVR

diff --git a/Engine/Plugins/VirtualProduction/DMX/DMXEngine/Source/DMXEditor/Private/Customizations/DMXEntityFixtureTypeDetails.cpp b/Engine/Plugins/VirtualProduction/DMX/DMXEngine/Source/DMXEditor/Private/Customizations/DMXEntityFixtureTypeDetails.cpp
--- a/Engine/Plugins/VirtualProduction/DMX/DMXEngine/Source/DMXEditor/Private/Customizations/DMXEntityFixtureTypeDetails.cpp
+++ b/Engine/Plugins/VirtualProduction/DMX/DMXEngine/Source/DMXEditor/Private/Customizations/DMXEntityFixtureTypeDetails.cpp
@@ -2,14 +2,12 @@
 
 #include "Customizations/DMXEntityFixtureTypeDetails.h"
 
-#include "DMXEditorLog.h"
-#include "DMXInitializeFixtureTypeFromGDTFHelper.h"
-#include "Library/DMXImportGDTF.h"
-#include "Library/DMXEntityFixtureType.h"
-
-#include "CoreMinimal.h"
 #include "DetailLayoutBuilder.h"
+#include "DMXEditorLog.h"
+#include "Factories/DMXGDTFToFixtureTypeConverter.h"
 #include "IPropertyUtilities.h"
+#include "Library/DMXEntityFixtureType.h"
+#include "Library/DMXImportGDTF.h"
 #include "PropertyHandle.h"
 
 
@@ -32,6 +30,8 @@ void FDMXEntityFixtureTypeDetails::CustomizeDetails(IDetailLayoutBuilder& Detail
 
 void FDMXEntityFixtureTypeDetails::OnGDTFSourceChanged()
 {
+	using namespace UE::DMX::GDTF;
+
 	const TArray<TWeakObjectPtr<UObject>>& SelectedObjects = PropertyUtilities->GetSelectedObjects();
 
 	for (TWeakObjectPtr<UObject> WeakFixtureTypeObject : SelectedObjects)
@@ -48,16 +48,9 @@ void FDMXEntityFixtureTypeDetails::OnGDTFSourceChanged()
 			}
 			UDMXImportGDTF* GDTF = FixtureType->GDTFSource.LoadSynchronous();
 
-			// Try to use the work around that supports creation of matrices, otherwise setup the fixture type with the old implementation
 			FixtureType->PreEditChange(nullptr);
-			const bool bAdvancedImportSuccess = FDMXInitializeFixtureTypeFromGDTFHelper::GenerateModesFromGDTF(*FixtureType, *GDTF);
-			if (!bAdvancedImportSuccess)
-			{
-				PRAGMA_DISABLE_DEPRECATION_WARNINGS
-				UE_LOG(LogDMXEditor, Warning, TEXT("Failed to initialize Fixture Type '%s', falling back to legacy method that doesn't support matrix fixtures."), *FixtureType->GetName());
-				FixtureType->SetModesFromDMXImport(GDTF);
-				PRAGMA_ENABLE_DEPRECATION_WARNINGS
-			}
+			constexpr bool bUpdateFixtureTypeName = true;
+			FDMXGDTFToFixtureTypeConverter::ConvertGDTF(*FixtureType, *GDTF, bUpdateFixtureTypeName);
 			FixtureType->PostEditChange();
 		}
 	}
diff --git a/Engine/Plugins/VirtualProduction/DMX/DMXEngine/Source/DMXEditor/Private/DMXEditor.cpp b/Engine/Plugins/VirtualProduction/DMX/DMXEngine/Source/DMXEditor/Private/DMXEditor.cpp
--- a/Engine/Plugins/VirtualProduction/DMX/DMXEngine/Source/DMXEditor/Private/DMXEditor.cpp
+++ b/Engine/Plugins/VirtualProduction/DMX/DMXEngine/Source/DMXEditor/Private/DMXEditor.cpp
@@ -1,43 +1,36 @@
 // Copyright Epic Games, Inc. All Rights Reserved.
 
 #include "DMXEditor.h"
+
+#include "Commands/DMXEditorCommands.h"
+#include "DesktopPlatformModule.h"
 #include "DMXEditorLog.h"
 #include "DMXEditorModule.h"
 #include "DMXEditorSettings.h"
 #include "DMXEditorTabNames.h"
-#include "DMXEditorUtils.h"
-#include "DMXFixtureTypeSharedData.h"
 #include "DMXFixturePatchSharedData.h"
+#include "DMXFixtureTypeSharedData.h"
 #include "DMXRuntimeLog.h"
-#include "DMXRuntimeUtils.h"
 #include "Exporters/DMXMVRExporter.h"
-#include "Library/DMXLibrary.h"
-#include "Library/DMXEntityFixtureType.h"
+#include "Framework/Application/SlateApplication.h"
+#include "Framework/Notifications/NotificationManager.h"
+#include "IDesktopPlatform.h"
 #include "Library/DMXEntityFixturePatch.h"
+#include "Library/DMXEntityFixtureType.h"
 #include "Library/DMXEntityReference.h"
+#include "Library/DMXLibrary.h"
+#include "Misc/MessageDialog.h"
 #include "Modes/DMXEditorApplicationMode.h"
+#include "ScopedTransaction.h"
+#include "Styling/AppStyle.h"
 #include "Toolbars/DMXEditorToolbar.h"
-#include "Commands/DMXEditorCommands.h"
-
-#include "Widgets/Layout/SScrollBox.h"
+#include "Utils.h"
 #include "Widgets/Docking/SDockTab.h"
-
-#include "Widgets/SDMXEntityEditor.h"
 #include "Widgets/FixturePatch/SDMXFixturePatchEditor.h"
 #include "Widgets/FixtureType/SDMXFixtureTypeEditor.h"
 #include "Widgets/LibrarySettings/SDMXLibraryEditorTab.h"
-
-#include "DesktopPlatformModule.h"
-#include "IDesktopPlatform.h"
-#include "ScopedTransaction.h"
-#include "Utils.h"
-#include "Framework/Application/SlateApplication.h"
-#include "Framework/Notifications/NotificationManager.h"
-#include "Misc/MessageDialog.h"
-#include "Modules/ModuleManager.h"
-#include "Styling/AppStyle.h"
 #include "Widgets/Notifications/SNotificationList.h"
-
+#include "Widgets/SDMXEntityEditor.h"
 
 #define LOCTEXT_NAMESPACE "FDMXEditor"
 
@@ -134,26 +127,13 @@ UDMXLibrary* FDMXEditor::GetDMXLibrary() const
 void FDMXEditor::ImportDMXLibrary() const
 {
 	UDMXLibrary* DMXLibrary = GetDMXLibrary();
-	if (!DMXLibrary)
-	{
-		return;
-	}
-
 	IDesktopPlatform* DesktopPlatform = FDesktopPlatformModule::Get();
-	if (!DesktopPlatform)
-	{
-		return;
-	}
-
 	UDMXEditorSettings* DMXEditorSettings = GetMutableDefault<UDMXEditorSettings>();
-	if (!DMXEditorSettings)
+	if (!DMXLibrary || !DesktopPlatform || !DMXEditorSettings)
 	{
 		return;
 	}
 
-	const FString LastMVRImportPath = DMXEditorSettings->LastMVRImportPath;
-	const FString DefaultPath = FPaths::DirectoryExists(LastMVRImportPath) ? LastMVRImportPath : FPaths::ProjectSavedDir();
-	
 	if (!DMXLibrary->GetEntities().IsEmpty())
 	{
 		const FText MessageText = LOCTEXT("MVRImportDialog", "DMX Library already contains data. Importing the MVR will clear existing data. Do you want to proceed?");
@@ -162,7 +142,10 @@ void FDMXEditor::ImportDMXLibrary() const
 			return;
 		}
 	}
-
+	
+	const FString LastMVRImportPath = DMXEditorSettings->LastMVRImportPath;
+	const FString DefaultPath = FPaths::DirectoryExists(LastMVRImportPath) ? LastMVRImportPath : FPaths::ProjectSavedDir();
+	
 	TArray<FString> OpenFilenames;
 	DesktopPlatform->OpenFileDialog(
 		FSlateApplication::Get().FindBestParentWindowHandleForDialogs(nullptr),
@@ -218,7 +201,7 @@ void FDMXEditor::ExportDMXLibrary() const
 		}
 
 		FText ErrorReason;
-		FDMXMVRExporter::Export(DMXLibrary, SaveFilenames[0], ErrorReason);
+		UE::DMX::FDMXMVRExporter::Export(DMXLibrary, SaveFilenames[0], ErrorReason);
 		if (ErrorReason.IsEmpty())
 		{
 			DMXEditorSettings->LastMVRExportPath = FPaths::GetPath(SaveFilenames[0]);
diff --git a/Engine/Plugins/VirtualProduction/DMX/DMXEngine/Source/DMXEditor/Private/DMXInitializeFixtureTypeFromGDTFHelper.cpp b/Engine/Plugins/VirtualProduction/DMX/DMXEngine/Source/DMXEditor/Private/DMXInitializeFixtureTypeFromGDTFHelper.cpp
deleted file mode 100644
diff --git a/Engine/Plugins/VirtualProduction/DMX/DMXEngine/Source/DMXEditor/Private/DMXInitializeFixtureTypeFromGDTFHelper.h b/Engine/Plugins/VirtualProduction/DMX/DMXEngine/Source/DMXEditor/Private/DMXInitializeFixtureTypeFromGDTFHelper.h
deleted file mode 100644
diff --git a/Engine/Plugins/VirtualProduction/DMX/DMXEngine/Source/DMXEditor/Private/Exporters/DMXFixtureTypeToGDTFConverter.cpp b/Engine/Plugins/VirtualProduction/DMX/DMXEngine/Source/DMXEditor/Private/Exporters/DMXFixtureTypeToGDTFConverter.cpp
new file mode 100644
--- /dev/null
+++ b/Engine/Plugins/VirtualProduction/DMX/DMXEngine/Source/DMXEditor/Private/Exporters/DMXFixtureTypeToGDTFConverter.cpp
@@ -0,0 +1,403 @@
+// Copyright Epic Games, Inc. All Rights Reserved.
+
+#include "Exporters/DMXFixtureTypeToGDTFConverter.h"
+
+#include "Algo/Find.h"
+#include "Algo/Sort.h"
+#include "DMXGDTF.h"
+#include "DMXUnrealToGDTFAttributeConversion.h"
+#include "GDTF/AttributeDefinitions/DMXGDTFAttribute.h"
+#include "GDTF/AttributeDefinitions/DMXGDTFAttributeDefinitions.h"
+#include "GDTF/AttributeDefinitions/DMXGDTFFeature.h"
+#include "GDTF/AttributeDefinitions/DMXGDTFFeatureGroup.h"
+#include "GDTF/AttributeDefinitions/DMXGDTFPhysicalUnit.h"
+#include "GDTF/DMXGDTFFixtureType.h"
+#include "GDTF/DMXModes/DMXGDTFChannelFunction.h"
+#include "GDTF/DMXModes/DMXGDTFDMXChannel.h"
+#include "GDTF/DMXModes/DMXGDTFDMXMode.h"
+#include "GDTF/DMXModes/DMXGDTFLogicalChannel.h"
+#include "GDTF/Geometries/DMXGDTFBeamGeometry.h"
+#include "GDTF/Geometries/DMXGDTFGeometryBreak.h"
+#include "GDTF/Geometries/DMXGDTFGeometryCollect.h"
+#include "GDTF/Geometries/DMXGDTFGeometryReference.h"
+#include "GDTF/Models/DMXGDTFModel.h"
+#include "Interfaces/IPluginManager.h"
+#include "Library/DMXEntityFixtureType.h"
+#include "Library/DMXLibrary.h"
+#include "XmlFile.h"
+
+namespace UE::DMX::GDTF
+{
+	const FName FDMXFixtureTypeToGTDFConverter::MatrixCellGeometryReferenceName = TEXT("Instance");
+
+	TSharedPtr<FXmlFile> FDMXFixtureTypeToGTDFConverter::Convert(const UDMXEntityFixtureType* UnrealFixtureType)
+	{
+		FDMXFixtureTypeToGTDFConverter Converter;
+		const TSharedRef<FDMXGDTFFixtureType> GDTFFixtureType = Converter.CreateFixtureType(UnrealFixtureType);
+
+		// Create the XML file
+		UDMXGDTF* GDTF = NewObject<UDMXGDTF>();
+		GDTF->InitializeFromFixtureType(GDTFFixtureType);
+		const TSharedPtr<FXmlFile> DescriptionXml = GDTF->ExportAsXml();
+
+		return DescriptionXml;
+	}
+
+	TSharedRef<FDMXGDTFFixtureType> FDMXFixtureTypeToGTDFConverter::CreateFixtureType(const UDMXEntityFixtureType* UnrealFixtureType)
+	{
+		const TSharedRef<FDMXGDTFFixtureType> GDTFFixtureType = MakeShared<FDMXGDTFFixtureType>();
+
+		GDTFFixtureType->Name = *UnrealFixtureType->Name;
+		GDTFFixtureType->ShortName = *UnrealFixtureType->Name;
+		GDTFFixtureType->LongName = UnrealFixtureType->GetParentLibrary()->GetName() + TEXT(" ") + UnrealFixtureType->Name;
+		GDTFFixtureType->Manufacturer = TEXT("Epic Games");
+		GDTFFixtureType->Description = FString::Printf(TEXT("Unreal Engine generated fixture type"));
+		GDTFFixtureType->FixtureTypeID = FGuid::NewGuid(); // Avoid any ambiguity with previously exported GDTFs, even if they're identical.
+		GDTFFixtureType->bCanHaveChildren = false;
+
+		CreateAttributeDefinitions(UnrealFixtureType, GDTFFixtureType);
+		CreateModels(UnrealFixtureType, GDTFFixtureType);
+		CreateGeometryCollect(UnrealFixtureType, GDTFFixtureType);
+		CreateDMXModes(UnrealFixtureType, GDTFFixtureType);
+
+		return GDTFFixtureType;
+	}
+
+	void FDMXFixtureTypeToGTDFConverter::CreateAttributeDefinitions(const UDMXEntityFixtureType* UnrealFixtureType, const TSharedRef<FDMXGDTFFixtureType>& GDTFFixtureType)
+	{
+		TArray<FName> AttributeNames;
+		for (const FDMXFixtureMode& Mode : UnrealFixtureType->Modes)
+		{
+			for (const FDMXFixtureFunction& Function : Mode.Functions)
+			{
+				AttributeNames.Add(Function.Attribute.Name);
+			}
+
+			if (Mode.bFixtureMatrixEnabled)
+			{
+				for (const FDMXFixtureCellAttribute& MatrixAttribute : Mode.FixtureMatrixConfig.CellAttributes)
+				{
+					AttributeNames.Add(MatrixAttribute.Attribute.Name);
+				}
+			}
+		}
+
+		GDTFFixtureType->AttributeDefinitions = MakeShared<FDMXGDTFAttributeDefinitions>(GDTFFixtureType);
+		for (const FName& AttributeName : AttributeNames)
+		{
+			const FName GDTFAttributeName = FDMXUnrealToGDTFAttributeConversion::ConvertUnrealToGDTFAttribute(AttributeName);
+			
+			const FName PrettyName = FDMXUnrealToGDTFAttributeConversion::GetPrettyFromGDTFAttribute(GDTFAttributeName);
+			const FName FeatureGroupName = FDMXUnrealToGDTFAttributeConversion::GetFeatureGroupForGDTFAttribute(GDTFAttributeName);
+			const FName FeatureName = FDMXUnrealToGDTFAttributeConversion::GetFeatureForGDTFAttribute(GDTFAttributeName);
+
+			// Get or create the GDTF feature group
+			const TSharedRef<FDMXGDTFFeatureGroup> GDTFFeatureGroup = [&GDTFFixtureType, &FeatureGroupName]()
+				{		
+					const TSharedPtr<FDMXGDTFFeatureGroup>* GDTFFeatureGroupPtr = Algo::FindBy(GDTFFixtureType->AttributeDefinitions->FeatureGroups, FeatureGroupName, &FDMXGDTFFeatureGroup::Name);
+
+					if (GDTFFeatureGroupPtr)
+					{
+						return (*GDTFFeatureGroupPtr).ToSharedRef();
+					}
+					else
+					{
+						const TSharedRef<FDMXGDTFFeatureGroup> NewFeatureGroup = MakeShared<FDMXGDTFFeatureGroup>(GDTFFixtureType->AttributeDefinitions.ToSharedRef());
+						GDTFFixtureType->AttributeDefinitions->FeatureGroups.Add(NewFeatureGroup);
+
+						return NewFeatureGroup;
+					}
+				}();
+
+			GDTFFeatureGroup->Name = FeatureGroupName;
+
+			// Get or create the GDTF feature
+			const TSharedRef<FDMXGDTFFeature> GDTFFeature = [GDTFFeatureGroup, &FeatureName]()
+				{					
+					const TSharedPtr<FDMXGDTFFeature>* GDTFFeaturePtr = Algo::FindBy(GDTFFeatureGroup->FeatureArray, FeatureName, &FDMXGDTFFeature::Name);
+					if (GDTFFeaturePtr)
+					{
+						return (*GDTFFeaturePtr).ToSharedRef();
+					}
+					else
+					{
+						const TSharedRef<FDMXGDTFFeature> NewFeature = MakeShared<FDMXGDTFFeature>(GDTFFeatureGroup);
+						GDTFFeatureGroup->FeatureArray.Add(NewFeature);
+
+						return NewFeature;
+					}
+				}();
+
+			GDTFFeature->Name = FeatureName;
+
+			// Create GDTF attribute
+			const TSharedRef<FDMXGDTFAttribute> GDTFAttribute = MakeShared<FDMXGDTFAttribute>(GDTFFixtureType->AttributeDefinitions.ToSharedRef());
+			GDTFFixtureType->AttributeDefinitions->Attributes.Add(GDTFAttribute);
+
+			GDTFAttribute->Name = GDTFAttributeName;
+			GDTFAttribute->Pretty = PrettyName.ToString();
+			GDTFAttribute->PhysicalUnit = EDMXGDTFPhysicalUnit::None;
+			GDTFAttribute->Feature = FeatureGroupName.ToString() + TEXT(".") + FeatureName.ToString();
+		}
+	}
+
+	void FDMXFixtureTypeToGTDFConverter::CreateModels(const UDMXEntityFixtureType* UnrealFixtureType, const TSharedRef<FDMXGDTFFixtureType>& GDTFFixtureType)
+	{
+		// Create a model for the matrix if this is a matrix
+		const bool bIsMatrix = Algo::FindBy(UnrealFixtureType->Modes, true, &FDMXFixtureMode::bFixtureMatrixEnabled) != nullptr;
+		if (bIsMatrix)
+		{
+			const TSharedRef<FDMXGDTFModel> InstanceModel = MakeShared<FDMXGDTFModel>(GDTFFixtureType);
+			GDTFFixtureType->Models.Add(InstanceModel);
+
+			InstanceModel->Name = "Layers";
+			InstanceModel->PrimitiveType = EDMXGDTFModelPrimitiveType::Cube;
+			InstanceModel->Height = 0.01f;
+			InstanceModel->Length = 1.f;
+			InstanceModel->Width = 0.3f;
+		}
+	}
+
+	void FDMXFixtureTypeToGTDFConverter::CreateGeometryCollect(const UDMXEntityFixtureType* UnrealFixtureType, const TSharedRef<FDMXGDTFFixtureType>& GDTFFixtureType)
+	{
+		GDTFFixtureType->GeometryCollect = MakeShared<FDMXGDTFGeometryCollect>(GDTFFixtureType);
+
+		// Always add a base geometry
+		const TSharedRef<FDMXGDTFGeometry> BaseGeometry = MakeShared<FDMXGDTFGeometry>(GDTFFixtureType->GeometryCollect.ToSharedRef());
+		GDTFFixtureType->GeometryCollect->GeometryArray.Add(BaseGeometry);
+
+		BaseGeometry->Name = TEXT("Base");
+
+		const TSharedRef<FDMXGDTFBeamGeometry> BeamGeometry = MakeShared<FDMXGDTFBeamGeometry>(GDTFFixtureType->GeometryCollect.ToSharedRef());
+		BaseGeometry->BeamArray.Add(BeamGeometry);
+
+		BeamGeometry->Name = TEXT("Beam");
+
+		for (const FDMXFixtureMode& UnrealMode : UnrealFixtureType->Modes)
+		{				
+			// Remember the geometry for this mode so it later can be referenced when building DMX Modes
+			UnrealModeToRootGeometryMap.Add(&UnrealMode, BaseGeometry);
+
+			CreateChildGeometries(UnrealMode, BeamGeometry);
+		}
+	}
+
+	void FDMXFixtureTypeToGTDFConverter::CreateChildGeometries(const FDMXFixtureMode& UnrealMode, const TSharedRef<FDMXGDTFGeometry>& RootGeometry)
+	{
+		if (UnrealMode.bFixtureMatrixEnabled && UnrealMode.FixtureMatrixConfig.GetNumChannels() > 0)
+		{
+			// Get the byte size of an Unreal Matrix Cell
+			const int32 CellSize = [UnrealMode]()
+				{
+					int32 OutCellSize = 0;
+					for (const FDMXFixtureCellAttribute& CellAttribute : UnrealMode.FixtureMatrixConfig.CellAttributes)
+					{
+						OutCellSize += CellAttribute.GetNumChannels();
+					}
+
+					return OutCellSize;
+				}();
+
+			// Create Geometry References for each Unreal Matrix Cell
+			const int32 NumCells = UnrealMode.FixtureMatrixConfig.XCells * UnrealMode.FixtureMatrixConfig.YCells;
+			int32 Offset = 1;
+			for (int32 CellID = 0; CellID < NumCells; CellID++)
+			{
+				const int32 DMXOffset = CellID * CellSize + 1;
+				const FName GeometryName = *FString::Printf(TEXT("Layer_%i"), CellID + 1);
+
+				const TSharedRef<FDMXGDTFGeometryReference> GeometryReference = MakeShared<FDMXGDTFGeometryReference>(RootGeometry);
+				RootGeometry->GeometryReferenceArray.Add(GeometryReference);
+
+				GeometryReference->Name = GeometryName;
+				GeometryReference->Geometry = MatrixCellGeometryReferenceName;
+				GeometryReference->Model = TEXT("Layers");
+
+				const TSharedRef<FDMXGDTFGeometryBreak> Break = MakeShared<FDMXGDTFGeometryBreak>(GeometryReference);
+				Break->DMXBreak = 1;
+				Break->DMXOffset = Offset;
+				GeometryReference->BreakArray.Add(Break);
+
+				Break->DMXBreak = 1; // Unreal does not support multi universe patches, DMXBreak is always 1
+				Break->DMXOffset = DMXOffset;
+
+				for (const FDMXFixtureCellAttribute& UnrealMatrixAttribute : UnrealMode.FixtureMatrixConfig.CellAttributes)
+				{
+					UnrealCellAttributeToGeometryReferenceMap.Add(&UnrealMatrixAttribute, GeometryReference);
+				}
+
+				Offset += CellSize;
+			}
+		}
+
+		// Always add common functions
+		for (const FDMXFixtureFunction& UnrealFunction : UnrealMode.Functions)
+		{
+			// Common Unreal Functions are assigned to the root geometry for now so they do not need specific handling
+			UnrealFunctionToGeometryMap.Add(&UnrealFunction, RootGeometry);
+		}
+	}
+
+	void FDMXFixtureTypeToGTDFConverter::CreateDMXModes(const UDMXEntityFixtureType* UnrealFixtureType, const TSharedRef<FDMXGDTFFixtureType>& GDTFFixtureType)
+	{
+		for (const FDMXFixtureMode& UnrealMode : UnrealFixtureType->Modes)
+		{
+			const TSharedRef<FDMXGDTFGeometry>* RootGeometryPtr = UnrealModeToRootGeometryMap.Find(&UnrealMode);
+			if (!ensureMsgf(RootGeometryPtr, TEXT("%hs: Unexpected cannot find root geometry for DMX Mode '%s'. Failed to convert mode to GDTF."), __FUNCTION__, *UnrealMode.ModeName))
+			{
+				continue;
+			}
+
+			// Create the mode
+			const TSharedRef<FDMXGDTFDMXMode> DMXMode = MakeShared<FDMXGDTFDMXMode>(GDTFFixtureType);
+			GDTFFixtureType->DMXModes.Add(DMXMode);
+
+			DMXMode->Name = *UnrealMode.ModeName;
+			DMXMode->Description = TEXT("Unreal Engine generated DMX Mode");
+			DMXMode->Geometry = (*RootGeometryPtr)->Name;
+
+			CreateDMXChannels(UnrealMode, DMXMode);
+		}
+	}
+
+	void FDMXFixtureTypeToGTDFConverter::CreateDMXChannels(const FDMXFixtureMode& UnrealMode, const TSharedRef<FDMXGDTFDMXMode>& GDTFDMXMode)
+	{			
+		// Create DMX Channels for non-matrix Unreal Functions
+		for (const FDMXFixtureFunction& UnrealFunction : UnrealMode.Functions)
+		{
+			const TSharedRef<FDMXGDTFGeometry>* DMXChannelGeometryPtr = UnrealFunctionToGeometryMap.Find(&UnrealFunction);
+			if (!ensureMsgf(DMXChannelGeometryPtr, TEXT("%hs: Unexpected cannot find geometry for DMX Function '%s'. Failed to convert mode to GDTF."), __FUNCTION__, *UnrealFunction.FunctionName))
+			{
+				continue;
+			}
+
+			const TSharedRef<FDMXGDTFDMXChannel> DMXChannel = MakeShared<FDMXGDTFDMXChannel>(GDTFDMXMode);
+			GDTFDMXMode->DMXChannels.Add(DMXChannel);
+
+			const FString GDTFAttribute = FDMXUnrealToGDTFAttributeConversion::ConvertUnrealToGDTFAttribute(UnrealFunction.Attribute.Name).ToString();
+			const FString ChannelFunctionName = UnrealFunction.FunctionName;
+
+			// The initial function has to be written in following format "GeometryName_LogicalChannelAttribute.ChannelFunctionAttribute.ChannelFunctionName"
+			DMXChannel->InitialFunction = FString::Printf(TEXT("%s_%s.%s.%s"), *(*DMXChannelGeometryPtr)->Name.ToString(), *GDTFAttribute, *GDTFAttribute, *ChannelFunctionName);
+			DMXChannel->Geometry = (*DMXChannelGeometryPtr)->Name;
+			DMXChannel->Offset = [UnrealFunction]()
+				{			
+					const int32 Offset = UnrealFunction.Channel;
+					const uint8 Size = UnrealFunction.GetNumChannels();
+
+					TArray<uint32> Offsets;
+					for (int32 ByteOffset = Offset; ByteOffset < Offset + Size; ByteOffset++)
+					{
+						Offsets.Add(ByteOffset);
+					}
+					const bool bUseLSBMode = UnrealFunction.bUseLSBMode;
+					Algo::Sort(Offsets, [bUseLSBMode](uint32 OffsetA, uint32 OffsetB)
+						{
+							return bUseLSBMode ? OffsetA >= OffsetB : OffsetA <= OffsetB;
+						});
+
+					return Offsets;
+				}();
+
+			CreateLogicalChannel(UnrealFunction, DMXChannel, GDTFAttribute);
+		}
+
+		// Create DMX Channels for Unreal Matrix Cells if this is a matrix mode
+		if (UnrealMode.bFixtureMatrixEnabled)
+		{
+			const int32 MatrixStartingChannel = UnrealMode.FixtureMatrixConfig.FirstCellChannel;
+
+			int32 Offset = UnrealMode.FixtureMatrixConfig.FirstCellChannel;
+			for (const FDMXFixtureCellAttribute& UnrealCellAttribute : UnrealMode.FixtureMatrixConfig.CellAttributes)
+			{
+				const TSharedRef<FDMXGDTFDMXChannel> DMXChannel = MakeShared<FDMXGDTFDMXChannel>(GDTFDMXMode);
+				GDTFDMXMode->DMXChannels.Add(DMXChannel);
+
+				const FString GDTFAttribute = FDMXUnrealToGDTFAttributeConversion::ConvertUnrealToGDTFAttribute(UnrealCellAttribute.Attribute.Name).ToString();
+
+				// For a matrix with geometry references, the initial function has to be written in following format "GeometryName_LogicalChannelAttribute.ChannelFunctionAttribute.ChannelFunctionName"
+				DMXChannel->InitialFunction = FString::Printf(TEXT("%s_%s.%s.%s"), *MatrixCellGeometryReferenceName.ToString(), *GDTFAttribute, *GDTFAttribute, *UnrealCellAttribute.Attribute.Name.ToString());
+				DMXChannel->Geometry = MatrixCellGeometryReferenceName;
+				DMXChannel->Offset = [&Offset, UnrealCellAttribute]()
+					{					
+						const uint8 Size = UnrealCellAttribute.GetNumChannels();
+
+						TArray<uint32> Offsets;
+						for (int32 ByteOffset = Offset; ByteOffset < Offset + Size; ByteOffset++)
+						{
+							Offsets.Add(ByteOffset);
+						}
+						const bool bUseLSBMode = UnrealCellAttribute.bUseLSBMode;
+						Algo::Sort(Offsets, [bUseLSBMode](uint32 OffsetA, uint32 OffsetB)
+							{
+								return bUseLSBMode ? OffsetA >= OffsetB : OffsetA <= OffsetB;
+							});
+
+						Offset += Size;
+
+						return Offsets;
+					}();
+
+				// Using a negative value to express special value "Overwrite"
+				DMXChannel->DMXBreak = -1;
+
+				CreateLogicalChannel(UnrealCellAttribute, DMXChannel, GDTFAttribute);
+			}
+		}
+
+		// Sort by Offset
+		Algo::SortBy(GDTFDMXMode->DMXChannels, [](const TSharedPtr<FDMXGDTFDMXChannel>& DMXChannel)
+			{
+				const uint32* MinElementPtr = Algo::MinElementBy(DMXChannel->Offset, [](int32 Element)
+					{
+						return Element;
+					});
+
+				return MinElementPtr ? *MinElementPtr : 0;
+			});
+	}
+
+	void FDMXFixtureTypeToGTDFConverter::CreateLogicalChannel(const FDMXFixtureFunction& UnrealFunction, const TSharedRef<FDMXGDTFDMXChannel>& GDTFDMXChannel, const FString& GDTFAttribute)
+	{
+		const TSharedRef<FDMXGDTFLogicalChannel> LogicalChannel = MakeShared<FDMXGDTFLogicalChannel>(GDTFDMXChannel);
+		GDTFDMXChannel->LogicalChannelArray.Add(LogicalChannel);
+
+		LogicalChannel->Attribute = *GDTFAttribute;
+	
+		CrateChannelFunction(UnrealFunction, LogicalChannel, GDTFAttribute);
+	}
+
+	void FDMXFixtureTypeToGTDFConverter::CreateLogicalChannel(const FDMXFixtureCellAttribute& UnrealCellAttribute, const TSharedRef<FDMXGDTFDMXChannel>& GDTFDMXChannel, const FString& GDTFAttribute)
+	{
+		const TSharedRef<FDMXGDTFLogicalChannel> LogicalChannel = MakeShared<FDMXGDTFLogicalChannel>(GDTFDMXChannel);
+		GDTFDMXChannel->LogicalChannelArray.Add(LogicalChannel);
+
+		LogicalChannel->Attribute = *GDTFAttribute;
+
+		CrateChannelFunction(UnrealCellAttribute, LogicalChannel, GDTFAttribute);
+	}
+
+	void FDMXFixtureTypeToGTDFConverter::CrateChannelFunction(const FDMXFixtureFunction& UnrealFunction, const TSharedRef<FDMXGDTFLogicalChannel>& GDTFLogicalChannel, const FString& GDTFAttribute)
+	{
+		const TSharedRef<FDMXGDTFChannelFunction> ChannelFunction = MakeShared<FDMXGDTFChannelFunction>(GDTFLogicalChannel);
+		GDTFLogicalChannel->ChannelFunctionArray.Add(ChannelFunction);
+
+		ChannelFunction->Name = *UnrealFunction.FunctionName;
+		ChannelFunction->Attribute = GDTFAttribute;
+
+		const FDMXGDTFDMXValue Default = UnrealFunction.DefaultValue;
+		ChannelFunction->Default = Default;
+		ChannelFunction->DMXFrom = 0;
+	}
+
+	void FDMXFixtureTypeToGTDFConverter::CrateChannelFunction(const FDMXFixtureCellAttribute& UnrealCellAttribute, const TSharedRef<FDMXGDTFLogicalChannel>& GDTFLogicalChannel, const FString& GDTFAttribute)
+	{
+		const TSharedRef<FDMXGDTFChannelFunction> ChannelFunction = MakeShared<FDMXGDTFChannelFunction>(GDTFLogicalChannel);
+		GDTFLogicalChannel->ChannelFunctionArray.Add(ChannelFunction);
+
+		ChannelFunction->Name = *(UnrealCellAttribute.Attribute.Name.ToString());
+		ChannelFunction->Attribute = GDTFAttribute;
+		ChannelFunction->Default = 0;
+		ChannelFunction->DMXFrom = 0;
+	}
+}
diff --git a/Engine/Plugins/VirtualProduction/DMX/DMXEngine/Source/DMXEditor/Private/Exporters/DMXFixtureTypeToGDTFConverter.h b/Engine/Plugins/VirtualProduction/DMX/DMXEngine/Source/DMXEditor/Private/Exporters/DMXFixtureTypeToGDTFConverter.h
new file mode 100644
--- /dev/null
+++ b/Engine/Plugins/VirtualProduction/DMX/DMXEngine/Source/DMXEditor/Private/Exporters/DMXFixtureTypeToGDTFConverter.h
@@ -0,0 +1,78 @@
+// Copyright Epic Games, Inc. All Rights Reserved.
+
+#pragma once
+
+#include "Templates/SharedPointer.h"
+#include "UObject/NameTypes.h"
+
+struct FDMXFixtureFunction;
+struct FDMXFixtureCellAttribute;
+struct FDMXFixtureMode;
+class FXmlFile;
+class UDMXEntityFixtureType;
+
+namespace UE::DMX::GDTF
+{
+	class FDMXGDTFAttributeDefinitions;
+	class FDMXGDTFDMXChannel;
+	class FDMXGDTFDMXMode;
+	class FDMXGDTFFixtureType;
+	class FDMXGDTFGeometry;
+	class FDMXGDTFGeometryCollect;
+	class FDMXGDTFGeometryReference;
+	class FDMXGDTFLogicalChannel;
+
+	/** Converts a Fixture Type to a GDTF. Internally caches of each collect, and finally assembles the GDTF. */
+	class FDMXFixtureTypeToGTDFConverter
+	{
+	public:
+		/** Converts the Fixture Type to a GDTF description */
+		static TSharedPtr<FXmlFile> Convert(const UDMXEntityFixtureType* UnrealFixtureType);
+
+	private:
+		/** Creates the Fixture Type */
+		TSharedRef<FDMXGDTFFixtureType> CreateFixtureType(const UDMXEntityFixtureType* UnrealFixtureType);
+
+		/** Creates attribute definitions */
+		void CreateAttributeDefinitions(const UDMXEntityFixtureType* UnrealFixtureType, const TSharedRef<FDMXGDTFFixtureType>& GDTFFixtureType);
+
+		/** Creates model */
+		void CreateModels(const UDMXEntityFixtureType* UnrealFixtureType, const TSharedRef<FDMXGDTFFixtureType>& GDTFFixtureType);
+
+		/** Creates the geometry collect */
+		void CreateGeometryCollect(const UDMXEntityFixtureType* UnrealFixtureType, const TSharedRef<FDMXGDTFFixtureType>& GDTFFixtureType);
+
+		/** Creates Child geometries inside a root geometry */
+		void CreateChildGeometries(const FDMXFixtureMode& UnrealMode, const TSharedRef<FDMXGDTFGeometry>& RootGeometry);
+
+		/** Creates DMX Modes */
+		void CreateDMXModes(const UDMXEntityFixtureType* UnrealFixtureType, const TSharedRef<FDMXGDTFFixtureType>& GDTFFixtureType);
+
+		/** Creates DMX Channels for specified mode */
+		void CreateDMXChannels(const FDMXFixtureMode& UnrealMode, const TSharedRef<FDMXGDTFDMXMode>& GDTFDMXMode);
+
+		/** Creates a Logical Channel for the DMX Channel, using an Unreal Function for as input */
+		void CreateLogicalChannel(const FDMXFixtureFunction& UnrealFunction, const TSharedRef<FDMXGDTFDMXChannel>& GDTFDMXChannel, const FString& GDTFAttribute);
+
+		/** Creates a Logical Channel for the DMX Channel, using an Unreal Cell Attribute as the input */
+		void CreateLogicalChannel(const FDMXFixtureCellAttribute& UnrealCellAttribute, const TSharedRef<FDMXGDTFDMXChannel>& GDTFDMXChannel, const FString& GDTFAttribute);
+
+		/** Creates a Channel Function for the Logical Channel, using an Unreal Function for as input */
+		void CrateChannelFunction(const FDMXFixtureFunction& UnrealFunction, const TSharedRef<FDMXGDTFLogicalChannel>& GDTFLogicalChannel, const FString& GDTFAttribute);
+
+		/** Creates a Channel Function for the Logical Channel, using an Unreal Cell Attribute as the input */
+		void CrateChannelFunction(const FDMXFixtureCellAttribute& UnrealCellAttribute, const TSharedRef<FDMXGDTFLogicalChannel>& GDTFLogicalChannel, const FString& GDTFAttribute);
+
+		/** Map of Unreal Modes to the root geometry they control */
+		TMap<const FDMXFixtureMode*, TSharedRef<FDMXGDTFGeometry>> UnrealModeToRootGeometryMap;
+
+		/** Map of Unreal Fixture Functions to the geometry they control */
+		TMap<const FDMXFixtureFunction*, TSharedRef<FDMXGDTFGeometry>> UnrealFunctionToGeometryMap;
+
+		/** Map of Unreal Matrix Attribute to the geometry they control */
+		TMap<const FDMXFixtureCellAttribute*, TSharedRef<FDMXGDTFGeometryReference>> UnrealCellAttributeToGeometryReferenceMap;
+
+		/** The name of a matrix cell geometry reference */
+		static const FName MatrixCellGeometryReferenceName;
+	};
+}
diff --git a/Engine/Plugins/VirtualProduction/DMX/DMXEngine/Source/DMXEditor/Private/Exporters/DMXMVRExporter.cpp b/Engine/Plugins/VirtualProduction/DMX/DMXEngine/Source/DMXEditor/Private/Exporters/DMXMVRExporter.cpp
--- a/Engine/Plugins/VirtualProduction/DMX/DMXEngine/Source/DMXEditor/Private/Exporters/DMXMVRExporter.cpp
+++ b/Engine/Plugins/VirtualProduction/DMX/DMXEngine/Source/DMXEditor/Private/Exporters/DMXMVRExporter.cpp
@@ -2,399 +2,433 @@
 
 #include "Exporters/DMXMVRExporter.h"
 
+#include "DesktopPlatformModule.h"
 #include "DMXEditorLog.h"
 #include "DMXEditorSettings.h"
-#include "DMXInitializeFixtureTypeFromGDTFHelper.h"
+#include "DMXFixtureTypeToGDTFConverter.h"
 #include "DMXMVRXmlMergeUtility.h"
 #include "DMXZipper.h"
+#include "Editor.h"
+#include "EngineUtils.h"
+#include "Factories/DMXGDTFToFixtureTypeConverter.h"
+#include "Framework/Notifications/NotificationManager.h"
+#include "Game/DMXComponent.h"
+#include "GenericPlatform/GenericPlatformFile.h"
+#include "IDesktopPlatform.h"
 #include "Library/DMXEntityFixturePatch.h"
 #include "Library/DMXEntityFixtureType.h"
 #include "Library/DMXGDTFAssetImportData.h"
 #include "Library/DMXImportGDTF.h"
 #include "Library/DMXLibrary.h"
-#include "MVR/DMXMVRAssetImportData.h"
-#include "MVR/DMXMVRAssetUserData.h"
-#include "MVR/DMXMVRGeneralSceneDescription.h"
-#include "MVR/DMXMVRSceneActor.h"
-#include "MVR/Types/DMXMVRFixtureNode.h"
-
-#include "DesktopPlatformModule.h"
-#include "Editor.h"
-#include "EngineUtils.h"
-#include "IDesktopPlatform.h"
-#include "XmlFile.h"
-#include "Framework/Notifications/NotificationManager.h"
 #include "Misc/FileHelper.h"
 #include "Misc/MessageDialog.h"
 #include "Misc/Paths.h"
+#include "MVR/DMXMVRAssetImportData.h"
+#include "MVR/DMXMVRGeneralSceneDescription.h"
+#include "MVR/Types/DMXMVRFixtureNode.h"
 #include "Subsystems/UnrealEditorSubsystem.h"
 #include "Widgets/Notifications/SNotificationList.h"
-
+#include "XmlFile.h"
 
 #define LOCTEXT_NAMESPACE "DMXMVRExporter"
 
-void FDMXMVRExporter::Export(UDMXLibrary* DMXLibrary, const FString& FilePathAndName, FText& OutErrorReason)
-{
-	FDMXMVRExporter Instance;
-	Instance.ExportInternal(DMXLibrary, FilePathAndName, OutErrorReason);
-}
-
-void FDMXMVRExporter::ExportInternal(UDMXLibrary* DMXLibrary, const FString& FilePathAndName, FText& OutErrorReason)
+namespace UE::DMX
 {
-	OutErrorReason = FText::GetEmpty();
-
-	if (!ensureAlwaysMsgf(DMXLibrary, TEXT("Trying to export DMX Library '%s' as MVR file, but the DMX Library is invalid."), *DMXLibrary->GetName()))
-	{
-		OutErrorReason = FText::Format(LOCTEXT("MVRExportDMXLibraryInvalidReason", "DMX Library is invalid. Cannot export {0}."), FText::FromString(FilePathAndName));
-		return;
-	}
-
-	// Create a copy of the library's general scene description, so transforms exported aren't written to the dmx library
-	DMXLibrary->UpdateGeneralSceneDescription();
-	UDMXMVRGeneralSceneDescription* TempGeneralSceneDescription = DuplicateObject<UDMXMVRGeneralSceneDescription>(DMXLibrary->GetLazyGeneralSceneDescription(), GetTransientPackage());
-	if (!ensureAlwaysMsgf(TempGeneralSceneDescription, TEXT("Trying to export DMX Library '%s' as MVR file, but its General Scene Description is invalid."), *DMXLibrary->GetName()))
-	{
-		OutErrorReason = FText::Format(LOCTEXT("MVRExportGeneralSceneDescriptionInvalidReason", "DMX Library is invalid. Cannot export {0}."), FText::FromString(FilePathAndName));
-		return;
-	}
-
-	WriteMVRFixtureTransformsFromLevel(TempGeneralSceneDescription);
-
-	const TSharedRef<FDMXZipper> Zip = MakeShared<FDMXZipper>();
-	if (!ZipGeneralSceneDescription(Zip, TempGeneralSceneDescription, OutErrorReason))
-	{
-		return;
-	}
-
-	if (!ZipGDTFs(Zip, DMXLibrary))
-	{
-		OutErrorReason = FText::Format(LOCTEXT("MVRExportZipGDTFsFailedReason", "Some GDTFs are missing import data and were not considered. Exported MVR to {0}."), FText::FromString(FilePathAndName));
-		// Allow continuation of export
-	}
-
-	ZipThirdPartyData(Zip, TempGeneralSceneDescription);
-	if (!Zip->SaveToFile(FilePathAndName))
-	{
-		OutErrorReason = FText::Format(LOCTEXT("MVRExportWriteZipFailedReason", "File is not writable or locked by another process. Cannot export {0}."), FText::FromString(FilePathAndName));
-		return;
-	}
-}
-
-bool FDMXMVRExporter::ZipGeneralSceneDescription(const TSharedRef<FDMXZipper>& Zip, const UDMXMVRGeneralSceneDescription* GeneralSceneDescription, FText& OutErrorReason)
-{
-	if (!GeneralSceneDescription->CanCreateXmlFile(OutErrorReason))
-	{
-		return false;
-	}
-
-	TSharedPtr<FXmlFile> XmlFile = GeneralSceneDescription->CreateXmlFile();
-	if (!ensureMsgf(XmlFile.IsValid(), TEXT("General Scene Descriptions returns CanCreateXmlFile() as true, but CreateXmlFile() fails.")))
+	void FDMXMVRExporter::Export(UDMXLibrary* DMXLibrary, const FString& FilePathAndName, FText& OutErrorReason)
 	{
-		return false;
+		FDMXMVRExporter Instance;
+		Instance.ExportInternal(DMXLibrary, FilePathAndName, OutErrorReason);
 	}
 
-	// Try to merge the source General Scene Description Xml
-	const TSharedPtr<FXmlFile> SourceXmlFile = CreateSourceGeneralSceneDescriptionXmlFile(GeneralSceneDescription);
-	if (SourceXmlFile.IsValid())
+	void FDMXMVRExporter::ExportInternal(UDMXLibrary* DMXLibrary, const FString& FilePathAndName, FText& OutErrorReason)
 	{
-		// Merge the General Scene Description with source data's xml (source) to retain 3rd party data
-		XmlFile = FDMXXmlMergeUtility::Merge(GeneralSceneDescription, SourceXmlFile.ToSharedRef());
-	}
-	
-	// Create a temp GeneralSceneDescription.xml file
-	const FString TempPath = FPaths::ConvertRelativePathToFull(FPaths::EngineSavedDir() / TEXT("DMX_Temp"));
-	constexpr TCHAR GeneralSceneDescriptionFileName[] = TEXT("GeneralSceneDescription.xml");
-	const FString TempGeneralSceneDescriptionFilePathAndName = TempPath / GeneralSceneDescriptionFileName;
-	if (!XmlFile->Save(TempGeneralSceneDescriptionFilePathAndName))
-	{
-		UE_LOG(LogDMXEditor, Error, TEXT("Failed to save General Scene Description. See previous errors for details."));
-		return false;
-	}
+		OutErrorReason = FText::GetEmpty();
 
-	TArray64<uint8> GeneralSceneDescriptionData;
-	FFileHelper::LoadFileToArray(GeneralSceneDescriptionData, *TempGeneralSceneDescriptionFilePathAndName);
-	Zip->AddFile(GeneralSceneDescriptionFileName, GeneralSceneDescriptionData);
+		if (!ensureAlwaysMsgf(DMXLibrary, TEXT("Trying to export DMX Library '%s' as MVR file, but the DMX Library is invalid."), *DMXLibrary->GetName()))
+		{
+			OutErrorReason = FText::Format(LOCTEXT("MVRExportDMXLibraryInvalidReason", "DMX Library is invalid. Cannot export {0}."), FText::FromString(FilePathAndName));
+			return;
+		}
 
-	// Delete temp GeneralSceneDescription.xml file
-	IFileManager& FileManager = IFileManager::Get();
-	constexpr bool bRequireExists = true;
-	constexpr bool bEvenIfReadOnly = false;
-	constexpr bool bQuiet = true;
-	FileManager.Delete(*TempGeneralSceneDescriptionFilePathAndName, bRequireExists, bEvenIfReadOnly, bQuiet);
+		// Create a copy of the library's general scene description, so transforms exported aren't written to the dmx library
+		UDMXMVRGeneralSceneDescription* TempGeneralSceneDescription = UDMXMVRGeneralSceneDescription::CreateFromDMXLibrary(*DMXLibrary, GetTransientPackage(), NAME_None);
+		if (!ensureAlwaysMsgf(TempGeneralSceneDescription, TEXT("Trying to export DMX Library '%s' as MVR file, but its General Scene Description is invalid."), *DMXLibrary->GetName()))
+		{
+			OutErrorReason = FText::Format(LOCTEXT("MVRExportGeneralSceneDescriptionInvalidReason", "DMX Library is invalid. Cannot export {0}."), FText::FromString(FilePathAndName));
+			return;
+		}
 
-	return true;
-}
+		WriteMVRFixtureTransformsFromLevel(TempGeneralSceneDescription);
 
-bool FDMXMVRExporter::ZipGDTFs(const TSharedRef<FDMXZipper>& Zip, UDMXLibrary* DMXLibrary)
-{
-	check(DMXLibrary);
+		const TSharedRef<FDMXZipper> Zip = MakeShared<FDMXZipper>();
+		if (!ZipGeneralSceneDescription(Zip, TempGeneralSceneDescription, OutErrorReason))
+		{
+			return;
+		}
 
-	// Gather GDTFs
-	TArray<UDMXEntityFixtureType*> FixtureTypesToExport;
-	const TArray<UDMXEntityFixturePatch*> FixturePatches = DMXLibrary->GetEntitiesTypeCast<UDMXEntityFixturePatch>();
-	for (UDMXEntityFixturePatch* FixturePatch : FixturePatches)
-	{
-		UDMXEntityFixtureType* FixtureType = FixturePatch->GetFixtureType();
-		if (!FixtureType || FixtureType->Modes.IsEmpty())
+		if (!ZipGDTFs(Zip, DMXLibrary))
 		{
-			continue;
+			OutErrorReason = FText::Format(LOCTEXT("MVRExportZipGDTFsFailedReason", "Some Fixture Types could not be converted to GDTF. Exported MVR to {0}."), FText::FromString(FilePathAndName));
+			// Allow continuation of export
 		}
 
-		FixtureTypesToExport.AddUnique(FixtureType);
+		ZipThirdPartyData(Zip, TempGeneralSceneDescription);
+		if (!Zip->SaveToFile(FilePathAndName))
+		{
+			OutErrorReason = FText::Format(LOCTEXT("MVRExportWriteZipFailedReason", "File is not writable or locked by another process. Cannot export {0}."), FText::FromString(FilePathAndName));
+			return;
+		}
 	}
 
-	// Zip GDTFs
-	bool bAllZippedSuccessfully = true;
-	for (UDMXEntityFixtureType* FixtureType : FixtureTypesToExport)
+	bool FDMXMVRExporter::ZipGeneralSceneDescription(const TSharedRef<FDMXZipper>& Zip, const UDMXMVRGeneralSceneDescription* GeneralSceneDescription, FText& OutErrorReason)
 	{
-		if (FixtureType->GDTFSource.IsNull())
+		if (!GeneralSceneDescription->CanCreateXmlFile(OutErrorReason))
 		{
-			UE_LOG(LogDMXEditor, Warning, TEXT("Cannot export Fixture Type '%s' to MVR, but the Fixture Type has no GDTF set."), *FixtureType->Name);
-			continue;
+			return false;
 		}
 
-		UDMXImportGDTF* DMXImportGDTF = FixtureType->GDTFSource.LoadSynchronous();
-		if (!DMXImportGDTF)
+		TSharedPtr<FXmlFile> XmlFile = GeneralSceneDescription->CreateXmlFile();
+		if (!ensureMsgf(XmlFile.IsValid(), TEXT("General Scene Descriptions returns CanCreateXmlFile() as true, but CreateXmlFile() fails.")))
 		{
-			UE_LOG(LogDMXEditor, Warning, TEXT("Cannot export Fixture Type '%s' to MVR, but the Fixture Type has a DMX Import Type which is not GDTF."), *FixtureType->Name);
-			continue;
+			return false;
 		}
 
-		UDMXGDTFAssetImportData* GDTFAssetImportData = DMXImportGDTF->GetGDTFAssetImportData();
-		if (!ensureAlwaysMsgf(GDTFAssetImportData, TEXT("Missing default GDTF Asset Import Data subobject in GDTF.")))
+		// Try to merge the source General Scene Description Xml
+		const TSharedPtr<FXmlFile> SourceXmlFile = CreateSourceGeneralSceneDescriptionXmlFile(GeneralSceneDescription);
+		if (SourceXmlFile.IsValid())
 		{
-			continue;
+			// Merge the General Scene Description with source data's xml (source) to retain 3rd party data
+			XmlFile = FDMXXmlMergeUtility::Merge(GeneralSceneDescription, SourceXmlFile.ToSharedRef());
 		}
 
-		const TArray64<uint8>& RawSourceData = RefreshSourceDataAndFixtureType(*FixtureType, *GDTFAssetImportData);
-		if (RawSourceData.Num() > 0)
-		{
-			const FString GDTFFilename = FPaths::GetCleanFilename(GDTFAssetImportData->GetFilePathAndName());
-			Zip->AddFile(GDTFFilename, GDTFAssetImportData->GetRawSourceData());
-		}
-		else
+		// Create a temp GeneralSceneDescription.xml file
+		const FString TempPath = FPaths::ConvertRelativePathToFull(FPaths::EngineSavedDir() / TEXT("DMX_Temp"));
+		constexpr TCHAR GeneralSceneDescriptionFileName[] = TEXT("GeneralSceneDescription.xml");
+		const FString TempGeneralSceneDescriptionFilePathAndName = TempPath / GeneralSceneDescriptionFileName;
+		if (!XmlFile->Save(TempGeneralSceneDescriptionFilePathAndName))
 		{
-			bAllZippedSuccessfully = false;
-			const UDMXImportGDTF* GDTF = Cast<UDMXImportGDTF>(GDTFAssetImportData->GetOuter());
-			UE_LOG(LogDMXEditor, Error, TEXT("Cannot export '%s' to MVR File. The asset is missing source data."), GDTF ? *GDTF->GetName() : TEXT("Invalid GDTF Asset"));
+			UE_LOG(LogDMXEditor, Error, TEXT("Failed to save General Scene Description. See previous errors for details."));
+			return false;
 		}
-	}
 
-	return bAllZippedSuccessfully;
-}
+		TArray64<uint8> GeneralSceneDescriptionData;
+		FFileHelper::LoadFileToArray(GeneralSceneDescriptionData, *TempGeneralSceneDescriptionFilePathAndName);
+		Zip->AddFile(GeneralSceneDescriptionFileName, GeneralSceneDescriptionData);
 
-void FDMXMVRExporter::ZipThirdPartyData(const TSharedRef<FDMXZipper>& Zip, const UDMXMVRGeneralSceneDescription* GeneralSceneDescription)
-{
-	const UDMXMVRAssetImportData* AssetImportData = GeneralSceneDescription->GetMVRAssetImportData();
-	if (!AssetImportData)
-	{
-		return;
+		// Delete temp GeneralSceneDescription.xml file
+		IFileManager& FileManager = IFileManager::Get();
+		constexpr bool bRequireExists = true;
+		constexpr bool bEvenIfReadOnly = false;
+		constexpr bool bQuiet = true;
+		FileManager.Delete(*TempGeneralSceneDescriptionFilePathAndName, bRequireExists, bEvenIfReadOnly, bQuiet);
+
+		return true;
 	}
 
-	if (AssetImportData && AssetImportData->GetRawSourceData().Num() > 0)
+	bool FDMXMVRExporter::ZipGDTFs(const TSharedRef<FDMXZipper>& Zip, UDMXLibrary* DMXLibrary)
 	{
-		const TSharedRef<FDMXZipper> SourceZip = MakeShared<FDMXZipper>();
-		if (SourceZip->LoadFromData(AssetImportData->GetRawSourceData()))
+		check(DMXLibrary);
+
+		// Gather GDTFs
+		TArray<UDMXEntityFixtureType*> FixtureTypesToExport;
+		const TArray<UDMXEntityFixturePatch*> FixturePatches = DMXLibrary->GetEntitiesTypeCast<UDMXEntityFixturePatch>();
+		for (UDMXEntityFixturePatch* FixturePatch : FixturePatches)
 		{
-			for (const FString& SourceFileName : SourceZip->GetFiles())
+			UDMXEntityFixtureType* FixtureType = FixturePatch->GetFixtureType();
+			if (!FixtureType || FixtureType->Modes.IsEmpty())
 			{
-				// Don't add GeneralSceneDescription.xml and GDTFs
-				constexpr TCHAR GeneralSceneDescriptionFileName[] = TEXT("GeneralSceneDescription.xml");
-				constexpr TCHAR GDTFExtension[] = TEXT("gdtf");
-				if (SourceFileName.EndsWith(GeneralSceneDescriptionFileName) || FPaths::GetExtension(SourceFileName) == GDTFExtension)
+				continue;
+			}
+
+			FixtureTypesToExport.AddUnique(FixtureType);
+		}
+
+		// Zip GDTFs
+		const FString TempGDTFDir = FPaths::Combine(FPaths::GameAgnosticSavedDir() / TEXT("GDTFExport"));
+
+		bool bAllZippedSuccessfully = true;
+		for (UDMXEntityFixtureType* FixtureType : FixtureTypesToExport)
+		{
+			constexpr bool bWithExtension = true;
+			const FString CleanGDTFFilename = FixtureType->GetCleanGDTFFileNameSynchronous(bWithExtension);
+
+			using namespace UE::DMX::GDTF;
+			const TSharedPtr<FXmlFile> XmlFile = FDMXFixtureTypeToGTDFConverter::Convert(FixtureType);
+			if (!XmlFile.IsValid())
+			{
+				UE_LOG(LogDMXEditor, Warning, TEXT("Failed to create GDTF from Fixture Type '%s'. See previous errors."), *FixtureType->Name);
+				bAllZippedSuccessfully &= false;
+				continue;
+			}
+
+			if (!FPaths::DirectoryExists(TempGDTFDir))
+			{
+				IPlatformFile::GetPlatformPhysical().CreateDirectory(*TempGDTFDir);
+			}
+
+			if (!ensureMsgf(FPaths::DirectoryExists(TempGDTFDir), TEXT("Cannot create temporary directory for GDTFs.")))
+			{
+				continue;
+			}
+
+			const FString TempXmlFilename = TempGDTFDir / FixtureType->Name + TEXT(".description.xml");
+			XmlFile->Save(TempXmlFilename);
+
+			TArray64<uint8> DescriptionXmlData;
+			if (FFileHelper::LoadFileToArray(DescriptionXmlData, *TempXmlFilename))
+			{
+				FDMXZipper NewGDTFZip;
+				NewGDTFZip.AddFile(TEXT("description.xml"), DescriptionXmlData);
+				const FString TempGDTFilePathAndName = TempGDTFDir / CleanGDTFFilename;
+
+				TArray64<uint8> GDTFData;
+				if (NewGDTFZip.GetData(GDTFData))
 				{
-					continue;
+					Zip->AddFile(CleanGDTFFilename, GDTFData);
 				}
-
-				TArray64<uint8> SourceFileData;
-				if (SourceZip->GetFileContent(SourceFileName, SourceFileData))
+				else
 				{
-					Zip->AddFile(SourceFileName, SourceFileData);
+					UE_LOG(LogDMXEditor, Warning, TEXT("Failed to load temporary GDTF for Fixture Type '%s'."), *FixtureType->Name);
+					bAllZippedSuccessfully &= false;
+					continue;
 				}
 			}
+			else
+			{
+				UE_LOG(LogDMXEditor, Warning, TEXT("Failed to load temporary description.xml for Fixture Type '%s'."), *FixtureType->Name);
+				bAllZippedSuccessfully &= false;
+				continue;
+			}
 		}
-	}
-}
 
-void FDMXMVRExporter::WriteMVRFixtureTransformsFromLevel(UDMXMVRGeneralSceneDescription* GeneralSceneDescription)
-{
-	if (!ensureMsgf(GeneralSceneDescription, TEXT("Trying to write MVR Fixture Transforms from Level to MVR General Scene Description, but General Scene Description is invalid")))
-	{
-		return;
+		IPlatformFile::GetPlatformPhysical().DeleteDirectory(*TempGDTFDir);
+
+		return bAllZippedSuccessfully;
 	}
 
-	UUnrealEditorSubsystem* UnrealEditorSubsystem = GEditor->GetEditorSubsystem<UUnrealEditorSubsystem>();
-	if (!UnrealEditorSubsystem)
-	{
-		return;
-	} 
-	
-	const UWorld* EditorWorld = UnrealEditorSubsystem->GetEditorWorld();
-	if (!EditorWorld)
+	void FDMXMVRExporter::ZipThirdPartyData(const TSharedRef<FDMXZipper>& Zip, const UDMXMVRGeneralSceneDescription* GeneralSceneDescription)
 	{
-		return;
+		const UDMXMVRAssetImportData* AssetImportData = GeneralSceneDescription->GetMVRAssetImportData();
+		if (!AssetImportData)
+		{
+			return;
+		}
+
+		if (AssetImportData && AssetImportData->GetRawSourceData().Num() > 0)
+		{
+			const TSharedRef<FDMXZipper> SourceZip = MakeShared<FDMXZipper>();
+			if (SourceZip->LoadFromData(AssetImportData->GetRawSourceData()))
+			{
+				for (const FString& SourceFileName : SourceZip->GetFiles())
+				{
+					// Don't add GeneralSceneDescription.xml and GDTFs
+					constexpr TCHAR GeneralSceneDescriptionFileName[] = TEXT("GeneralSceneDescription.xml");
+					constexpr TCHAR GDTFExtension[] = TEXT("gdtf");
+					if (SourceFileName.EndsWith(GeneralSceneDescriptionFileName) || FPaths::GetExtension(SourceFileName) == GDTFExtension)
+					{
+						continue;
+					}
+
+					TArray64<uint8> SourceFileData;
+					if (SourceZip->GetFileContent(SourceFileName, SourceFileData))
+					{
+						Zip->AddFile(SourceFileName, SourceFileData);
+					}
+				}
+			}
+		}
 	}
 
-	TArray<const ADMXMVRSceneActor*> MVRSceneActors;
-	for (TActorIterator<ADMXMVRSceneActor> It(EditorWorld, ADMXMVRSceneActor::StaticClass()); It; ++It)
+	void FDMXMVRExporter::WriteMVRFixtureTransformsFromLevel(UDMXMVRGeneralSceneDescription* GeneralSceneDescription)
 	{
-		ADMXMVRSceneActor* Actor = *It;
-		MVRSceneActors.Add(Actor);
-	}
+		if (!ensureMsgf(GeneralSceneDescription, TEXT("Trying to write MVR Fixture Transforms from Level to MVR General Scene Description, but General Scene Description is invalid")))
+		{
+			return;
+		}
 
-	TArray<UDMXMVRFixtureNode*> FixtureNodes;
-	GeneralSceneDescription->GetFixtureNodes(FixtureNodes);
+		UUnrealEditorSubsystem* UnrealEditorSubsystem = GEditor->GetEditorSubsystem<UUnrealEditorSubsystem>();
+		if (!UnrealEditorSubsystem)
+		{
+			return;
+		}
 
-	TArray<UDMXMVRFixtureNode*> HandledMVRFixtureNodes;
-	for (const ADMXMVRSceneActor* MVRSceneActor : MVRSceneActors)
-	{
-		for (const TSoftObjectPtr<AActor>& RelatedActorSoftObjectPtr : MVRSceneActor->GetRelatedActors())
+		const UWorld* EditorWorld = UnrealEditorSubsystem->GetEditorWorld();
+		if (!EditorWorld)
 		{
-			const AActor* RelatedActor = RelatedActorSoftObjectPtr.Get();
-			if (!RelatedActor)
+			return;
+		}
+
+		// Find actors with a patched DMX component
+		TMap<const UDMXEntityFixturePatch*, FTransform> PatchToTransformMap;
+		for (TActorIterator<AActor> It(EditorWorld, AActor::StaticClass()); It; ++It)
+		{
+			const AActor* Actor = *It;
+			if (!Actor)
 			{
 				continue;
 			}
 
-			const FString MVRFixtureUUIDString = UDMXMVRAssetUserData::GetMVRAssetUserDataValueForKey(*RelatedActor, UDMXMVRAssetUserData::MVRFixtureUUIDMetaDataKey);
-			FGuid MVRFixtureUUID;
-			if (!FGuid::Parse(MVRFixtureUUIDString, MVRFixtureUUID))
+			const TSet<UActorComponent*> Components = Actor->GetComponents();
+
+			TArray<const UDMXEntityFixturePatch*> FixturePatches;
+			Algo::TransformIf(Components, FixturePatches, 
+				[](const UActorComponent* Component)
+				{
+					return Component && 
+						Component->IsA(UDMXComponent::StaticClass()) &&
+						CastChecked<UDMXComponent>(Component)->GetFixturePatch();
+				},
+				[](const UActorComponent* Component)
+				{
+					return CastChecked<UDMXComponent>(Component)->GetFixturePatch();
+				});
+
+			for (const UDMXEntityFixturePatch* FixturePatch : FixturePatches)
 			{
-				continue;
+				PatchToTransformMap.Add(FixturePatch, Actor->GetTransform());
 			}
+		}
 
-			UDMXMVRFixtureNode** FixtureNodePtr = Algo::FindByPredicate(FixtureNodes, [MVRFixtureUUID](UDMXMVRFixtureNode* MVRFixtureNode)
+		TArray<UDMXMVRFixtureNode*> FixtureNodes;
+		GeneralSceneDescription->GetFixtureNodes(FixtureNodes);
+
+		TArray<UDMXMVRFixtureNode*> HandledMVRFixtureNodes;
+		for (const TTuple<const UDMXEntityFixturePatch*, FTransform>& PatchToTransformPair : PatchToTransformMap)
+		{
+			UDMXMVRFixtureNode** FixtureNodePtr = Algo::FindByPredicate(FixtureNodes, [PatchToTransformPair](UDMXMVRFixtureNode* MVRFixtureNode)
 				{
-					return MVRFixtureNode->UUID == MVRFixtureUUID;
+					return MVRFixtureNode->UUID == PatchToTransformPair.Key->GetMVRFixtureUUID();
 				});
 
 			if (FixtureNodePtr)
 			{
 				if (HandledMVRFixtureNodes.Contains(*FixtureNodePtr))
 				{
-					UE_LOG(LogDMXEditor, Warning, TEXT("Ambigous Fixture ID for Fixture '%s' when exporting MVR. Each Fixture Patch can only be used once in an MVR scene. Using transform of arbitrary instance with ambigous ID."), *RelatedActor->GetName());
+					UE_LOG(LogDMXEditor, Warning, TEXT("Ambigous Fixture ID for Fixture Patch '%s' when exporting MVR. Each Fixture Patch can only be used once in an MVR scene. Using transform of arbitrary instance with ambigous ID."), *PatchToTransformPair.Key->Name);
 					continue;
 				}
 
-				(*FixtureNodePtr)->SetTransformAbsolute(RelatedActor->GetTransform());
+				(*FixtureNodePtr)->SetTransformAbsolute(PatchToTransformPair.Value);
 				HandledMVRFixtureNodes.Add(*FixtureNodePtr);
 			}
 		}
 	}
-}
 
-const TSharedPtr<FXmlFile> FDMXMVRExporter::CreateSourceGeneralSceneDescriptionXmlFile(const UDMXMVRGeneralSceneDescription* GeneralSceneDescription) const
-{
-	const UDMXMVRAssetImportData* AssetImportData = GeneralSceneDescription->GetMVRAssetImportData();
-	if (!AssetImportData)
+	const TSharedPtr<FXmlFile> FDMXMVRExporter::CreateSourceGeneralSceneDescriptionXmlFile(const UDMXMVRGeneralSceneDescription* GeneralSceneDescription) const
 	{
-		return nullptr;
-	}
-
-	if (AssetImportData->GetRawSourceData().IsEmpty())
-	{
-		return nullptr;
-	}
+		const UDMXMVRAssetImportData* AssetImportData = GeneralSceneDescription->GetMVRAssetImportData();
+		if (!AssetImportData)
+		{
+			return nullptr;
+		}
 
-	const TSharedRef<FDMXZipper> DMXZipper = MakeShared<FDMXZipper>();
-	if (!DMXZipper->LoadFromData(AssetImportData->GetRawSourceData()))
-	{
-		return nullptr;
-	}
+		if (AssetImportData->GetRawSourceData().IsEmpty())
+		{
+			return nullptr;
+		}
 
-	constexpr TCHAR GeneralSceneDescriptionFilename[] = TEXT("GeneralSceneDescription.xml");
-	FDMXZipper::FDMXScopedUnzipToTempFile UnzipTempFileScope(DMXZipper, GeneralSceneDescriptionFilename);
+		const TSharedRef<FDMXZipper> DMXZipper = MakeShared<FDMXZipper>();
+		if (!DMXZipper->LoadFromData(AssetImportData->GetRawSourceData()))
+		{
+			return nullptr;
+		}
 
-	if (UnzipTempFileScope.TempFilePathAndName.IsEmpty())
-	{
-		return nullptr;
-	}
+		constexpr TCHAR GeneralSceneDescriptionFilename[] = TEXT("GeneralSceneDescription.xml");
+		FDMXZipper::FDMXScopedUnzipToTempFile UnzipTempFileScope(DMXZipper, GeneralSceneDescriptionFilename);
 
-	const TSharedRef<FXmlFile> XmlFile = MakeShared<FXmlFile>();
-	if (!XmlFile->LoadFile(UnzipTempFileScope.TempFilePathAndName))
-	{
-		return nullptr;
-	}
+		if (UnzipTempFileScope.TempFilePathAndName.IsEmpty())
+		{
+			return nullptr;
+		}
 
-	return XmlFile;
-}
+		const TSharedRef<FXmlFile> XmlFile = MakeShared<FXmlFile>();
+		if (!XmlFile->LoadFile(UnzipTempFileScope.TempFilePathAndName))
+		{
+			return nullptr;
+		}
 
-const TArray64<uint8>& FDMXMVRExporter::RefreshSourceDataAndFixtureType(UDMXEntityFixtureType& FixtureType, UDMXGDTFAssetImportData& InOutGDTFAssetImportData) const
-{
-	const UDMXImportGDTF* GDTF = Cast<UDMXImportGDTF>(InOutGDTFAssetImportData.GetOuter());
-	if (!InOutGDTFAssetImportData.GetRawSourceData().IsEmpty() || !GDTF)
-	{
-		return InOutGDTFAssetImportData.GetRawSourceData();
+		return XmlFile;
 	}
 
-	static bool bAskForReloadAgain = true;
-	static EAppReturnType::Type MessageDialogResult = EAppReturnType::Yes;
-
-	if (bAskForReloadAgain)
+	const TArray64<uint8>& FDMXMVRExporter::RefreshSourceDataAndFixtureType(UDMXEntityFixtureType& FixtureType, UDMXGDTFAssetImportData& InOutGDTFAssetImportData) const
 	{
-		static const FText MessageTitle = LOCTEXT("NoGDTFSourceAvailableTitle", "Trying to use old GDTF asset.");
-		static const FText Message = FText::Format(LOCTEXT("NoGDTFSourceAvailableMessage", "Insufficient data to export '{0}' to MVR file. The GDTF asset was created prior to UE5.1. Do you want to reload the source GDTF?"), FText::FromString(GDTF->GetName()));
-		MessageDialogResult = FMessageDialog::Open(EAppMsgType::YesNoYesAllNoAll, Message, MessageTitle);
-
-		if (MessageDialogResult == EAppReturnType::YesAll || MessageDialogResult == EAppReturnType::NoAll)
+		const UDMXImportGDTF* GDTF = Cast<UDMXImportGDTF>(InOutGDTFAssetImportData.GetOuter());
+		if (!InOutGDTFAssetImportData.GetRawSourceData().IsEmpty() || !GDTF)
 		{
-			bAskForReloadAgain = false;
+			return InOutGDTFAssetImportData.GetRawSourceData();
 		}
-	}
 
-	if (MessageDialogResult == EAppReturnType::YesAll || MessageDialogResult == EAppReturnType::Yes)
-	{
-		IDesktopPlatform* const DesktopPlatform = FDesktopPlatformModule::Get();
+		static bool bAskForReloadAgain = true;
+		static EAppReturnType::Type MessageDialogResult = EAppReturnType::Yes;
 
-		if (DesktopPlatform)
+		if (bAskForReloadAgain)
 		{
-			UDMXEditorSettings* EditorSettings = GetMutableDefault<UDMXEditorSettings>();
-			if (!ensureAlwaysMsgf(EditorSettings, TEXT("Unexpected cannot access DMX Editor Settings CDO.")))
-			{
-				return InOutGDTFAssetImportData.GetRawSourceData();
-			}
+			static const FText MessageTitle = LOCTEXT("NoGDTFSourceAvailableTitle", "Trying to use old GDTF asset.");
+			static const FText Message = FText::Format(LOCTEXT("NoGDTFSourceAvailableMessage", "Insufficient data to export '{0}' to MVR file. The GDTF asset was created prior to UE5.1. Do you want to reload the source GDTF?"), FText::FromString(GDTF->GetName()));
+			MessageDialogResult = FMessageDialog::Open(EAppMsgType::YesNoYesAllNoAll, Message, MessageTitle);
 
-			TArray<FString> Filenames;
-			if (!InOutGDTFAssetImportData.GetSourceData().SourceFiles.IsEmpty() &&
-				FPaths::FileExists(InOutGDTFAssetImportData.GetFilePathAndName()))
-			{
-				Filenames.Add(InOutGDTFAssetImportData.GetFilePathAndName());
-			}
-			else
+			if (MessageDialogResult == EAppReturnType::YesAll || MessageDialogResult == EAppReturnType::NoAll)
 			{
-				DesktopPlatform->OpenFileDialog(
-					nullptr,
-					FText::Format(LOCTEXT("OpenGDTFTitle", "Choose a GDTF file for '%s'."), FText::FromString(GDTF->GetName())).ToString(),
-					EditorSettings->LastGDTFImportPath,
-					TEXT(""),
-					TEXT("General Scene Description (*.gdtf)|*.gdtf"),
-					EFileDialogFlags::None,
-					Filenames
-				);
+				bAskForReloadAgain = false;
 			}
+		}
 
-			if (Filenames.Num() > 0)
-			{
-				EditorSettings->LastGDTFImportPath = FPaths::GetPath(Filenames[0]);
+		if (MessageDialogResult == EAppReturnType::YesAll || MessageDialogResult == EAppReturnType::Yes)
+		{
+			IDesktopPlatform* const DesktopPlatform = FDesktopPlatformModule::Get();
 
-				InOutGDTFAssetImportData.PreEditChange(nullptr);
-				InOutGDTFAssetImportData.SetSourceFile(Filenames[0]);
-				InOutGDTFAssetImportData.PostEditChange();
+			if (DesktopPlatform)
+			{
+				UDMXEditorSettings* EditorSettings = GetMutableDefault<UDMXEditorSettings>();
+				if (!ensureAlwaysMsgf(EditorSettings, TEXT("Unexpected cannot access DMX Editor Settings CDO.")))
+				{
+					return InOutGDTFAssetImportData.GetRawSourceData();
+				}
 
-				FDMXInitializeFixtureTypeFromGDTFHelper::GenerateModesFromGDTF(FixtureType, *GDTF);
+				TArray<FString> Filenames;
+				if (!InOutGDTFAssetImportData.GetSourceData().SourceFiles.IsEmpty() &&
+					FPaths::FileExists(InOutGDTFAssetImportData.GetFilePathAndName()))
+				{
+					Filenames.Add(InOutGDTFAssetImportData.GetFilePathAndName());
+				}
+				else
+				{
+					DesktopPlatform->OpenFileDialog(
+						nullptr,
+						FText::Format(LOCTEXT("OpenGDTFTitle", "Choose a GDTF file for '%s'."), FText::FromString(GDTF->GetName())).ToString(),
+						EditorSettings->LastGDTFImportPath,
+						TEXT(""),
+						TEXT("General Scene Description (*.gdtf)|*.gdtf"),
+						EFileDialogFlags::None,
+						Filenames
+					);
+				}
 
-				if (InOutGDTFAssetImportData.GetRawSourceData().Num() == 0)
+				if (Filenames.Num() > 0)
 				{
-					FNotificationInfo NotificationInfo(FText::Format(LOCTEXT("ReloadGDTFFailure", "Failed to update GDTF '{0}' from '{1}."), FText::FromString(GDTF->GetName()), FText::FromString(Filenames[0])));
-					NotificationInfo.ExpireDuration = 10.f;
+					EditorSettings->LastGDTFImportPath = FPaths::GetPath(Filenames[0]);
+
+					InOutGDTFAssetImportData.PreEditChange(nullptr);
+					InOutGDTFAssetImportData.SetSourceFile(Filenames[0]);
+					InOutGDTFAssetImportData.PostEditChange();
 
-					FSlateNotificationManager::Get().AddNotification(NotificationInfo);
+					using namespace UE::DMX::GDTF;
+					constexpr bool bUpdateFixtureTypeName = true;
+					FDMXGDTFToFixtureTypeConverter::ConvertGDTF(FixtureType, *GDTF, bUpdateFixtureTypeName);
+
+					if (InOutGDTFAssetImportData.GetRawSourceData().Num() == 0)
+					{
+						FNotificationInfo NotificationInfo(FText::Format(LOCTEXT("ReloadGDTFFailure", "Failed to update GDTF '{0}' from '{1}."), FText::FromString(GDTF->GetName()), FText::FromString(Filenames[0])));
+						NotificationInfo.ExpireDuration = 10.f;
+
+						FSlateNotificationManager::Get().AddNotification(NotificationInfo);
+					}
 				}
 			}
 		}
-	}
 
-	return InOutGDTFAssetImportData.GetRawSourceData();
+		return InOutGDTFAssetImportData.GetRawSourceData();
+	}
 }
 
 #undef LOCTEXT_NAMESPACE
diff --git a/Engine/Plugins/VirtualProduction/DMX/DMXEngine/Source/DMXEditor/Private/Exporters/DMXMVRExporter.h b/Engine/Plugins/VirtualProduction/DMX/DMXEngine/Source/DMXEditor/Private/Exporters/DMXMVRExporter.h
--- a/Engine/Plugins/VirtualProduction/DMX/DMXEngine/Source/DMXEditor/Private/Exporters/DMXMVRExporter.h
+++ b/Engine/Plugins/VirtualProduction/DMX/DMXEngine/Source/DMXEditor/Private/Exporters/DMXMVRExporter.h
@@ -2,47 +2,52 @@
 
 #pragma once
 
+#if UE_ENABLE_INCLUDE_ORDER_DEPRECATED_IN_5_2
 #include "CoreMinimal.h"
+#endif
+#include "Containers/Array.h"
+#include "Templates/SharedPointer.h"
 
 class FDMXZipper;
+class FXmlFile;
 class UDMXEntityFixtureType;
 class UDMXGDTFAssetImportData;
 class UDMXLibrary;
 class UDMXMVRAssetImportData;
 class UDMXMVRGeneralSceneDescription;
 
-class FXmlFile;
-
-
-/** Helper class to export a DMX Library as MVR file */
-class FDMXMVRExporter
+namespace UE::DMX
 {
-public:
-	/** Exports the DMX Library as MVR File */
-	static void Export(UDMXLibrary* DMXLibrary, const FString& FilePathAndName, FText& OutErrorReason);
-
-private:
-	/** Exports the DMX Library as MVR File. If OutErrorReason is not empty there were issues with the export. */
-	void ExportInternal(UDMXLibrary* DMXLibrary, const FString& FilePathAndName, FText& OutErrorReason);
-
-	/** Zips the GeneralSceneDescription.xml */
-	[[nodiscard]] bool ZipGeneralSceneDescription(const TSharedRef<FDMXZipper>& Zip, const UDMXMVRGeneralSceneDescription* GeneralSceneDescription, FText& OutErrorReason);
-
-	/** Zips the GDTFs from the Library */
-	[[nodiscard]] bool ZipGDTFs(const TSharedRef<FDMXZipper>& Zip, UDMXLibrary* DMXLibrary);
-
-	/** Zips 3rd Party Data from the MVR Asset Import Data */
-	void ZipThirdPartyData(const TSharedRef<FDMXZipper>& Zip, const UDMXMVRGeneralSceneDescription* GeneralSceneDescription);
-
-	/** Creates an General Scene Description Xml File from the MVR Source, as it was imported */
-	const TSharedPtr<FXmlFile> CreateSourceGeneralSceneDescriptionXmlFile(const UDMXMVRGeneralSceneDescription* GeneralSceneDescription) const;
-
-	/** Writes transforms from level to the MVR Fixture Actors where applicable */
-	void WriteMVRFixtureTransformsFromLevel(UDMXMVRGeneralSceneDescription* GeneralSceneDescription);
-
-	/** 
-	 * Gets raw source data or creates (possibly empty) source data where the source data is not present. 
-	 * Prior 5.1 there was no source data stored. Offers a dialog to load missing data 
-	 */
-	const TArray64<uint8>& RefreshSourceDataAndFixtureType(UDMXEntityFixtureType& FixtureType, UDMXGDTFAssetImportData& InOutGDTFAssetImportData) const;
-};
+	/** Helper class to export a DMX Library as MVR file */
+	class FDMXMVRExporter
+	{
+	public:
+		/** Exports the DMX Library as MVR File */
+		static void Export(UDMXLibrary* DMXLibrary, const FString& FilePathAndName, FText& OutErrorReason);
+
+	private:
+		/** Exports the DMX Library as MVR File. If OutErrorReason is not empty there were issues with the export. */
+		void ExportInternal(UDMXLibrary* DMXLibrary, const FString& FilePathAndName, FText& OutErrorReason);
+
+		/** Zips the GeneralSceneDescription.xml */
+		[[nodiscard]] bool ZipGeneralSceneDescription(const TSharedRef<FDMXZipper>& Zip, const UDMXMVRGeneralSceneDescription* GeneralSceneDescription, FText& OutErrorReason);
+
+		/** Zips the GDTFs from the Library */
+		[[nodiscard]] bool ZipGDTFs(const TSharedRef<FDMXZipper>& Zip, UDMXLibrary* DMXLibrary);
+
+		/** Zips 3rd Party Data from the MVR Asset Import Data */
+		void ZipThirdPartyData(const TSharedRef<FDMXZipper>& Zip, const UDMXMVRGeneralSceneDescription* GeneralSceneDescription);
+
+		/** Creates an General Scene Description Xml File from the MVR Source, as it was imported */
+		const TSharedPtr<FXmlFile> CreateSourceGeneralSceneDescriptionXmlFile(const UDMXMVRGeneralSceneDescription* GeneralSceneDescription) const;
+
+		/** Writes transforms from level to the MVR Fixture Actors where applicable */
+		void WriteMVRFixtureTransformsFromLevel(UDMXMVRGeneralSceneDescription* GeneralSceneDescription);
+
+		/**
+		 * Gets raw source data or creates (possibly empty) source data where the source data is not present.
+		 * Prior 5.1 there was no source data stored. Offers a dialog to load missing data
+		 */
+		const TArray64<uint8>& RefreshSourceDataAndFixtureType(UDMXEntityFixtureType& FixtureType, UDMXGDTFAssetImportData& InOutGDTFAssetImportData) const;
+	};
+}
diff --git a/Engine/Plugins/VirtualProduction/DMX/DMXEngine/Source/DMXEditor/Private/Exporters/DMXUnrealToGDTFAttributeConversion.h b/Engine/Plugins/VirtualProduction/DMX/DMXEngine/Source/DMXEditor/Private/Exporters/DMXUnrealToGDTFAttributeConversion.h
new file mode 100644
--- /dev/null
+++ b/Engine/Plugins/VirtualProduction/DMX/DMXEngine/Source/DMXEditor/Private/Exporters/DMXUnrealToGDTFAttributeConversion.h
@@ -0,0 +1,166 @@
+// Copyright Epic Games, Inc. All Rights Reserved.
+
+#pragma once
+
+#include "Containers/Map.h"
+#include "Internationalization/Regex.h"
+#include "UObject/NameTypes.h"
+
+namespace UE::DMX::GDTF
+{
+	class FDMXUnrealToGDTFAttributeConversion
+	{
+	public:
+		/** Creates a GDTF attribute from an unreal attribute */
+		static FName ConvertUnrealToGDTFAttribute(const FName& InUnrealAttribute)
+		{			
+			// Remove, but remember trailing numbers
+			FString CleanAttribute;
+			FString TrailingNumberString; 
+
+			const FRegexPattern RegexPattern(TEXT("(f\\d+oo)(\\d+)$"));
+			FRegexMatcher RegexMatcher(RegexPattern, InUnrealAttribute.ToString());
+			if (RegexMatcher.FindNext())
+			{
+				CleanAttribute = *RegexMatcher.GetCaptureGroup(0);
+				TrailingNumberString = RegexMatcher.GetCaptureGroup(1);
+			}
+			else
+			{
+				CleanAttribute = InUnrealAttribute.ToString();
+			}
+
+			const FName* AttributeNamePtr = UnrealToGDTFAttributeMap.Find(*CleanAttribute);
+			if (AttributeNamePtr)
+			{
+				return *((*AttributeNamePtr).ToString() + TrailingNumberString);
+			}
+			else
+			{
+				return InUnrealAttribute;
+			}
+		}
+
+		static FName GetPrettyFromGDTFAttribute(const FName& InGDTFAttribute)
+		{
+			const FName* PrettyPtr = GDTFAttributeToPrettyMap.Find(InGDTFAttribute);
+			if (PrettyPtr)
+			{
+				return (*PrettyPtr);
+			}
+			else
+			{
+				return InGDTFAttribute;
+			}
+		}
+
+		/** Creates a GDTF feature group for a GDTF Attribute */
+		static FName GetFeatureGroupForGDTFAttribute(const FName& InGDTFAttribute)
+		{
+			const TTuple<FName, FName>* FeaturePairPtr = GDTFAttributeToFeatureMap.Find(InGDTFAttribute);
+			if (FeaturePairPtr)
+			{
+				return (*FeaturePairPtr).Key;
+			}
+			else
+			{
+				return "Control";
+			}
+		}
+
+		/** Creates a GDTF feature for a GDTF Attribute */
+		static FName GetFeatureForGDTFAttribute(const FName& InGDTFAttribute)
+		{
+			const TPair<FName, FName>* FeaturePairPtr = GDTFAttributeToFeatureMap.Find(InGDTFAttribute);
+			if (FeaturePairPtr)
+			{
+				return (*FeaturePairPtr).Value;
+			}
+			else
+			{
+				return "Control";
+			}
+		}
+
+	private:
+		/** Conversion from Unreal attributes to GDTF attributes. See also how UDMXProtocolSettings::Attributes is initialized and applied. */
+		static inline const TMap<FName, FName> UnrealToGDTFAttributeMap =
+		{
+			{ "Intensity", "Dimmer" },
+			{ "Strength", "Dimmer" },
+			{ "Brightness", "Dimmer" },
+
+			{ "Red", "ColorAdd_R" },
+			{ "Green", "ColorAdd_G" },
+			{ "Blue", "ColorAdd_B" },
+			{ "Cyan", "ColorAdd_C" },
+			{ "Magenta", "ColorAdd_M" },
+			{ "Yellow", "ColorAdd_Y" },
+			{ "White", "ColorAdd_W" },
+			{ "Amber", "ColorAdd_A" },
+
+			{ "Gobo Spin", "GoboSpin" },
+			{ "Gobo Wheel Rotate", "GoboWheel" }
+		};
+
+		/** Defines pretty attribute names for GDTF attribute names */
+		static inline const TMap<FName, FName> GDTFAttributeToPrettyMap =
+		{
+			{ "Dimmer", "Dim" },
+
+			{ "ColorAdd_R", "R" },
+			{ "ColorAdd_G", "G" },
+			{ "ColorAdd_B", "B" },
+			{ "ColorAdd_C", "C" },
+			{ "ColorAdd_M", "M" },
+			{ "ColorAdd_Y", "Y" },
+			{ "ColorAdd_W", "W" },
+			{ "ColorAdd_A", "A" },
+
+			{ "Pan", "P" },
+			{ "Tilt", "T" }
+		};
+
+		/** 
+		 * Defines attributes that should be assigned to a specific feature group. 
+		 * 
+		 * Assumes that GDTF and not Unreal attributes are used.
+		 */
+		static inline const TMap<FName, TPair<FName, FName>> GDTFAttributeToFeatureMap =
+		{
+			// Dimmer feature group
+			{ "Dimmer",			{ "Dimmer", "Dimmer" } },
+
+			// Color feature group
+			{ "Color",			{ "Color", "Color" } },
+			{ "CTC",			{ "Color", "Color" } },
+			{ "ColorAdd_R",		{ "Color", "RGB" } },
+			{ "ColorAdd_G",		{ "Color", "RGB" } },
+			{ "ColorAdd_B",		{ "Color", "RGB" } },
+			{ "ColorAdd_C",		{ "Color", "RGB" } },
+			{ "ColorAdd_M",		{ "Color", "RGB" } },
+			{ "ColorAdd_Y",		{ "Color", "RGB" } },
+			{ "ColorAdd_W",		{ "Color", "RGB" } },
+			{ "ColorAdd_A",		{ "Color", "RGB" } },
+			{ "CIE_X",			{ "Color", "CIE" } },
+			{ "CIE_Y",			{ "Color", "CIE" } },
+			{ "CIE_Brightness", { "Color", "CIE" } },
+
+			// Position feature group
+			{ "Pan",			{ "Position", "PanTilt" } },
+			{ "Tilt",			{ "Position", "PanTilt" } },
+
+			// Gobo feature group
+			{ "GoboSpin",		{ "Gobo", "Gobo" } },
+			{ "GoboWheel",		{ "Gobo", "Gobo" } },
+
+			// Focus feature group
+			{ "Focus",			{ "Focus", "Focus" } },
+			{ "Zoom",			{ "Focus", "Focus" } },
+
+			// Beam feature group
+			{ "Shutter",		{ "Beam", "Beam" } },
+			{ "Frost",			{ "Beam", "Beam" } }
+		};
+	};
+}
diff --git a/Engine/Plugins/VirtualProduction/DMX/DMXEngine/Source/DMXEditor/Private/Factories/DMXGDTFToFixtureTypeConverter.cpp b/Engine/Plugins/VirtualProduction/DMX/DMXEngine/Source/DMXEditor/Private/Factories/DMXGDTFToFixtureTypeConverter.cpp
new file mode 100644
--- /dev/null
+++ b/Engine/Plugins/VirtualProduction/DMX/DMXEngine/Source/DMXEditor/Private/Factories/DMXGDTFToFixtureTypeConverter.cpp
@@ -0,0 +1,559 @@
+// Copyright Epic Games, Inc. All Rights Reserved.
+
+#include "DMXGDTFToFixtureTypeConverter.h"
+
+#include "Algo/Copy.h"
+#include "Algo/Count.h"
+#include "Algo/Find.h"
+#include "Algo/MaxElement.h"
+#include "Algo/MinElement.h"
+#include "DMXEditorLog.h"
+#include "DMXGDTF.h"
+#include "DMXProtocolSettings.h"
+#include "DMXRuntimeUtils.h"
+#include "DMXZipper.h"
+#include "GDTF/DMXGDTFDescription.h"
+#include "GDTF/DMXGDTFFixtureType.h"
+#include "GDTF/DMXModes/DMXGDTFChannelFunction.h"
+#include "GDTF/DMXModes/DMXGDTFDMXChannel.h"
+#include "GDTF/DMXModes/DMXGDTFDMXMode.h"
+#include "GDTF/DMXModes/DMXGDTFLogicalChannel.h"
+#include "GDTF/Geometries/DMXGDTFGeometry.h"
+#include "GDTF/Geometries/DMXGDTFGeometryBreak.h"
+#include "GDTF/Geometries/DMXGDTFGeometryCollect.h"
+#include "GDTF/Geometries/DMXGDTFGeometryReference.h"
+#include "GDTF/Models/DMXGDTFModel.h"
+#include "Library/DMXEntityFixtureType.h"
+#include "Library/DMXGDTFAssetImportData.h"
+#include "Library/DMXImportGDTF.h"
+#include "XmlFile.h"
+#include "XmlNode.h"
+#include <limits>
+
+#define LOCTEXT_NAMESPACE "DMXInitializeFixtureTypeFromGDTFHelper"
+
+namespace UE::DMX::GDTF
+{
+	/** Helper to interprete a GDTF Channel */
+	class FDMXGDTFChannelInterpreter
+	{
+	public:
+		/** Properties of a channel */
+		struct FChannelProperties
+		{
+			FString AttributeName;
+			uint32 FirstChannel = 1;
+			EDMXFixtureSignalFormat SignalFormat = EDMXFixtureSignalFormat::E8Bit;
+			bool bLSBMode = false;
+			uint32 NumCells = 1;
+			uint32 DefaultValue = 0;
+		};
+
+		static bool GetChannelProperties(
+			const TSharedRef<FDMXGDTFDMXMode>& InDMXMode, 
+			const TSharedRef<FDMXGDTFDMXChannel>& InDMXChannelNode, 
+			FChannelProperties& OutChannelProperties)
+		{
+			FDMXGDTFChannelInterpreter Instance(InDMXMode, InDMXChannelNode);
+			
+			return Instance.GetChannelPropertiesInternal(OutChannelProperties);
+		}
+
+	private:
+		FDMXGDTFChannelInterpreter(const TSharedRef<FDMXGDTFDMXMode>& InDMXModeNode, const TSharedRef<FDMXGDTFDMXChannel>& InDMXChannelNode)
+			: DMXModeNode(InDMXModeNode)
+			, DMXChannelNode(InDMXChannelNode)
+		{
+			LogicalChannel = InDMXChannelNode->LogicalChannelArray.IsValidIndex(0) ? InDMXChannelNode->LogicalChannelArray[0] : nullptr;
+			ChannelFunction = InDMXChannelNode->ResolveInitialFunction();
+		}
+
+		/** Returns the properties of the channel. The root geometry collect is the geometry collect the mode references. */
+		bool GetChannelPropertiesInternal(FChannelProperties& OutChannelProperties) const
+		{
+			if (!GetAttributeName(OutChannelProperties.AttributeName))
+			{
+				return false;
+			}
+
+			OutChannelProperties.DefaultValue = GetDefaultValue();
+			GetDataType(OutChannelProperties.SignalFormat, OutChannelProperties.bLSBMode);
+
+			const int32 NumGeometries = CountGeometries();
+			if (NumGeometries > 1)
+			{				
+				// Try to interprete as Unreal matrix channel
+				const TArray<TSharedPtr<FDMXGDTFGeometryReference>> GeometryReferences = DMXChannelNode->ResolveGeometryReferences();
+				if (!GeometryReferences.IsEmpty())
+				{
+					OutChannelProperties.NumCells = GeometryReferences.Num();
+
+					if (!GetMatrixOffset(GeometryReferences, OutChannelProperties.FirstChannel))
+					{
+						return false;
+					}
+				}
+			}
+			else
+			{				
+				// Try to interprete as Unreal channel function
+				if (!GetMinOffsetOfDMXChannel(OutChannelProperties.FirstChannel))
+				{
+					return false;
+				}
+			}
+
+			return true;
+		}
+
+		/** Returns the number of geometries the DMX channel references */
+		int32 CountGeometries() const
+		{
+			// Try to find a single geometry that is referenced
+			const TSharedPtr<FDMXGDTFGeometry> Geometry = DMXChannelNode->ResolveGeometry();
+			if (Geometry.IsValid())
+			{
+				return 1;
+			}
+
+			// Try to find geometry references
+			const TArray<TSharedPtr<FDMXGDTFGeometryReference>> GeometryReferences = DMXChannelNode->ResolveGeometryReferences();
+			if (!GeometryReferences.IsEmpty())
+			{
+				return GeometryReferences.Num();
+			}
+
+			// Legacy GDTFs may reference a model instead of a geometry
+			const TSharedPtr<FDMXGDTFFixtureType> FixtureType = DMXChannelNode->GetFixtureType().Pin();
+			if (!ensureMsgf(FixtureType.IsValid(), TEXT("Unexpected invalid fixture type for DMX Channel node.")))
+			{
+				return 0;
+			}
+			const TSharedPtr<FDMXGDTFModel>* ModelPtr = Algo::FindByPredicate(FixtureType->Models, [this](const TSharedPtr<FDMXGDTFModel>& Model)
+				{
+					return Model->Name == DMXChannelNode->Geometry;
+				});
+			if (ModelPtr)
+			{
+				return 1;
+			}
+
+			// Fall back to the mode geometry 
+			if (!FixtureType->GeometryCollect.IsValid())
+			{
+				return 0;
+			}
+			
+			const TSharedPtr<FDMXGDTFGeometry> ModeGeometry = FixtureType->GeometryCollect->FindGeometryByName(*DMXModeNode->Geometry.ToString());
+			if (ModeGeometry.IsValid())
+			{
+				return 1;
+			}
+
+			// Accept even if the mode self only references a model
+			const TSharedPtr<FDMXGDTFModel>* ModeModelPtr = Algo::FindByPredicate(FixtureType->Models, [this](const TSharedPtr<FDMXGDTFModel>& Model)
+				{
+					return Model->Name == DMXModeNode->Geometry;
+				});
+
+			return ModeModelPtr ? 1 : 0;
+		}
+
+		/** 
+		 * Returns the attribute name of the channel node. 
+		 * As per legacy considers only the first logical channel - FDMXFixtureFunction is not in support of more than one. 
+		 */
+		bool GetAttributeName(FString& OutName) const
+		{
+			if (!LogicalChannel.IsValid())
+			{
+				return false;
+			}
+
+			OutName = LogicalChannel->Attribute.ToString();
+			return true;
+		}
+
+		/** 
+		 * Gets the default value of the channel. Fails quietly if no default value present.
+		 * Note, FDMXFixtureFunction only supports int32 default values currently.
+		 */
+		int64 GetDefaultValue() const
+		{
+			if (ChannelFunction.IsValid() &&
+				ChannelFunction->Default.IsSet())
+			{
+				return static_cast<int64>(ChannelFunction->Default.AsIntChecked());
+			}
+
+			return 0;
+		}
+
+		bool GetMinOffsetOfDMXChannel(uint32& OutChannelOffset) const
+		{
+			if (DMXChannelNode->Offset.IsEmpty())
+			{
+				return false;
+			}
+			OutChannelOffset = *Algo::MinElement(DMXChannelNode->Offset);
+
+			return true;
+		}
+
+		bool GetMaxOffsetOfDMXChannel(uint32& OutChannelOffset) const
+		{
+			if (DMXChannelNode->Offset.IsEmpty())
+			{
+				return false;
+			}
+			OutChannelOffset = *Algo::MaxElement(DMXChannelNode->Offset);
+
+			return true;
+		}
+
+		bool GetMinDMXBreak(const TArray<TSharedPtr<FDMXGDTFGeometryBreak>>& InDMXBreakArray, uint8& OutMinBreak) const
+		{
+			if (InDMXBreakArray.IsEmpty())
+			{
+				return false;
+			}
+			const TSharedPtr<FDMXGDTFGeometryBreak>* MinBreakPtr = Algo::MinElementBy(InDMXBreakArray, &FDMXGDTFGeometryBreak::DMXBreak);
+
+			OutMinBreak = (*MinBreakPtr)->DMXBreak;
+			return true;
+		}
+
+		bool GetMaxDMXBreak(const TArray<TSharedPtr<FDMXGDTFGeometryBreak>>& InDMXBreakArray, uint8& OutMaxBreak) const
+		{
+			if (InDMXBreakArray.IsEmpty())
+			{
+				return false;
+			}
+			const TSharedPtr<FDMXGDTFGeometryBreak>* MaxBreakPtr = Algo::MaxElementBy(InDMXBreakArray, &FDMXGDTFGeometryBreak::DMXBreak);
+
+			OutMaxBreak = (*MaxBreakPtr)->DMXBreak;
+			return true;
+		}
+
+		/** Gets the matrix offset attribute value from a Channel Node. Returns false if no offset can be found. */
+		bool GetMatrixOffset(const TArray<TSharedPtr<FDMXGDTFGeometryReference>>& GeometryReferences, uint32& OutOffset) const
+		{
+			const FName GeometryName = DMXChannelNode->Geometry;
+
+			uint32 ChannelOffset;
+			if (!GetMinOffsetOfDMXChannel(ChannelOffset))
+			{
+				return false;
+			}
+			TArray<uint32> OffsetArray;
+			for (const TSharedPtr<FDMXGDTFGeometryReference>& GeometryReference : GeometryReferences)
+			{
+				uint8 DMXBreak;
+				if (GeometryReference.IsValid() &&
+					GetMinDMXBreak(GeometryReference->BreakArray, DMXBreak))
+				{
+					OffsetArray.Add(DMXBreak);
+				}
+			}
+
+			if (OffsetArray.IsEmpty())
+			{
+				return false;
+			}
+			
+			OutOffset = *Algo::MinElement(OffsetArray) + ChannelOffset - 1;
+			return true;
+		}
+
+		/** Returns the data type offset array implies */
+		void GetDataType(EDMXFixtureSignalFormat& OutSignalFormat, bool& OutLSBOrder) const
+		{
+			OutLSBOrder = false;
+
+			// Compute number of used addresses
+			uint32 OffsetMin;
+			uint32 OffsetMax;
+			if (!GetMinOffsetOfDMXChannel(OffsetMin) ||
+				!GetMaxOffsetOfDMXChannel(OffsetMax))
+			{
+				return;
+			}
+
+			const uint32 NumUsedAddresses = FMath::Clamp(OffsetMax - OffsetMin + 1, 1, DMX_MAX_FUNCTION_SIZE);
+
+			OutSignalFormat = static_cast<EDMXFixtureSignalFormat>(NumUsedAddresses - 1);
+
+			// Offsets represent the channels in MSB order. If they are in reverse order, it means this Function uses LSB format.
+			if (DMXChannelNode->Offset.Num() > 1)
+			{
+				OutLSBOrder = DMXChannelNode->Offset[0] > DMXChannelNode->Offset[1];
+			}
+		}
+
+		TSharedPtr<FDMXGDTFLogicalChannel> LogicalChannel;
+		TSharedPtr<FDMXGDTFChannelFunction> ChannelFunction;
+
+		const TSharedRef<FDMXGDTFDMXMode> DMXModeNode;
+		const TSharedRef<FDMXGDTFDMXChannel> DMXChannelNode; 
+	};
+
+	bool FDMXGDTFToFixtureTypeConverter::ConvertGDTF(UDMXEntityFixtureType& InOutFixtureType, const UDMXImportGDTF& InGDTF, const bool bUpdateFixtureTypeName)
+	{
+		FDMXGDTFToFixtureTypeConverter Instance;
+		return Instance.ConvertGDTFInternal(InOutFixtureType, InGDTF, bUpdateFixtureTypeName);
+	}
+
+	bool FDMXGDTFToFixtureTypeConverter::ConvertGDTFInternal(UDMXEntityFixtureType& InOutFixtureType, const UDMXImportGDTF& InGDTF, const bool bUpdateFixtureTypeName) const
+	{
+		UDMXGDTFAssetImportData* GDTFAssetImportData = InGDTF.GetGDTFAssetImportData();
+		if (!ensureMsgf(GDTFAssetImportData, TEXT("Found GDTF Asset that has no GDTF asset import data subobject.")))
+		{
+			return false;
+		}
+
+		const TSharedRef<FDMXZipper> GDTFZip = MakeShared<FDMXZipper>();
+		TArray64<uint8> RawGDTFZip;
+		if (!GDTFZip->LoadFromData(GDTFAssetImportData->GetRawSourceData()) ||
+			!GDTFZip->GetData(RawGDTFZip))
+		{
+			return false;
+		}
+
+		UDMXGDTF* DMXGDTF = NewObject<UDMXGDTF>();
+		DMXGDTF->InitializeFromData(RawGDTFZip);
+
+		const TSharedPtr<FDMXGDTFFixtureType> GDTFFixtureType = DMXGDTF->GetDescription().IsValid() ? DMXGDTF->GetDescription()->GetFixtureType() : nullptr;
+		if (!GDTFFixtureType.IsValid())
+		{
+			return false;
+		}
+
+		if (bUpdateFixtureTypeName && GDTFFixtureType->Name.GetStringLength() > 0)
+		{
+			InOutFixtureType.Name = GDTFFixtureType->Name.ToString();
+		}
+
+		for (const TSharedPtr<FDMXGDTFDMXMode>& GDTFDMXMode : GDTFFixtureType->DMXModes)
+		{
+			FDMXFixtureMode Mode;
+			if (GenerateMode(GDTFFixtureType, GDTFDMXMode, Mode))
+			{
+				InOutFixtureType.Modes.Add(Mode);
+			}
+		}
+
+		for (int32 ModeIndex = 0; ModeIndex < InOutFixtureType.Modes.Num(); ModeIndex++)
+		{
+			InOutFixtureType.UpdateChannelSpan(ModeIndex);
+		}
+		InOutFixtureType.GetOnFixtureTypeChanged().Broadcast(&InOutFixtureType);
+
+		CleanupAttributes(InOutFixtureType);
+
+		return true;
+	}
+
+	bool FDMXGDTFToFixtureTypeConverter::GenerateMode(const TSharedPtr<FDMXGDTFFixtureType>& InFixtureTypeNode, const TSharedPtr<FDMXGDTFDMXMode>& InDMXModeNode, FDMXFixtureMode& OutMode) const
+	{
+		if (!ensureMsgf(InFixtureTypeNode.IsValid(), TEXT("Cannot create mode for GDTF Fixture Type. Fixture Type is invalid.")) ||
+			!ensureMsgf(InDMXModeNode.IsValid(), TEXT("Cannot create mode for GDTF Fixture Type. DMX Mode is invalid.")))
+		{
+			return false;
+		}
+
+		FDMXFixtureMode Mode;
+		Mode.ModeName = InDMXModeNode->Name.ToString();
+
+		for (const TSharedPtr<FDMXGDTFDMXChannel>& DMXChannelNode : InDMXModeNode->DMXChannels)
+		{
+			FDMXGDTFChannelInterpreter::FChannelProperties ChannelProperties;
+			if (!FDMXGDTFChannelInterpreter::GetChannelProperties(InDMXModeNode.ToSharedRef(), DMXChannelNode.ToSharedRef(), ChannelProperties))
+			{
+				continue;
+			}
+
+			const bool bIsMatrix = ChannelProperties.NumCells > 1;
+			if (bIsMatrix)
+			{
+				if (!Mode.bFixtureMatrixEnabled)
+				{
+					Mode.FixtureMatrixConfig.FirstCellChannel = ChannelProperties.FirstChannel;
+					Mode.bFixtureMatrixEnabled = true;
+
+					// Fixture matrix has one cell attribute by default, clear it.
+					Mode.FixtureMatrixConfig.CellAttributes.Reset();
+				}
+
+				FDMXFixtureCellAttribute MatrixAttribute;
+				MatrixAttribute.Attribute = FDMXAttributeName(*ChannelProperties.AttributeName);
+				MatrixAttribute.bUseLSBMode = ChannelProperties.bLSBMode;
+				MatrixAttribute.DataType = ChannelProperties.SignalFormat;
+				MatrixAttribute.DefaultValue = ChannelProperties.DefaultValue;
+
+				Mode.FixtureMatrixConfig.CellAttributes.Add(MatrixAttribute);
+				Mode.FixtureMatrixConfig.YCells = ChannelProperties.NumCells;
+			}
+			else
+			{
+				FDMXFixtureFunction Function;
+				Function.FunctionName = ChannelProperties.AttributeName;
+				Function.Attribute = FDMXAttributeName(*ChannelProperties.AttributeName);
+				Function.Channel = ChannelProperties.FirstChannel;
+				Function.bUseLSBMode = ChannelProperties.bLSBMode;
+				Function.DataType = ChannelProperties.SignalFormat;
+				Function.DefaultValue = ChannelProperties.DefaultValue;
+
+				Mode.Functions.Add(Function);
+			}
+		}
+
+		CleanupMode(Mode);
+
+		OutMode = Mode;
+		return true;
+	}
+
+	void FDMXGDTFToFixtureTypeConverter::CleanupMode(FDMXFixtureMode& InOutMode) const
+	{
+		TOptional<TRange<int32>> MatrixRange;
+		if (InOutMode.bFixtureMatrixEnabled)
+		{
+			// Only one single, consecutive matrix is supported by the engine in this version.
+			MatrixRange = TRange<int32>(InOutMode.FixtureMatrixConfig.FirstCellChannel, InOutMode.FixtureMatrixConfig.GetLastChannel() + 1);
+			const FDMXFixtureFunction* OverlappingFunction = Algo::FindByPredicate(InOutMode.Functions, [MatrixRange](const FDMXFixtureFunction& Function)
+				{
+					const TRange<int32> FunctionRange(Function.Channel, Function.GetLastChannel() + 1);
+					return FunctionRange.Overlaps(MatrixRange.GetValue());
+				});
+
+			if (OverlappingFunction)
+			{
+				UE_LOG(LogDMXEditor, Warning, TEXT("Mode '%s' contains many matrices, but this version of Unreal Engine only supports one matrix. Skipping import of mode."), *InOutMode.ModeName);
+				InOutMode.Functions.Reset();
+				InOutMode.bFixtureMatrixEnabled = false;
+				InOutMode.FixtureMatrixConfig.CellAttributes.Reset();
+				InOutMode.ModeName = FString::Printf(TEXT("n/a '%s' [not supported in this Engine Version]"), *InOutMode.ModeName);
+			}
+		}
+
+		// Make sure functions are in consecutive order, insert 'reserved' channels where no channel is specified
+		if (!InOutMode.Functions.IsEmpty())
+		{
+			Algo::SortBy(InOutMode.Functions, [](const FDMXFixtureFunction& Function)
+				{
+					return Function.Channel;
+				});
+
+			// Fill in blank functions
+			const int32 LastFunctionChannel = InOutMode.Functions.Last().Channel;
+			for (int32 Channel = 1; Channel < LastFunctionChannel; Channel++)
+			{
+				const FDMXFixtureFunction* FunctionOnChannelPtr = Algo::FindByPredicate(InOutMode.Functions, [Channel](const FDMXFixtureFunction& Function)
+					{
+						const TRange<int32> FunctionRange(Function.Channel, Function.GetLastChannel() + 1);
+						return FunctionRange.Contains(Channel);
+					});
+
+				const bool bChannelHasFunction = FunctionOnChannelPtr != nullptr;
+				const bool bChannelHasMatrix = MatrixRange.IsSet() && MatrixRange.GetValue().Contains(Channel);
+
+				if (!bChannelHasFunction && !bChannelHasMatrix)
+				{
+					FDMXFixtureFunction EmptyFunction;
+					EmptyFunction.Channel = Channel;
+					EmptyFunction.FunctionName = TEXT("<empty>");
+					InOutMode.Functions.Add(EmptyFunction);
+				}
+			}
+		}
+
+		Algo::SortBy(InOutMode.Functions, [](const FDMXFixtureFunction& Function)
+			{
+				return Function.Channel;
+			});
+	}
+
+	void FDMXGDTFToFixtureTypeConverter::CleanupAttributes(UDMXEntityFixtureType& InOutFixtureType) const
+	{
+		// Get Protocol Setting's default attributes
+		const UDMXProtocolSettings* ProtocolSettings = GetDefault<UDMXProtocolSettings>();
+		TMap<FName, TArray<FString> > SettingsAttributeNameToKeywordsMap;
+		for (const FDMXAttribute& Attribute : ProtocolSettings->Attributes)
+		{
+			TArray<FString> Keywords = Attribute.GetKeywords();
+
+			SettingsAttributeNameToKeywordsMap.Emplace(Attribute.Name, Keywords);
+		}
+
+		for (FDMXFixtureMode& Mode : InOutFixtureType.Modes)
+		{
+			TArray<FName> AssignedAttributeNames;
+			for (FDMXFixtureFunction& Function : Mode.Functions)
+			{
+				const TTuple<FName, TArray<FString>>* SettingsAttributeNameToKeywordPairPtr = Algo::FindByPredicate(SettingsAttributeNameToKeywordsMap, [Function](const TTuple<FName, TArray<FString>> AttributeNameToKeywordPair)
+					{
+						if (Function.Attribute.Name == AttributeNameToKeywordPair.Key ||
+							AttributeNameToKeywordPair.Value.Contains(Function.Attribute.Name.ToString()))
+						{
+							return true;
+						}
+
+						return false;
+					});
+
+				if (SettingsAttributeNameToKeywordPairPtr && !AssignedAttributeNames.Contains(SettingsAttributeNameToKeywordPairPtr->Key))
+				{
+					Function.Attribute.Name = SettingsAttributeNameToKeywordPairPtr->Key;
+					AssignedAttributeNames.Add(SettingsAttributeNameToKeywordPairPtr->Key);
+				}
+				else
+				{
+					AssignedAttributeNames.Add(Function.Attribute.Name);
+				}
+			}
+
+			if (Mode.bFixtureMatrixEnabled)
+			{
+				TArray<FName> AssignedCellAttributeNames;
+				for (FDMXFixtureCellAttribute& CellAttribute : Mode.FixtureMatrixConfig.CellAttributes)
+				{
+					const TTuple<FName, TArray<FString>>* SettingsAttributeNameToKeywordPairPtr = Algo::FindByPredicate(SettingsAttributeNameToKeywordsMap, [CellAttribute](const TTuple<FName, TArray<FString>> AttributeNameToKeywordPair)
+						{
+							if (CellAttribute.Attribute.Name == AttributeNameToKeywordPair.Key ||
+								AttributeNameToKeywordPair.Value.Contains(CellAttribute.Attribute.Name.ToString()))
+							{
+								return true;
+							}
+
+							return false;
+						});
+
+					if (SettingsAttributeNameToKeywordPairPtr && !AssignedAttributeNames.Contains(SettingsAttributeNameToKeywordPairPtr->Key))
+					{
+						CellAttribute.Attribute.Name = SettingsAttributeNameToKeywordPairPtr->Key;
+						AssignedAttributeNames.Add(SettingsAttributeNameToKeywordPairPtr->Key);
+					}
+					else
+					{
+						AssignedAttributeNames.Add(CellAttribute.Attribute.Name);
+					}
+				}
+			}
+		}
+	}
+
+	TArray<const FXmlNode*> FDMXGDTFToFixtureTypeConverter::GetChildrenRecursive(const FXmlNode& ParentNode) const
+	{
+		TArray<const FXmlNode*> Result;
+
+		for (const FXmlNode* Child : ParentNode.GetChildrenNodes())
+		{
+			Result.Add(Child);
+			Result.Append(GetChildrenRecursive(*Child));
+		}
+
+		return Result;
+	}
+}
+
+#undef LOCTEXT_NAMESPACE
diff --git a/Engine/Plugins/VirtualProduction/DMX/DMXEngine/Source/DMXEditor/Private/Factories/DMXGDTFToFixtureTypeConverter.h b/Engine/Plugins/VirtualProduction/DMX/DMXEngine/Source/DMXEditor/Private/Factories/DMXGDTFToFixtureTypeConverter.h
new file mode 100644
--- /dev/null
+++ b/Engine/Plugins/VirtualProduction/DMX/DMXEngine/Source/DMXEditor/Private/Factories/DMXGDTFToFixtureTypeConverter.h
@@ -0,0 +1,50 @@
+// Copyright Epic Games, Inc. All Rights Reserved.
+
+#pragma once
+
+#include "CoreMinimal.h"
+
+class UDMXEntityFixtureType;
+struct FDMXFixtureFunction;
+struct FDMXFixtureMode;
+class FDMXZipper;
+class FXmlNode;
+class UDMXImportGDTF;
+
+namespace UE::DMX::GDTF
+{
+	class FDMXGDTFFixtureType;
+	class FDMXGDTFDMXMode;
+
+	/**
+	 * This is a 5.1 private editor only workaround to initialize Fixture Types from GDTF assets.
+	 *
+	 * Required since DMXImportGDTF does not hold required data, hence UDMXEntityFixtureType::SetModesFromDMXImport cannot create matrix fixture types.
+	 *
+	 * Opeted for this workaround since it would require to rework the entirety of the GDTF implementation to hold its data correctly.
+	 * This exceeds planed features of 5.1 and is a planned feature for 5.2.
+	 */
+	class FDMXGDTFToFixtureTypeConverter
+	{
+	public:
+		/** Imports the modes directly from the asset import data */
+		static bool ConvertGDTF(UDMXEntityFixtureType& FixtureType, const UDMXImportGDTF& InGDTF, const bool bUpdateFixtureTypeName);
+
+	private:
+		/** Imports the modes directly from the asset import data. Non-static implementation */
+		bool ConvertGDTFInternal(UDMXEntityFixtureType& InOutFixtureType, const UDMXImportGDTF& InGDTF, const bool bUpdateFixtureTypeName) const;
+
+		/** Generates a Mode from specified Nodes. Returns true on success. */
+		bool GenerateMode(const TSharedPtr<FDMXGDTFFixtureType>& InFixtureTypeNode, const TSharedPtr<FDMXGDTFDMXMode>& InDMXModeNode, FDMXFixtureMode& OutMode) const;
+
+		/** Makes sure the mode can be used by the engine */
+		void CleanupMode(FDMXFixtureMode& InOutMode) const;
+
+		/** Makes sure that attributes are replaced with attributes specified in protocol settings where possible */
+		void CleanupAttributes(UDMXEntityFixtureType& InOutFixtureType) const;
+
+		/** Returns all nodes that are decendants of the node */
+		TArray<const FXmlNode*> GetChildrenRecursive(const FXmlNode& ParentNode) const;
+	};
+
+}
diff --git a/Engine/Plugins/VirtualProduction/DMX/DMXEngine/Source/DMXEditor/Private/Factories/DMXLibraryFromMVRImporter.cpp b/Engine/Plugins/VirtualProduction/DMX/DMXEngine/Source/DMXEditor/Private/Factories/DMXLibraryFromMVRImporter.cpp
--- a/Engine/Plugins/VirtualProduction/DMX/DMXEngine/Source/DMXEditor/Private/Factories/DMXLibraryFromMVRImporter.cpp
+++ b/Engine/Plugins/VirtualProduction/DMX/DMXEngine/Source/DMXEditor/Private/Factories/DMXLibraryFromMVRImporter.cpp
@@ -2,31 +2,28 @@
 
 #include "Factories/DMXLibraryFromMVRImporter.h"
 
+#include "Algo/Find.h"
+#include "AssetRegistry/AssetRegistryModule.h"
+#include "AssetToolsModule.h"
 #include "DMXEditorLog.h"
-#include "DMXInitializeFixtureTypeFromGDTFHelper.h"
+#include "DMXGDTFToFixtureTypeConverter.h"
 #include "DMXZipper.h"
 #include "Factories/DMXGDTFFactory.h"
 #include "Factories/DMXLibraryFromMVRImportOptions.h"
-#include "Library/DMXEntityFixtureType.h"
+#include "Framework/Application/SlateApplication.h"
+#include "Interfaces/IMainFrameModule.h"
+#include "Layout/SlateRect.h"
 #include "Library/DMXEntityFixturePatch.h"
+#include "Library/DMXEntityFixtureType.h"
 #include "Library/DMXGDTFAssetImportData.h"
 #include "Library/DMXImportGDTF.h"
 #include "Library/DMXLibrary.h"
 #include "MVR/DMXMVRAssetImportData.h"
 #include "MVR/DMXMVRGeneralSceneDescription.h"
 #include "MVR/Types/DMXMVRFixtureNode.h"
-#include "Widgets/SDMXLibraryFromMVRImportOptions.h"
-
-#include "AssetToolsModule.h"
 #include "ObjectTools.h"
+#include "Widgets/SDMXLibraryFromMVRImportOptions.h"
 #include "XmlFile.h"
-#include "Algo/Find.h"
-#include "AssetRegistry/AssetRegistryModule.h"
-#include "Framework/Application/SlateApplication.h"
-#include "Interfaces/IMainFrameModule.h"
-#include "Layout/SlateRect.h"
-#include "Misc/MessageDialog.h"
-
 
 #define LOCTEXT_NAMESPACE "DMXLibraryFromMVRImporter"
 
@@ -197,6 +194,8 @@ void FDMXLibraryFromMVRImporter::Reimport(UDMXLibrary* InDMXLibrary, UDMXLibrary
 
 void FDMXLibraryFromMVRImporter::InitializeDMXLibrary(UDMXLibrary* DMXLibrary, const TArray<UDMXImportGDTF*>& GDTFAssets)
 {
+	using namespace UE::DMX::GDTF;
+
 	check(DMXLibrary);
 	check(Zip.IsValid());
 	check(GeneralSceneDescription);
@@ -234,16 +233,11 @@ void FDMXLibraryFromMVRImporter::InitializeDMXLibrary(UDMXLibrary* DMXLibrary, c
 			FixtureTypeConstructionParams.ParentDMXLibrary = DMXLibrary;
 
 			UDMXEntityFixtureType* NewFixtureType = UDMXEntityFixtureType::CreateFixtureTypeInLibrary(FixtureTypeConstructionParams, FPaths::GetBaseFilename(GDTFFilename));
-			const bool bAdvancedImportSuccess = FDMXInitializeFixtureTypeFromGDTFHelper::GenerateModesFromGDTF(*NewFixtureType, *GDTF);
-			if (!bAdvancedImportSuccess)
-			{
-				PRAGMA_DISABLE_DEPRECATION_WARNINGS
-				UE_LOG(LogDMXEditor, Warning, TEXT("Failed to initialize Fixture Type '%s', falling back to legacy method that doesn't support matrix fixtures."), *NewFixtureType->GetName());
-				NewFixtureType->SetModesFromDMXImport(GDTF);
-				PRAGMA_ENABLE_DEPRECATION_WARNINGS
-			}
 			NewFixtureType->GDTFSource = GDTF;
 
+			constexpr bool bUpdateFixtureTypeName = true;
+			FDMXGDTFToFixtureTypeConverter::ConvertGDTF(*NewFixtureType, *GDTF, bUpdateFixtureTypeName);
+
 			GDTFSpecToFixtureTypeMap.Add(GDTFFilename, NewFixtureType);
 		}
 	}
@@ -268,6 +262,11 @@ void FDMXLibraryFromMVRImporter::InitializeDMXLibrary(UDMXLibrary* DMXLibrary, c
 	// Update or create new Fixture Patches for the MVR Fixtures
 	for (const UDMXMVRFixtureNode* FixtureNode : FixtureNodes)
 	{
+		if (!FixtureNode)
+		{
+			continue;
+		}
+
 		// Find the GDTF Filename, irregardless if it has the .gdtf extension
 		FString GDTFFilename = FixtureNode->GDTFSpec;
 		if (!GDTFSpecToFixtureTypeMap.Contains(GDTFFilename))
@@ -302,6 +301,12 @@ void FDMXLibraryFromMVRImporter::InitializeDMXLibrary(UDMXLibrary* DMXLibrary, c
 
 		if (ActiveModeIndex != INDEX_NONE)
 		{
+			const int32 DesiredFixtureID = [FixtureNode]()
+				{
+					int32 Value;
+					return LexTryParseString(Value, *FixtureNode->FixtureID) ? Value : 1;
+				}();
+
 			UDMXEntityFixturePatch* const* ExistingFixturePatchPtr = Algo::FindByPredicate(FixturePatches, [FixtureNode](const UDMXEntityFixturePatch* FixturePatch)
 				{
 					return FixturePatch->GetMVRFixtureUUID() == FixtureNode->UUID;
@@ -312,6 +317,7 @@ void FDMXLibraryFromMVRImporter::InitializeDMXLibrary(UDMXLibrary* DMXLibrary, c
 				UDMXEntityFixturePatch& ExistingFixturePatch = **ExistingFixturePatchPtr;
 
 				ExistingFixturePatch.SetFixtureType(FixtureType);
+				ExistingFixturePatch.GenerateFixtureID(DesiredFixtureID);
 				ExistingFixturePatch.SetActiveModeIndex(ActiveModeIndex);
 				ExistingFixturePatch.SetUniverseID(FixtureNode->GetUniverseID());
 				ExistingFixturePatch.SetStartingChannel(FixtureNode->GetStartingChannel());
@@ -328,6 +334,7 @@ void FDMXLibraryFromMVRImporter::InitializeDMXLibrary(UDMXLibrary* DMXLibrary, c
 				FixturePatchConstructionParams.MVRFixtureUUID = FixtureNode->UUID;
 
 				UDMXEntityFixturePatch* FixturePatch = UDMXEntityFixturePatch::CreateFixturePatchInLibrary(FixturePatchConstructionParams, FixtureNode->Name);
+				FixturePatch->GenerateFixtureID(DesiredFixtureID);
 				FixturePatch->EditorColor = FixtureTypeToColorMap.FindChecked(FixtureType);
 			}
 		}
diff --git a/Engine/Plugins/VirtualProduction/DMX/DMXEngine/Source/DMXRuntime/Private/Library/DMXEntityFixturePatch.cpp b/Engine/Plugins/VirtualProduction/DMX/DMXEngine/Source/DMXRuntime/Private/Library/DMXEntityFixturePatch.cpp
--- a/Engine/Plugins/VirtualProduction/DMX/DMXEngine/Source/DMXRuntime/Private/Library/DMXEntityFixturePatch.cpp
+++ b/Engine/Plugins/VirtualProduction/DMX/DMXEngine/Source/DMXRuntime/Private/Library/DMXEntityFixturePatch.cpp
@@ -4,6 +4,7 @@
 
 #include "Algo/Find.h"
 #include "Algo/MaxElement.h"
+#include "Algo/NoneOf.h"
 #include "DMXConversions.h"
 #include "DMXProtocolConstants.h"
 #include "DMXRuntimeLog.h"
@@ -11,13 +12,11 @@
 #include "DMXRuntimeUtils.h"
 #include "DMXStats.h"
 #include "DMXTypes.h"
-#include "Interfaces/IDMXProtocol.h"
 #include "IO/DMXInputPort.h"
 #include "IO/DMXOutputPort.h"
 #include "IO/DMXTrace.h"
 #include "Library/DMXEntityController.h"
 #include "Library/DMXEntityFixtureType.h"
-#include "Library/DMXImportGDTF.h"
 #include "Library/DMXLibrary.h"
 #include "Modulators/DMXModulator.h"
 #include "MVR/Types/DMXMVRFixtureNode.h"
@@ -28,7 +27,6 @@ DECLARE_LOG_CATEGORY_CLASS(DMXEntityFixturePatchLog, Log, All);
 DECLARE_CYCLE_STAT(TEXT("FixturePatch receive DMX"), STAT_DMXFixturePatchReceiveDMX, STATGROUP_DMX);
 DECLARE_CYCLE_STAT(TEXT("FixturePatch cache values"), STAT_DMXFixturePatchCacheValues, STATGROUP_DMX);
 
-
 #define LOCTEXT_NAMESPACE "DMXEntityFixturePatch"
 
 FDMXOnFixturePatchChangedDelegate UDMXEntityFixturePatch::OnFixturePatchChangedDelegate;
@@ -92,7 +90,7 @@ UDMXEntityFixturePatch* UDMXEntityFixturePatch::CreateFixturePatchInLibrary(FDMX
 			}
 			else
 			{
-				// No MVR Fixture UUID specified, generate one.
+				// If no MVR Fixture UUID is specified, generate one.
 				ConstructionParams.MVRFixtureUUID = FGuid::NewGuid();
 			}
 			NewFixturePatch->MVRFixtureUUID = ConstructionParams.MVRFixtureUUID;
@@ -200,6 +198,29 @@ void UDMXEntityFixturePatch::PostLoad()
 	if (!HasAnyFlags(RF_ClassDefaultObject | RF_ArchetypeObject))
 	{		
 		RebuildCache();
+		
+		// Upgrade Fixture ID and UUID from deprecated general scene description member in DMX Library
+		PRAGMA_DISABLE_DEPRECATION_WARNINGS
+		UDMXLibrary* DMXLibrary = GetParentLibrary();
+		UDMXMVRGeneralSceneDescription* GeneralSceneDescription = DMXLibrary ? DMXLibrary->GetLazyGeneralSceneDescription() : nullptr;
+		UDMXMVRFixtureNode* FixtureNode = GeneralSceneDescription ? GeneralSceneDescription->FindFixtureNode(MVRFixtureUUID) : nullptr;
+		if (FixtureNode)
+		{
+			MVRFixtureUUID = FixtureNode->UUID;
+
+			int32 IntegralMVRFixtureID;
+			if (LexTryParseString(IntegralMVRFixtureID, *FixtureNode->FixtureID))
+			{
+				FixtureID = IntegralMVRFixtureID;
+			}
+		}
+		PRAGMA_ENABLE_DEPRECATION_WARNINGS
+
+		// Mend invalid fixture IDs
+		if (FixtureID < 1)
+		{
+			GenerateFixtureID();
+		}
 	}
 }
 
@@ -585,7 +606,7 @@ bool UDMXEntityFixturePatch::SetActiveModeIndex(int32 NewActiveModeIndex)
 	return false;
 }
 
-void UDMXEntityFixturePatch::GenerateFixtureID()
+void UDMXEntityFixturePatch::GenerateFixtureID(int32 DesiredFixtureID)
 {
 	if (!ParentLibrary.IsValid())
 	{
@@ -594,6 +615,25 @@ void UDMXEntityFixturePatch::GenerateFixtureID()
 	}
 
 	const TArray<UDMXEntityFixturePatch*> FixturePatches = ParentLibrary->GetEntitiesTypeCast<UDMXEntityFixturePatch>();
+
+	// Try to use the desired fixture ID
+	if (DesiredFixtureID > 0)
+	{
+		const bool bCanUseDesiredFixtureID = Algo::NoneOf(FixturePatches, [DesiredFixtureID, this](const UDMXEntityFixturePatch* Other)
+			{
+				return 
+					Other &&
+					Other != this &&
+					Other->GetFixtureID() == DesiredFixtureID;
+			});
+		if (bCanUseDesiredFixtureID)
+		{
+			FixtureID = DesiredFixtureID;
+			return;
+		}
+	}
+
+	// Generate a new fixture ID
 	const UDMXEntityFixturePatch* const* MaxFixtureIDPatchPtr = Algo::MaxElementBy(FixturePatches, [this](const UDMXEntityFixturePatch* Other)
 		{
 			if (Other && Other != this)
diff --git a/Engine/Plugins/VirtualProduction/DMX/DMXEngine/Source/DMXRuntime/Private/Library/DMXEntityFixtureType.cpp b/Engine/Plugins/VirtualProduction/DMX/DMXEngine/Source/DMXRuntime/Private/Library/DMXEntityFixtureType.cpp
--- a/Engine/Plugins/VirtualProduction/DMX/DMXEngine/Source/DMXRuntime/Private/Library/DMXEntityFixtureType.cpp
+++ b/Engine/Plugins/VirtualProduction/DMX/DMXEngine/Source/DMXRuntime/Private/Library/DMXEntityFixtureType.cpp
@@ -3,9 +3,7 @@
 #include "Library/DMXEntityFixtureType.h"
 
 #include "DMXConversions.h"
-#include "DMXGDTF.h"
 #include "DMXProtocolSettings.h"
-#include "DMXRuntimeLog.h"
 #include "DMXRuntimeMainStreamObjectVersion.h"
 #include "DMXRuntimeUtils.h"
 #include "Library/DMXEntityFixturePatch.h"
@@ -13,8 +11,8 @@
 #include "Library/DMXImport.h"
 #include "Library/DMXImportGDTF.h"
 #include "Library/DMXLibrary.h"
-
-#include "Algo/Find.h"
+#include "Misc/EngineVersion.h"
+#include "Misc/Paths.h"
 
 #define LOCTEXT_NAMESPACE "DMXEntityFixtureType"
 
@@ -314,6 +312,29 @@ void UDMXEntityFixtureType::PostEditUndo()
 }
 #endif // WITH_EDITOR
 
+#if WITH_EDITOR
+FString UDMXEntityFixtureType::GetCleanGDTFFileNameSynchronous(bool bWithExtension) const
+{
+	if (!GDTFSource.IsNull())
+	{
+		if (UDMXImportGDTF* GDTF = GDTFSource.LoadSynchronous())
+		{
+			const FString GDTFFilename = FPaths::GetBaseFilename(GDTF->GetGDTFAssetImportData()->GetFilePathAndName());
+			const FString Extension = bWithExtension ? TEXT(".gdtf") : TEXT("");
+
+			return GDTFFilename + Extension;
+		}
+	}
+
+	const FString EngineVersion = FString::Printf(TEXT("%u_%u"), FEngineVersion::Current().GetMajor(), FEngineVersion::Current().GetMinor());
+	const FString DateTime = FDateTime::Now().ToString(TEXT("%d_%m_%y"));
+	const FString GDTFFilename = FString::Printf(TEXT("EpicGames@UE%s_Generated_%s@%s"), *EngineVersion, *Name, *DateTime);
+	const FString Extension = bWithExtension ? TEXT(".gdtf") : TEXT("");
+
+	return GDTFFilename + Extension;
+}
+#endif // WITH_EDITOR
+
 #if WITH_EDITOR
 void UDMXEntityFixtureType::SetModesFromDMXImport(UDMXImport* DMXImportAsset)
 {
diff --git a/Engine/Plugins/VirtualProduction/DMX/DMXEngine/Source/DMXRuntime/Private/MVR/DMXMVRGeneralSceneDescription.cpp b/Engine/Plugins/VirtualProduction/DMX/DMXEngine/Source/DMXRuntime/Private/MVR/DMXMVRGeneralSceneDescription.cpp
--- a/Engine/Plugins/VirtualProduction/DMX/DMXEngine/Source/DMXRuntime/Private/MVR/DMXMVRGeneralSceneDescription.cpp
+++ b/Engine/Plugins/VirtualProduction/DMX/DMXEngine/Source/DMXRuntime/Private/MVR/DMXMVRGeneralSceneDescription.cpp
@@ -2,30 +2,16 @@
 
 #include "MVR/DMXMVRGeneralSceneDescription.h"
 
-#include "Algo/MaxElement.h"
-#include "Algo/RemoveIf.h"
-#include "DMXProtocolCommon.h"
-#include "DMXRuntimeLog.h"
-#include "DMXRuntimeMainStreamObjectVersion.h"
-#include "EditorFramework/AssetImportData.h"
 #include "Library/DMXEntityFixturePatch.h"
 #include "Library/DMXEntityFixtureType.h"
-#include "Library/DMXGDTFAssetImportData.h"
-#include "Library/DMXImportGDTF.h"
 #include "Library/DMXLibrary.h"
-#include "Misc/Paths.h"
 #include "MVR/DMXMVRAssetImportData.h"
 #include "MVR/Types/DMXMVRChildListNode.h"
 #include "MVR/Types/DMXMVRFixtureNode.h"
-#include "MVR/Types/DMXMVRGroupObjectNode.h"
-#include "MVR/Types/DMXMVRLayerNode.h"
-#include "MVR/Types/DMXMVRLayersNode.h"
 #include "MVR/Types/DMXMVRParametricObjectNodeBase.h"
 #include "MVR/Types/DMXMVRRootNode.h"
-#include "MVR/Types/DMXMVRSceneNode.h"
 #include "XmlFile.h"
 
-
 #define LOCTEXT_NAMESPACE "DMXMVRGeneralSceneDescription"
 
 UDMXMVRGeneralSceneDescription::UDMXMVRGeneralSceneDescription()
@@ -144,7 +130,10 @@ TSharedPtr<FXmlFile> UDMXMVRGeneralSceneDescription::CreateXmlFile() const
 #if WITH_EDITOR
 void UDMXMVRGeneralSceneDescription::WriteFixturePatchToGeneralSceneDescription(const UDMXEntityFixturePatch& FixturePatch)
 {
-	checkf(RootNode, TEXT("Unexpected: MVR General Scene Description Root Node is invalid."));
+	if (!ensureMsgf(RootNode, TEXT("Unexpected: MVR General Scene Description Root Node is invalid.")))
+	{
+		return;
+	}
 
 	const UDMXLibrary* DMXLibrary = FixturePatch.GetParentLibrary();
 	if (!DMXLibrary)
@@ -183,22 +172,10 @@ void UDMXMVRGeneralSceneDescription::WriteFixturePatchToGeneralSceneDescription(
 	{
 		MVRFixtureNode->GDTFMode = FixtureType->Modes[ModeIndex].ModeName;
 
-		UDMXImportGDTF* GDTF = FixtureType->GDTFSource.LoadSynchronous();
-		if (GDTF)
-		{
-			const FString SourceFilename = [GDTF]()
-			{
-				if (GDTF && GDTF->GetGDTFAssetImportData())
-				{
-					return GDTF->GetGDTFAssetImportData()->GetFilePathAndName();
-				}
-				return FString();
-			}();
+		constexpr bool bWithExtension = false;
+		MVRFixtureNode->GDTFSpec = FixtureType->GetCleanGDTFFileNameSynchronous(bWithExtension);
 
-			MVRFixtureNode->GDTFSpec = FPaths::GetCleanFilename(SourceFilename);
-
-			bSetGDTFSpec = true;
-		}
+		bSetGDTFSpec = true;
 	}
 
 	if (!bSetGDTFSpec)
diff --git a/Engine/Plugins/VirtualProduction/DMX/DMXEngine/Source/DMXRuntime/Public/Library/DMXEntityFixturePatch.h b/Engine/Plugins/VirtualProduction/DMX/DMXEngine/Source/DMXRuntime/Public/Library/DMXEntityFixturePatch.h
--- a/Engine/Plugins/VirtualProduction/DMX/DMXEngine/Source/DMXRuntime/Public/Library/DMXEntityFixturePatch.h
+++ b/Engine/Plugins/VirtualProduction/DMX/DMXEngine/Source/DMXRuntime/Public/Library/DMXEntityFixturePatch.h
@@ -233,8 +233,11 @@ public:
 	/** Returns the MVR Fixture ÎD of this patch */
 	int32 GetFixtureID() const { return FixtureID; }
 
-	/** Generates a unique Fixture ID for this patch. */
-	void GenerateFixtureID();
+	/** 
+	 * Generates a unique Fixture ID for this patch. 
+	 * If DesiredFixtureID is > 0, tries to use this fixture ID, generates a unique one if the desired Fixture ID was already in use.
+	 */
+	void GenerateFixtureID(int32 DesiredFixtureID = -1);
 
 	/** 
 	 * Tries to find the fixture ID of the patch. Looks up the general scene description resulting in a relatively slow operation. 
@@ -249,6 +252,7 @@ public:
 	static FName GetParentFixtureTypeTemplatePropertyNameChecked() { return GET_MEMBER_NAME_CHECKED(UDMXEntityFixturePatch, ParentFixtureTypeTemplate); }
 	static FName GetActiveModePropertyNameChecked() { return GET_MEMBER_NAME_CHECKED(UDMXEntityFixturePatch, ActiveMode); }
 	static FName GetMVRFixtureUUIDPropertyNameChecked() { return GET_MEMBER_NAME_CHECKED(UDMXEntityFixturePatch, MVRFixtureUUID); }
+	static FName GetFixtureIDPropertyNameChecked() { return GET_MEMBER_NAME_CHECKED(UDMXEntityFixturePatch, FixtureID); }
 	static FName GetStartingChannelPropertyNameChecked() { return GET_MEMBER_NAME_CHECKED(UDMXEntityFixturePatch, StartingChannel); }
 
 	UE_DEPRECATED(5.1, "bAutoAssignAddress and related members are deprecated. Auto assign is now a method in FDMXEditorUtils and applied on demand.")
diff --git a/Engine/Plugins/VirtualProduction/DMX/DMXEngine/Source/DMXRuntime/Public/Library/DMXEntityFixtureType.h b/Engine/Plugins/VirtualProduction/DMX/DMXEngine/Source/DMXRuntime/Public/Library/DMXEntityFixtureType.h
--- a/Engine/Plugins/VirtualProduction/DMX/DMXEngine/Source/DMXRuntime/Public/Library/DMXEntityFixtureType.h
+++ b/Engine/Plugins/VirtualProduction/DMX/DMXEngine/Source/DMXRuntime/Public/Library/DMXEntityFixtureType.h
@@ -78,7 +78,7 @@ struct DMXRUNTIME_API FDMXFixtureFunction
 	UPROPERTY(EditAnywhere, BlueprintReadWrite, meta = (DisplayPriority = "20"), Category = "Function Settings")
 	FString Description;
 
-	/** The Default Value of the function, imported from GDTF. The plugin doesn't make use of this value, but it can be used in blueprints */
+	/** The Default Value of the function */
 	UPROPERTY(EditAnywhere, BlueprintReadWrite, meta = (DisplayPriority = "30"), Category = "Function Settings")
 	int64 DefaultValue;
 
@@ -314,8 +314,15 @@ public:
 #endif
 	//~ End UObject interface
 
-public:
 #if WITH_EDITOR
+	/**
+	 * Acquires the GDTF file name of this Fixture Type. If bWithExtension is true, appends the .gdtf extension
+	 *
+	 * Note this is a slow operation that will load the GDTF Source if required and look up the filename from asset import data.
+	 * For Fixture Types that do not stem from an imported GDTF, a filename is generated based on the Fixture Type name.
+	 */
+	FString GetCleanGDTFFileNameSynchronous(bool bWithExtension) const;
+
 	UE_DEPRECATED(5.5, "Setting GDTFs this way is not supported. Instead set the GDTFSource and generate the modes via UDMXGDTF.")
 	UFUNCTION(BlueprintCallable, Category = "Fixture Settings", Meta = (DeprecatedFunction, DeprecationMessage = "Setting GDTFs from blueprints was never fully supported and now deprecated. Instead please refer to the Create Fixture Type In DMX Library function."))
 	void SetModesFromDMXImport(UDMXImport* DMXImportAsset);
diff --git a/Engine/Plugins/VirtualProduction/DMX/DMXGDTF/Source/DMXGDTF/Internal/GDTF/AttributeDefinitions/DMXGDTFPhysicalUnit.h b/Engine/Plugins/VirtualProduction/DMX/DMXGDTF/Source/DMXGDTF/Internal/GDTF/AttributeDefinitions/DMXGDTFPhysicalUnit.h
--- a/Engine/Plugins/VirtualProduction/DMX/DMXGDTF/Source/DMXGDTF/Internal/GDTF/AttributeDefinitions/DMXGDTFPhysicalUnit.h
+++ b/Engine/Plugins/VirtualProduction/DMX/DMXGDTF/Source/DMXGDTF/Internal/GDTF/AttributeDefinitions/DMXGDTFPhysicalUnit.h
@@ -34,7 +34,7 @@ enum class EDMXGDTFPhysicalUnit : uint8
 	Volume UMETA(DisplayName = "Volume (m3)"),
 	Speed UMETA(DisplayName = "Speed (m/s)"),
 	Acceleration UMETA(DisplayName = "Acceleration (m/s2)"),
-	AngularSpeed UMETA(DisplayName = "ngularSpeed (degree/s)"),
+	AngularSpeed UMETA(DisplayName = "AngularSpeed (degree/s)"),
 	AngularAccc UMETA(DisplayName = "AngularAccc (degree/s2"),
 	WaveLength UMETA(DisplayName = "WaveLength (nm)"),
 	ColorComponent UMETA(DisplayName = "ColorComponent"),
diff --git a/Engine/Plugins/VirtualProduction/DMX/DMXGDTF/Source/DMXGDTF/Internal/GDTF/DMXGDTFDescription.h b/Engine/Plugins/VirtualProduction/DMX/DMXGDTF/Source/DMXGDTF/Internal/GDTF/DMXGDTFDescription.h
--- a/Engine/Plugins/VirtualProduction/DMX/DMXGDTF/Source/DMXGDTF/Internal/GDTF/DMXGDTFDescription.h
+++ b/Engine/Plugins/VirtualProduction/DMX/DMXGDTF/Source/DMXGDTF/Internal/GDTF/DMXGDTFDescription.h
@@ -31,6 +31,9 @@ namespace UE::DMX::GDTF
 		/** Exports an Xml file from the node tree in this description. Returns the XML file or nullptr if not a valid XML */
 		TSharedPtr<FXmlFile> ExportAsXml() const;
 
+		/** Returns the Fixture Type of this GDTF description */
+		TSharedPtr<FDMXGDTFFixtureType> GetFixtureType() const { return FixtureType; }
+
 	private:
 		/** The Fixture Type Child Node */
 		TSharedPtr<FDMXGDTFFixtureType> FixtureType;
diff --git a/Engine/Plugins/VirtualProduction/DMX/DMXGDTF/Source/DMXGDTF/Internal/GDTF/DMXGDTFFixtureType.h b/Engine/Plugins/VirtualProduction/DMX/DMXGDTF/Source/DMXGDTF/Internal/GDTF/DMXGDTFFixtureType.h
--- a/Engine/Plugins/VirtualProduction/DMX/DMXGDTF/Source/DMXGDTF/Internal/GDTF/DMXGDTFFixtureType.h
+++ b/Engine/Plugins/VirtualProduction/DMX/DMXGDTF/Source/DMXGDTF/Internal/GDTF/DMXGDTFFixtureType.h
@@ -29,7 +29,7 @@ namespace UE::DMX::GDTF
 		virtual FXmlNode* CreateXmlNode(FXmlNode& Parent) override;
 		//~ End FDMXGDTFNode interface
 
-		/** Name of the fixture type. As it is based on Name. [UE: GDTF version specific, unreal only parses the name] */
+		/** Name of the fixture type. As it is based on Name. */
 		FName Name;
 
 		/** Shortened name of the fixture type. Non detailed version or an abbreviation. Can use any characters or symbols */
diff --git a/Engine/Plugins/VirtualProduction/DMX/DMXGDTF/Source/DMXGDTF/Internal/GDTF/DMXGDTFNode.h b/Engine/Plugins/VirtualProduction/DMX/DMXGDTF/Source/DMXGDTF/Internal/GDTF/DMXGDTFNode.h
--- a/Engine/Plugins/VirtualProduction/DMX/DMXGDTF/Source/DMXGDTF/Internal/GDTF/DMXGDTFNode.h
+++ b/Engine/Plugins/VirtualProduction/DMX/DMXGDTF/Source/DMXGDTF/Internal/GDTF/DMXGDTFNode.h
@@ -29,7 +29,7 @@ namespace UE::DMX::GDTF
 		virtual void Initialize(const FXmlNode& InXmlNode) = 0;
 
 		/** Creates an XML node in the parent node */
-		virtual FXmlNode* CreateXmlNode(FXmlNode& Parent) { return nullptr; };
+		virtual FXmlNode* CreateXmlNode(FXmlNode& Parent) = 0;
 
 		/** Returns the fixture type this node resides in */
 		TWeakPtr<FDMXGDTFFixtureType> GetFixtureType() const { return WeakFixtureType; }
diff --git a/Engine/Plugins/VirtualProduction/DMX/DMXGDTF/Source/DMXGDTF/Internal/GDTF/DMXGDTFVersion.h b/Engine/Plugins/VirtualProduction/DMX/DMXGDTF/Source/DMXGDTF/Internal/GDTF/DMXGDTFVersion.h
--- a/Engine/Plugins/VirtualProduction/DMX/DMXGDTF/Source/DMXGDTF/Internal/GDTF/DMXGDTFVersion.h
+++ b/Engine/Plugins/VirtualProduction/DMX/DMXGDTF/Source/DMXGDTF/Internal/GDTF/DMXGDTFVersion.h
@@ -12,7 +12,7 @@ namespace UE::DMX::GDTF
 	{
 		static const uint8 MajorVersion = 1;
 
-		static const uint8 MinorVersion = 1;
+		static const uint8 MinorVersion = 2;
 
 		static const FString GetMajorVersionAsString();
 		static const FString GetMinorVersionAsString();
diff --git a/Engine/Plugins/VirtualProduction/DMX/DMXGDTF/Source/DMXGDTF/Internal/GDTF/DMXModes/DMXGDTFChannelFunction.h b/Engine/Plugins/VirtualProduction/DMX/DMXGDTF/Source/DMXGDTF/Internal/GDTF/DMXModes/DMXGDTFChannelFunction.h
--- a/Engine/Plugins/VirtualProduction/DMX/DMXGDTF/Source/DMXGDTF/Internal/GDTF/DMXModes/DMXGDTFChannelFunction.h
+++ b/Engine/Plugins/VirtualProduction/DMX/DMXGDTF/Source/DMXGDTF/Internal/GDTF/DMXModes/DMXGDTFChannelFunction.h
@@ -84,10 +84,10 @@ namespace UE::DMX::GDTF
 		FString ModeMaster;
 
 		/** Only used together with ModeMaster; DMX start value; Default value: 0/1 */
-		FDMXGDTFDMXValue ModeFrom = TEXT("0/1");
+		FDMXGDTFDMXValue ModeFrom = 0;
 
 		/** Only used together with ModeMaster; DMX start value; Default value: 0/1 */
-		FDMXGDTFDMXValue ModeTo = TEXT("0/1");
+		FDMXGDTFDMXValue ModeTo = 0;
 
 		/** (Optional) Link to DMX Profile; Starting point: DMX Profile Collect */
 		FString DMXProfile;
diff --git a/Engine/Plugins/VirtualProduction/DMX/DMXGDTF/Source/DMXGDTF/Internal/GDTF/DMXModes/DMXGDTFDMXChannel.h b/Engine/Plugins/VirtualProduction/DMX/DMXGDTF/Source/DMXGDTF/Internal/GDTF/DMXModes/DMXGDTFDMXChannel.h
--- a/Engine/Plugins/VirtualProduction/DMX/DMXGDTF/Source/DMXGDTF/Internal/GDTF/DMXModes/DMXGDTFDMXChannel.h
+++ b/Engine/Plugins/VirtualProduction/DMX/DMXGDTF/Source/DMXGDTF/Internal/GDTF/DMXModes/DMXGDTFDMXChannel.h
@@ -9,6 +9,8 @@ namespace UE::DMX::GDTF
 {
 	class FDMXGDTFChannelFunction;
 	class FDMXGDTFDMXMode;
+	class FDMXGDTFGeometry;
+	class FDMXGDTFGeometryReference;
 	class FDMXGDTFLogicalChannel;
 
 	/**
@@ -68,7 +70,7 @@ PRAGMA_ENABLE_DEPRECATION_WARNINGS
 		 * location, put it in the top level geometry of the geometry tree. Attributes follow a trickle down principle, so
 		 * they are inherited from top down
 		 */
-		FString Geometry;
+		FName Geometry;
 
 		/** A list of logical channels */
 		TArray<TSharedPtr<FDMXGDTFLogicalChannel>> LogicalChannelArray;
@@ -82,6 +84,17 @@ PRAGMA_ENABLE_DEPRECATION_WARNINGS
 		/** Resolves the linked initial function. Returns the initial function, or nullptr if no initial function is linked */
 		TSharedPtr<FDMXGDTFChannelFunction> ResolveInitialFunction() const;
 
+		/** 
+		 * Resolves the linked geometry. Returns the geometry, or nullptr if no geometry is linked. 
+		 * To resolve as geometry references, see ResolveGeometryReferences.
+		 * 
+		 * Note, GDTFs of older version directly reference a model as geometry. Such models are not considered.
+		 */
+		TSharedPtr<FDMXGDTFGeometry> ResolveGeometry() const;
+
+		/** Resolves the linked geometry as geometry references. Returns geometry references or an empty array, if no linked geometry references could be found. */
+		TArray<TSharedPtr<FDMXGDTFGeometryReference>> ResolveGeometryReferences() const;
+
 	private:
 		/** Converts a GDTF string to an array of offsets */
 		TArray<uint32> ParseOffset(const FString& GDTFString) const;
diff --git a/Engine/Plugins/VirtualProduction/DMX/DMXGDTF/Source/DMXGDTF/Internal/GDTF/DMXModes/DMXGDTFDMXMode.h b/Engine/Plugins/VirtualProduction/DMX/DMXGDTF/Source/DMXGDTF/Internal/GDTF/DMXModes/DMXGDTFDMXMode.h
--- a/Engine/Plugins/VirtualProduction/DMX/DMXGDTF/Source/DMXGDTF/Internal/GDTF/DMXModes/DMXGDTFDMXMode.h
+++ b/Engine/Plugins/VirtualProduction/DMX/DMXGDTF/Source/DMXGDTF/Internal/GDTF/DMXModes/DMXGDTFDMXMode.h
@@ -11,6 +11,7 @@ namespace UE::DMX::GDTF
 	class FDMXGDTFDMXChannel;
 	class FDMXGDTFFixtureType;
 	class FDMXGDTFFTMacro;
+	class FDMXGDTFGeometry;
 
 	/** Each DMX mode describes logical control a part of the device in a specific mode (XML node <DMXMode>). */
 	class DMXGDTF_API FDMXGDTFDMXMode
@@ -32,7 +33,7 @@ namespace UE::DMX::GDTF
 		FString Description;
 
 		/** The first geometry in the device; Only top level geometries are allowed to be linked	 */
-		FString Geometry;
+		FName Geometry;
 
 		/** Description of all DMX channels used in the mode */
 		TArray<TSharedPtr<FDMXGDTFDMXChannel>> DMXChannels;
@@ -46,6 +47,9 @@ namespace UE::DMX::GDTF
 		/** The outer fixture type */
 		const TWeakPtr<FDMXGDTFFixtureType> OuterFixtureType;
 
+		/** Resolves the Geometry for this mode */
+		TSharedPtr<FDMXGDTFGeometry> ResolveGeometry() const;
+
 		/** Resolves a DMX channel or channel function */
 		void ResolveChannel(const FString& Link, TSharedPtr<FDMXGDTFDMXChannel>& OutDMXChannel, TSharedPtr<FDMXGDTFChannelFunction>& OutChannelFunction) const;
 	};
diff --git a/Engine/Plugins/VirtualProduction/DMX/DMXGDTF/Source/DMXGDTF/Internal/GDTF/DMXModes/DMXGDTFDMXValue.h b/Engine/Plugins/VirtualProduction/DMX/DMXGDTF/Source/DMXGDTF/Internal/GDTF/DMXModes/DMXGDTFDMXValue.h
--- a/Engine/Plugins/VirtualProduction/DMX/DMXGDTF/Source/DMXGDTF/Internal/GDTF/DMXModes/DMXGDTFDMXValue.h
+++ b/Engine/Plugins/VirtualProduction/DMX/DMXGDTF/Source/DMXGDTF/Internal/GDTF/DMXModes/DMXGDTFDMXValue.h
@@ -4,6 +4,7 @@
 
 #include "Containers/UnrealString.h"
 #include "GenericPlatform/GenericPlatform.h"
+#include "Misc/Optional.h"
 
 namespace UE::DMX::GDTF
 {
@@ -15,11 +16,36 @@ namespace UE::DMX::GDTF
 	struct DMXGDTF_API FDMXGDTFDMXValue
 	{
 		FDMXGDTFDMXValue() = default;
-		FDMXGDTFDMXValue(const TCHAR* InValue);
+		FDMXGDTFDMXValue(const FString& InValue);
+		FDMXGDTFDMXValue(const uint32 InValue);
+		FDMXGDTFDMXValue(const TOptional<uint32>& InValue);
 
-		FString Value = TEXT("None");
+		bool operator==(const FDMXGDTFDMXValue& Other) const { return IntegerValue == Other.IntegerValue; }
+		bool operator!=(const FDMXGDTFDMXValue& Other) const { return IntegerValue != Other.IntegerValue; }
 
-		/** Returns the value as integer. Note this is a relatively slow operation as it parses the string. */
-		bool ToInt(uint32& OutInteger) const;
+		/** Returns the value as an optional integer. If the optional is not set, it equals to "None". */
+		TOptional<uint32> AsInt() const { return IntegerValue; }
+
+		/** Returns the value as an integer. Assumes the int is set (checked). */
+		uint32 AsIntChecked() const;
+
+		/** Returns the value as string. */
+		FString AsString() const { return StringValue; }
+
+		/** Returns true if the value is set. Otherwise it uses special value "None". */
+		bool IsSet() const;
+
+		/** Sets the DMX Value by int. */
+		void Set(uint32 InValue);
+
+		/** Sets the value by string. Special value "None" resets the optional int. */
+		void Set(const FString& InValue);
+
+		/** Resets to None. */
+		void Reset();
+
+	private:
+		FString StringValue = TEXT("None");
+		TOptional<uint32> IntegerValue;
 	};
 }
diff --git a/Engine/Plugins/VirtualProduction/DMX/DMXGDTF/Source/DMXGDTF/Internal/GDTF/Geometries/DMXGDTFBreak.h b/Engine/Plugins/VirtualProduction/DMX/DMXGDTF/Source/DMXGDTF/Internal/GDTF/Geometries/DMXGDTFBreak.h
deleted file mode 100644
diff --git a/Engine/Plugins/VirtualProduction/DMX/DMXGDTF/Source/DMXGDTF/Internal/GDTF/Geometries/DMXGDTFGeometry.h b/Engine/Plugins/VirtualProduction/DMX/DMXGDTF/Source/DMXGDTF/Internal/GDTF/Geometries/DMXGDTFGeometry.h
--- a/Engine/Plugins/VirtualProduction/DMX/DMXGDTF/Source/DMXGDTF/Internal/GDTF/Geometries/DMXGDTFGeometry.h
+++ b/Engine/Plugins/VirtualProduction/DMX/DMXGDTF/Source/DMXGDTF/Internal/GDTF/Geometries/DMXGDTFGeometry.h
@@ -26,14 +26,24 @@ namespace UE::DMX::GDTF
 		virtual FXmlNode* CreateXmlNode(FXmlNode& Parent) override;
 		//~ End DMXGDTFGeneralGeometryNode interface
 
-		/** The unique name of geometry. See standard for details of specific instances, e.g. Beam, Shaper Filter etc. */
+		//~ Begin DMXGDTFGeometryCollectBase interface
+		virtual TSharedPtr<FDMXGDTFGeometry> FindGeometryByName(const TCHAR* InName) const override;
+		//~ End DMXGDTFGeometryCollectBase interface
+
+		/** The unique name of geometry. */
 		FName Name;
 
-		/** Link to the corresponding model. */
-		FString Model;
+		/** Relative position of geometry; Default value : Identity Matrix */
+		FTransform Position;
 
-		/** Relative position of geometry */
-		FTransform Position = FTransform::Identity;
+		/**
+		 * (Optional) Link to the corresponding model. The model only replaces
+		 * the model of the parent of the referenced geometry. The models of
+		 * the children of the referenced geometry are not affected. The starting
+		 * point is Models Collect. If model is not set, the model is taken from
+		 * the referenced geometry.
+		 */
+		FString Model;
 
 		/** The outer geometry collect */
 		const TWeakPtr<FDMXGDTFGeometryCollectBase> OuterGeometryCollect;
diff --git a/Engine/Plugins/VirtualProduction/DMX/DMXGDTF/Source/DMXGDTF/Internal/GDTF/Geometries/DMXGDTFGeometryBreak.h b/Engine/Plugins/VirtualProduction/DMX/DMXGDTF/Source/DMXGDTF/Internal/GDTF/Geometries/DMXGDTFGeometryBreak.h
new file mode 100644
--- /dev/null
+++ b/Engine/Plugins/VirtualProduction/DMX/DMXGDTF/Source/DMXGDTF/Internal/GDTF/Geometries/DMXGDTFGeometryBreak.h
@@ -0,0 +1,34 @@
+// Copyright Epic Games, Inc. All Rights Reserved.
+
+#pragma once
+
+#include "GDTF/DMXGDTFNode.h"
+#include "GDTF/Geometries/DMXGDTFDMXAddress.h"
+
+namespace UE::DMX::GDTF
+{
+	class FDMXGDTFGeometryReference;
+
+	/** This XML node specifies the DMX offset for the DMX channel of the referenced geometry (XML node <Break>) */
+	class DMXGDTF_API FDMXGDTFGeometryBreak
+		: public FDMXGDTFNode
+	{
+	public:
+		FDMXGDTFGeometryBreak(const TSharedRef<FDMXGDTFGeometryReference>& InGeometryReference);
+
+		//~ Begin DMXGDTFGeneralGeometryNode interface
+		virtual const TCHAR* GetXmlTag() const override { return TEXT("Break"); }
+		virtual void Initialize(const FXmlNode& XmlNode) override;
+		virtual FXmlNode* CreateXmlNode(FXmlNode& Parent) override;
+		//~ End DMXGDTFGeneralGeometryNode interface
+
+		/** DMX offset; Default value:1 (Means no offset for the corresponding DMX Channel) */
+		FDMXGDTFDMXAddress DMXOffset;
+
+		/** Defines the unique number of the DMX Break for which the Offset is given. Size: 1 byte; Default value 1 */
+		uint8 DMXBreak = 1;
+
+		/** The outer geometry reference */
+		const TWeakPtr<FDMXGDTFGeometryReference> OuterGeometryReference;
+	};
+}
diff --git a/Engine/Plugins/VirtualProduction/DMX/DMXGDTF/Source/DMXGDTF/Internal/GDTF/Geometries/DMXGDTFGeometryCollectBase.h b/Engine/Plugins/VirtualProduction/DMX/DMXGDTF/Source/DMXGDTF/Internal/GDTF/Geometries/DMXGDTFGeometryCollectBase.h
--- a/Engine/Plugins/VirtualProduction/DMX/DMXGDTF/Source/DMXGDTF/Internal/GDTF/Geometries/DMXGDTFGeometryCollectBase.h
+++ b/Engine/Plugins/VirtualProduction/DMX/DMXGDTF/Source/DMXGDTF/Internal/GDTF/Geometries/DMXGDTFGeometryCollectBase.h
@@ -29,7 +29,7 @@ namespace UE::DMX::GDTF
 	class IDMXGDTFGeometryNodeInterface;
 
 	/** UE specific. Bases class for all classes that have a geometry collect. */
-	class FDMXGDTFGeometryCollectBase
+	class DMXGDTF_API FDMXGDTFGeometryCollectBase
 		: public FDMXGDTFNode
 	{
 	public:
@@ -39,7 +39,6 @@ namespace UE::DMX::GDTF
 		virtual FXmlNode* CreateXmlNode(FXmlNode& Parent) override;
 		//~ End FDMXGDTFNode interface
 
-	public:
 		/** Any General Geometry. */
 		TArray<TSharedPtr<FDMXGDTFGeometry>> GeometryArray;
 
@@ -63,7 +62,7 @@ namespace UE::DMX::GDTF
 
 		/** Any Geometry that describes a media representation layer of a media device. */
 		TArray<TSharedPtr<FDMXGDTFMediaServerLayerGeometry>> MediaServerLayerArray;
-		
+
 		/** Any Geometry that describes a camera or output layer of a media device. */
 		TArray<TSharedPtr<FDMXGDTFMediaServerCameraGeometry>> MediaServerCameraArray;
 
@@ -94,22 +93,33 @@ namespace UE::DMX::GDTF
 		/** Any Geometry that describes a point where other geometries should be attached */
 		TArray<TSharedPtr<FDMXGDTFMagnetGeometry>> MagnetArray;
 
-		/** Finds the geometry by name. Can be either a geometry or a geometry reference */
-		void FindGeometryByName(const TCHAR* InName, TSharedPtr<FDMXGDTFGeometry>& OutGeometry, TSharedPtr<FDMXGDTFGeometryReference>& OutGeometryReference) const;
+		/** Returns all child geometries. Does not include self. */
+		void GetGeometriesRecursive(TArray<TSharedPtr<FDMXGDTFGeometry>>& OutGeometries, TArray<TSharedPtr<FDMXGDTFGeometryReference>>& OutGeometryReferences);
+
+		/** Finds the geometry by name. */
+		virtual TSharedPtr<FDMXGDTFGeometry> FindGeometryByName(const TCHAR* InName) const;
+
+		/** Finds geometry references by name. */
+		TSharedPtr<FDMXGDTFGeometryReference> FindGeometryReferenceByName(const TCHAR* InName) const;
 
+		/**  
+		 * Resolves a string as a link to a geometry. 
+		 * 
+		 * The string needs to be formated in the form of "Geometry1.Geometry2.[...].GeometryN" whereas Geometry1 resides in the geometry collect of the fixture type. 
+		 */
 		template <typename GeometryType>
 		TSharedPtr<GeometryType> ResolveGeometryLink(const FString& Link) const
 		{
 			TArray<FString> LinkArray;
 			Link.ParseIntoArray(LinkArray, TEXT("."));
 
-			TSharedPtr<const FDMXGDTFGeometryCollectBase> OuterGeometryCollect = StaticCastSharedRef<const FDMXGDTFGeometryCollectBase>(SharedThis(this));
+			TSharedPtr<const FDMXGDTFGeometryCollectBase> NextGeometryCollect = StaticCastSharedRef<const FDMXGDTFGeometryCollectBase>(SharedThis(this));
 			for (const FString& GeometryName : LinkArray)
 			{
 				if (&LinkArray.Last() == &GeometryName)
 				{
 					TArray<TSharedPtr<GeometryType>> FinalGeometryArray;
-					OuterGeometryCollect->GetGeometriesOfType<GeometryType>(FinalGeometryArray);
+					NextGeometryCollect->GetGeometriesOfType<GeometryType>(FinalGeometryArray);
 					const TSharedPtr<GeometryType>* GeometryPtr = Algo::FindBy(FinalGeometryArray, GeometryName, &GeometryType::Name);
 					if (GeometryPtr)
 					{
@@ -117,18 +127,23 @@ namespace UE::DMX::GDTF
 					}
 				}
 
-				TSharedPtr<FDMXGDTFGeometry> Geometry;
-				TSharedPtr<FDMXGDTFGeometryReference> GeometryReference;
-				OuterGeometryCollect->FindGeometryByName(*GeometryName, Geometry, GeometryReference);
+				const TSharedPtr<FDMXGDTFGeometry> Geometry = NextGeometryCollect->FindGeometryByName(*GeometryName);
 				if (Geometry.IsValid())
 				{
-					OuterGeometryCollect = StaticCastSharedPtr<const FDMXGDTFGeometryCollectBase>(Geometry);
+					NextGeometryCollect = StaticCastSharedPtr<const FDMXGDTFGeometryCollectBase>(Geometry);
 				}
 			}
 
 			return nullptr;
 		}
 
+		/** 
+		 * Returns all geometries of a specific type in this collect.
+		 * 
+		 * Example:
+		 * TArray<FDMXGDTFBeamGeometry> BeamGeometries;
+		 * MyGeometryCollect->GetGeometriesOfType(BeamGeometries); // Returns all beam geometries
+		 */
 		template <typename GeometryType>
 		void GetGeometriesOfType(TArray<TSharedPtr<GeometryType>>& OutArray) const
 		{
diff --git a/Engine/Plugins/VirtualProduction/DMX/DMXGDTF/Source/DMXGDTF/Internal/GDTF/Geometries/DMXGDTFGeometryReference.h b/Engine/Plugins/VirtualProduction/DMX/DMXGDTF/Source/DMXGDTF/Internal/GDTF/Geometries/DMXGDTFGeometryReference.h
--- a/Engine/Plugins/VirtualProduction/DMX/DMXGDTF/Source/DMXGDTF/Internal/GDTF/Geometries/DMXGDTFGeometryReference.h
+++ b/Engine/Plugins/VirtualProduction/DMX/DMXGDTF/Source/DMXGDTF/Internal/GDTF/Geometries/DMXGDTFGeometryReference.h
@@ -8,7 +8,7 @@
 
 namespace UE::DMX::GDTF
 {
-	class FDMXGDTFBreak;
+	class FDMXGDTFGeometryBreak;
 	class FDMXGDTFGeometry;
 	class FDMXGDTFGeometryCollectBase;
 	class FDMXGDTFModel;
@@ -35,7 +35,7 @@ namespace UE::DMX::GDTF
 		FTransform Position;
 
 		/** The referenced geometry. Only top level geometries are allowed to be referenced */
-		FString Geometry;
+		FName Geometry;
 
 		/**
 		 * (Optional) Link to the corresponding model. The model only replaces
@@ -54,7 +54,7 @@ namespace UE::DMX::GDTF
 		 * channels of the referenced geometry have the special value “Overwrite” as a DMX break, the DMX break for
 		 * those channels and the DMX offsets need to be defined
 		 */
-		TArray<TSharedPtr<FDMXGDTFBreak>> BreakArray;
+		TArray<TSharedPtr<FDMXGDTFGeometryBreak>> BreakArray;
 
 		/** The outer geometry collect */
 		const TWeakPtr<FDMXGDTFGeometryCollectBase> OuterGeometryCollect;
diff --git a/Engine/Plugins/VirtualProduction/DMX/DMXGDTF/Source/DMXGDTF/Internal/GDTF/Geometries/DMXGDTFSupportGeometry.h b/Engine/Plugins/VirtualProduction/DMX/DMXGDTF/Source/DMXGDTF/Internal/GDTF/Geometries/DMXGDTFSupportGeometry.h
--- a/Engine/Plugins/VirtualProduction/DMX/DMXGDTF/Source/DMXGDTF/Internal/GDTF/Geometries/DMXGDTFSupportGeometry.h
+++ b/Engine/Plugins/VirtualProduction/DMX/DMXGDTF/Source/DMXGDTF/Internal/GDTF/Geometries/DMXGDTFSupportGeometry.h
@@ -40,11 +40,11 @@ namespace UE::DMX::GDTF
 		/** The type of support */
 		EDMXGDTFGeometrySupportType SupportType = EDMXGDTFGeometrySupportType::Rope;
 
-		/** The name of the rope cross section. Only for Ropes.  */
+		/** The name of the rope cross section. Only for Ropes. */
 		FString RopeCrossSection;
 
 		/** The Offset of the rope from bottom to top. Only for Ropes. Unit: meter */
-		FVector RopeOffset;
+		FVector RopeOffset = FVector::ZeroVector;
 
 		/** The allowable force on the X-Axis applied to the object according to the Eurocode. Unit: N. */
 		float CapacityX = 0.f;
diff --git a/Engine/Plugins/VirtualProduction/DMX/DMXGDTF/Source/DMXGDTF/Private/GDTF/DMXModes/DMXGDTFChannelFunction.cpp b/Engine/Plugins/VirtualProduction/DMX/DMXGDTF/Source/DMXGDTF/Private/GDTF/DMXModes/DMXGDTFChannelFunction.cpp
--- a/Engine/Plugins/VirtualProduction/DMX/DMXGDTF/Source/DMXGDTF/Private/GDTF/DMXModes/DMXGDTFChannelFunction.cpp
+++ b/Engine/Plugins/VirtualProduction/DMX/DMXGDTF/Source/DMXGDTF/Private/GDTF/DMXModes/DMXGDTFChannelFunction.cpp
@@ -68,8 +68,9 @@ namespace UE::DMX::GDTF
 	FXmlNode* FDMXGDTFChannelFunction::CreateXmlNode(FXmlNode& Parent)
 	{
 		const FString DefaultLink = TEXT("");
-		const FDMXGDTFDMXValue DefaultModeFrom = TEXT("0/1");
-		const FDMXGDTFDMXValue DefaultModeTo = TEXT("0/1");
+		const FDMXGDTFDMXValue DefaultDefault = FDMXGDTFDMXValue(0);
+		const FDMXGDTFDMXValue DefaultModeFrom = FDMXGDTFDMXValue(0);
+		const FDMXGDTFDMXValue DefaultModeTo = FDMXGDTFDMXValue(0);
 		const float DefaultMin = PhysicalFrom;
 		const float DefaultMax = PhysicalTo;
 		const FString DefaultCustomName = TEXT("");
@@ -79,7 +80,7 @@ namespace UE::DMX::GDTF
 			.SetAttribute(TEXT("Attribute"), Attribute)
 			.SetAttribute(TEXT("OriginalAttribute"), OriginalAttribute)
 			.SetAttribute(TEXT("DMXFrom"), DMXFrom)
-			.SetAttribute(TEXT("Default"), Default)
+			.SetAttribute(TEXT("Default"), Default, FDMXGDTFDMXValue())
 			.SetAttribute(TEXT("PhysicalFrom"), PhysicalFrom)
 			.SetAttribute(TEXT("PhysicalTo"), PhysicalTo)
 			.SetAttribute(TEXT("RealFade"), RealFade)
diff --git a/Engine/Plugins/VirtualProduction/DMX/DMXGDTF/Source/DMXGDTF/Private/GDTF/DMXModes/DMXGDTFDMXChannel.cpp b/Engine/Plugins/VirtualProduction/DMX/DMXGDTF/Source/DMXGDTF/Private/GDTF/DMXModes/DMXGDTFDMXChannel.cpp
--- a/Engine/Plugins/VirtualProduction/DMX/DMXGDTF/Source/DMXGDTF/Private/GDTF/DMXModes/DMXGDTFDMXChannel.cpp
+++ b/Engine/Plugins/VirtualProduction/DMX/DMXGDTF/Source/DMXGDTF/Private/GDTF/DMXModes/DMXGDTFDMXChannel.cpp
@@ -10,6 +10,7 @@
 #include "GDTF/DMXModes/DMXGDTFLogicalChannel.h"
 #include "GDTF/Geometries/DMXGDTFGeometry.h"
 #include "GDTF/Geometries/DMXGDTFGeometryCollect.h"
+#include "GDTF/Geometries/DMXGDTFGeometryReference.h"
 #include "Serialization/DMXGDTFNodeInitializer.h"
 #include "Serialization/DMXGDTFXmlNodeBuilder.h"
 
@@ -30,7 +31,7 @@ namespace UE::DMX::GDTF
 			.GetAttribute(TEXT("Highlight"), Highlight)
 			.GetAttribute(TEXT("Geometry"), Geometry)
 PRAGMA_DISABLE_DEPRECATION_WARNINGS
-			.GetAttribute(TEXT("Default"), Default) // Deprecated with GDTF 1.1, but still initialize so old GDTFs can be supported. See DMXGDTFChannelFunction for the upgrade path.
+			.GetAttribute(TEXT("Default"), Default) // Deprecated with GDTF 1.1, but still initialized so old GDTFs can be supported. See DMXGDTFChannelFunction for the upgrade path.
 PRAGMA_ENABLE_DEPRECATION_WARNINGS
 			.CreateChildren(TEXT("LogicalChannel"), LogicalChannelArray);
 	}
@@ -39,8 +40,21 @@ PRAGMA_ENABLE_DEPRECATION_WARNINGS
 	{
 		const FString DefaultInitialFunction = TEXT("");
 
-		const FDMXGDTFXmlNodeBuilder ChildBuilder = FDMXGDTFXmlNodeBuilder(Parent, *this)
-			.SetAttribute(TEXT("DMXBreak"), DMXBreak)
+		FDMXGDTFXmlNodeBuilder ChildBuilder = FDMXGDTFXmlNodeBuilder(Parent, *this);
+
+		// Write special value overwrite if the DMXBreak is < 1
+		if (DMXBreak < 0)
+		{
+			ChildBuilder
+				.SetAttribute(TEXT("DMXBreak"), TEXT("Overwrite"));
+		}
+		else
+		{
+			ChildBuilder
+				.SetAttribute(TEXT("DMXBreak"), DMXBreak);
+		}
+
+		ChildBuilder
 			.SetAttribute(TEXT("Offset"), Offset)
 			.SetAttribute(TEXT("InitialFunction"), InitialFunction, DefaultInitialFunction)
 			.SetAttribute(TEXT("Highlight"), Highlight)
@@ -78,6 +92,43 @@ PRAGMA_ENABLE_DEPRECATION_WARNINGS
 		return nullptr;
 	}
 
+	TSharedPtr<FDMXGDTFGeometry> FDMXGDTFDMXChannel::ResolveGeometry() const
+	{
+		const TSharedPtr<FDMXGDTFDMXMode> DMXMode = OuterDMXMode.Pin();
+		const TSharedPtr<FDMXGDTFGeometry> TopLevelGeometry = DMXMode.IsValid() ? DMXMode->ResolveGeometry() : nullptr;
+		const TSharedPtr<FDMXGDTFGeometry> GeometryNode = TopLevelGeometry.IsValid() ? TopLevelGeometry->FindGeometryByName(*Geometry.ToString()) : nullptr;
+
+		return GeometryNode;
+	}
+
+	TArray<TSharedPtr<FDMXGDTFGeometryReference>> FDMXGDTFDMXChannel::ResolveGeometryReferences() const
+	{
+		const TSharedPtr<FDMXGDTFDMXMode> DMXMode = OuterDMXMode.Pin();
+		const TSharedPtr<FDMXGDTFGeometry> TopLevelGeometry = DMXMode.IsValid() ? DMXMode->ResolveGeometry() : nullptr;
+
+		TArray<TSharedPtr<FDMXGDTFGeometryReference>> AllGeometryReferences;
+		if (TopLevelGeometry.IsValid())
+		{
+			TArray<TSharedPtr<FDMXGDTFGeometry>> Geometries;
+			TopLevelGeometry->GetGeometriesRecursive(Geometries, AllGeometryReferences);
+		}
+
+		TArray<TSharedPtr<FDMXGDTFGeometryReference>> GeometryReferences;
+		Algo::TransformIf(AllGeometryReferences, GeometryReferences, 
+			[this](const TSharedPtr<FDMXGDTFGeometryReference>& GeometryReference)
+			{
+				return 
+					GeometryReference.IsValid() && 
+					GeometryReference->Geometry == Geometry;
+			},
+			[](const TSharedPtr<FDMXGDTFGeometryReference>& GeometryReference)
+			{
+				return GeometryReference;
+			});
+
+		return GeometryReferences;
+	}
+
 	TArray<uint32> FDMXGDTFDMXChannel::ParseOffset(const FString& GDTFString) const
 	{
 		TArray<FString> Substrings;
diff --git a/Engine/Plugins/VirtualProduction/DMX/DMXGDTF/Source/DMXGDTF/Private/GDTF/DMXModes/DMXGDTFDMXMode.cpp b/Engine/Plugins/VirtualProduction/DMX/DMXGDTF/Source/DMXGDTF/Private/GDTF/DMXModes/DMXGDTFDMXMode.cpp
--- a/Engine/Plugins/VirtualProduction/DMX/DMXGDTF/Source/DMXGDTF/Private/GDTF/DMXModes/DMXGDTFDMXMode.cpp
+++ b/Engine/Plugins/VirtualProduction/DMX/DMXGDTF/Source/DMXGDTF/Private/GDTF/DMXModes/DMXGDTFDMXMode.cpp
@@ -3,12 +3,14 @@
 #include "GDTF/DMXModes/DMXGDTFDMXMode.h"
 
 #include "Algo/Find.h"
+#include "GDTF/DMXGDTFFixtureType.h"
 #include "GDTF/DMXModes/DMXGDTFChannelFunction.h"
 #include "GDTF/DMXModes/DMXGDTFChannelRelation.h"
 #include "GDTF/DMXModes/DMXGDTFDMXChannel.h"
 #include "GDTF/DMXModes/DMXGDTFDMXChannel.h"
 #include "GDTF/DMXModes/DMXGDTFFTMacro.h"
 #include "GDTF/DMXModes/DMXGDTFLogicalChannel.h"
+#include "GDTF/Geometries/DMXGDTFGeometryCollect.h"
 #include "Serialization/DMXGDTFNodeInitializer.h"
 #include "Serialization/DMXGDTFXmlNodeBuilder.h"
 
@@ -42,6 +44,15 @@ namespace UE::DMX::GDTF
 		return ChildBuilder.GetIntermediateXmlNode();
 	}
 
+	TSharedPtr<FDMXGDTFGeometry> FDMXGDTFDMXMode::ResolveGeometry() const
+	{
+		const TSharedPtr<FDMXGDTFFixtureType> FixtureType = GetFixtureType().Pin();
+		const TSharedPtr<FDMXGDTFGeometryCollect> GeometryCollect = FixtureType.IsValid() ? FixtureType->GeometryCollect : nullptr;
+		const TSharedPtr<FDMXGDTFGeometry> GeometryNode = GeometryCollect.IsValid() ? GeometryCollect->FindGeometryByName(*Geometry.ToString()) : nullptr;
+
+		return GeometryNode;
+	}
+
 	void FDMXGDTFDMXMode::ResolveChannel(const FString& Link, TSharedPtr<FDMXGDTFDMXChannel>& OutDMXChannel, TSharedPtr<FDMXGDTFChannelFunction>& OutChannelFunction) const
 	{
 		if (Link.IsEmpty())
diff --git a/Engine/Plugins/VirtualProduction/DMX/DMXGDTF/Source/DMXGDTF/Private/GDTF/DMXModes/DMXGDTFDMXValue.cpp b/Engine/Plugins/VirtualProduction/DMX/DMXGDTF/Source/DMXGDTF/Private/GDTF/DMXModes/DMXGDTFDMXValue.cpp
--- a/Engine/Plugins/VirtualProduction/DMX/DMXGDTF/Source/DMXGDTF/Private/GDTF/DMXModes/DMXGDTFDMXValue.cpp
+++ b/Engine/Plugins/VirtualProduction/DMX/DMXGDTF/Source/DMXGDTF/Private/GDTF/DMXModes/DMXGDTFDMXValue.cpp
@@ -4,93 +4,124 @@
 
 namespace UE::DMX::GDTF
 {
-	FDMXGDTFDMXValue::FDMXGDTFDMXValue(const TCHAR* InValue)
-		: Value(InValue)
-	{}
+	FDMXGDTFDMXValue::FDMXGDTFDMXValue(const FString& InValue)
+	{
+		Set(InValue);
+	}
+
+	FDMXGDTFDMXValue::FDMXGDTFDMXValue(const uint32 InValue)
+	{
+		Set(InValue);
+	}
+
+	FDMXGDTFDMXValue::FDMXGDTFDMXValue(const TOptional<uint32>& InOptionalValue)
+	{
+		if (InOptionalValue.IsSet())
+		{
+			Set(InOptionalValue.GetValue());
+		}
+	}
+
+	uint32 FDMXGDTFDMXValue::AsIntChecked() const
+	{
+		check(IntegerValue.IsSet());
+
+		return IntegerValue.GetValue();
+	}
+
+	bool FDMXGDTFDMXValue::IsSet() const
+	{
+		return IntegerValue.IsSet();
+	}
+
+	void FDMXGDTFDMXValue::Set(uint32 InValue)
+	{
+		IntegerValue = InValue;
+		StringValue = *FString::Printf(TEXT("%u/1"), InValue);
+	}
 
-	bool FDMXGDTFDMXValue::ToInt(uint32& OutInteger) const
+	void FDMXGDTFDMXValue::Set(const FString& InValue)
 	{
-		if (Value.IsEmpty() || Value == TEXT("None"))
+		StringValue = InValue;
+		if (StringValue.IsEmpty() || StringValue == TEXT("None"))
 		{
-			return false;
+			Reset();
+			return;
 		}
 
 		TArray<FString> Substrings;
-		Value.ParseIntoArray(Substrings, TEXT("/"));
+		StringValue.ParseIntoArray(Substrings, TEXT("/"));
 
 		// As per GDTF specs
-		if (Substrings.Num() == 2)
+		if (Substrings.Num() != 2)
 		{
-			uint32 DMXValue;
-			if (!LexTryParseString(DMXValue, *Substrings[0]))
+			Reset();
+			return;
+		}
+
+		uint32 DMXValue;
+		if (!LexTryParseString(DMXValue, *Substrings[0]))
+		{
+			Reset();
+			return;
+		}
+
+		const bool bUseByteShifting = Substrings[1].EndsWith(TEXT("s"));
+		if (bUseByteShifting)
+		{
+			Substrings[1].RemoveFromEnd(TEXT("s"));
+
+			uint8 ByteShiftValue;
+			if (!LexTryParseString(ByteShiftValue, *Substrings[1]) ||
+				ByteShiftValue > 4)
+			{
+				Reset();
+				return;
+			}
+		
+			IntegerValue = DMXValue << ByteShiftValue * 8;
+			return;
+		}
+		else
+		{
+			uint8 ByteMirroringValue;
+			if (!LexTryParseString(ByteMirroringValue, *Substrings[1]) ||
+				ByteMirroringValue > 4)
 			{
-				return false;
+				Reset();
+				return;
 			}
 
-			if (Substrings[1].EndsWith(TEXT("s")))
+			if (ByteMirroringValue == 1)
 			{
-				Substrings[1].RemoveFromEnd(TEXT("s"));
-
-				uint8 ByteShiftValue;
-				if (LexTryParseString(ByteShiftValue, *Substrings[1]))
-				{
-					if (ByteShiftValue == 1)
-					{
-						// No shifting needed
-						OutInteger = DMXValue;
-						return true;
-					}
-					else if (ByteShiftValue > 4)
-					{
-						return false;
-					}
-
-					OutInteger = DMXValue << (ByteShiftValue - 1) * 8;
-					return true;
-				}
+				// No mirroring needed
+				IntegerValue = DMXValue;
+				return;
 			}
-			else
+
+			constexpr uint32 Masks[] = { 0x000000FF, 0x0000FF00, 0x00FF0000, 0xFF000000 };
+			if (ByteMirroringValue == 2)
+			{
+				IntegerValue = ((DMXValue & 0x000000FF) << 8) | ((DMXValue & 0x0000FF00) >> 8);
+			}
+			else if (ByteMirroringValue == 3)
 			{
-				uint8 ByteMirroringValue;
-				if (LexTryParseString(ByteMirroringValue, *Substrings[1]))
-				{
-					if (ByteMirroringValue == 1)
-					{
-						// No mirroring needed
-						OutInteger = DMXValue;
-						return true;
-					}
-					else if (ByteMirroringValue > 4)
-					{
-						return false;
-					}
-					else
-					{
-						constexpr uint32 Masks[] = { 0x000000FF, 0x0000FF00, 0x00FF0000, 0xFF000000 };
-
-						if (ByteMirroringValue == 2)
-						{
-							OutInteger = ((DMXValue & 0x000000FF) << 8) | ((DMXValue & 0x0000FF00) >> 8);
-						}
-						else if (ByteMirroringValue == 3)
-						{
 #if PLATFORM_LITTLE_ENDIAN
-							OutInteger = ((DMXValue & 0x0000FF00) << 16) | ((DMXValue & 0xFF000000) >> 16);
+				IntegerValue = ((DMXValue & 0x0000FF00) << 16) | ((DMXValue & 0xFF000000) >> 16);
 #else
-							OutInteger = ((DMXValue & 0x000000FF) << 16) | ((DMXValue & 0x00FF0000) >> 16);
+				IntegerValue = ((DMXValue & 0x000000FF) << 16) | ((DMXValue & 0x00FF0000) >> 16);
 #endif
-						}
-						else
-						{
-							OutInteger = ((DMXValue & 0x000000FF) << 24) | ((DMXValue & 0x0000FF00) << 8) | ((DMXValue & 0x00FF0000) >> 8) | ((DMXValue & 0xFF000000) >> 24);
-						}
-					}
-
-					return true;
-				}
+			}
+			else
+			{
+				IntegerValue = ((DMXValue & 0x000000FF) << 24) | ((DMXValue & 0x0000FF00) << 8) | ((DMXValue & 0x00FF0000) >> 8) | ((DMXValue & 0xFF000000) >> 24);
 			}
 		}
+	}
 
-		return false;
+	void FDMXGDTFDMXValue::Reset()
+	{
+		IntegerValue.Reset();
+		StringValue = TEXT("None");
 	}
 }
diff --git a/Engine/Plugins/VirtualProduction/DMX/DMXGDTF/Source/DMXGDTF/Private/GDTF/Geometries/DMXGDTFBreak.cpp b/Engine/Plugins/VirtualProduction/DMX/DMXGDTF/Source/DMXGDTF/Private/GDTF/Geometries/DMXGDTFBreak.cpp
deleted file mode 100644
diff --git a/Engine/Plugins/VirtualProduction/DMX/DMXGDTF/Source/DMXGDTF/Private/GDTF/Geometries/DMXGDTFGeometry.cpp b/Engine/Plugins/VirtualProduction/DMX/DMXGDTF/Source/DMXGDTF/Private/GDTF/Geometries/DMXGDTFGeometry.cpp
--- a/Engine/Plugins/VirtualProduction/DMX/DMXGDTF/Source/DMXGDTF/Private/GDTF/Geometries/DMXGDTFGeometry.cpp
+++ b/Engine/Plugins/VirtualProduction/DMX/DMXGDTF/Source/DMXGDTF/Private/GDTF/Geometries/DMXGDTFGeometry.cpp
@@ -38,6 +38,18 @@ namespace UE::DMX::GDTF
 		return ChildBuilder.GetIntermediateXmlNode();
 	}
 
+	TSharedPtr<FDMXGDTFGeometry> FDMXGDTFGeometry::FindGeometryByName(const TCHAR* InName) const
+	{
+		if (InName == Name)
+		{
+			// Allow to find mutable self
+			FDMXGDTFGeometry* NonConstThis = const_cast<FDMXGDTFGeometry*>(this);
+			return StaticCastSharedRef<FDMXGDTFGeometry>(NonConstThis->AsShared());
+		}
+		
+		return FDMXGDTFGeometryCollectBase::FindGeometryByName(InName);
+	}
+
 	TSharedPtr<FDMXGDTFModel> FDMXGDTFGeometry::ResolveModel() const
 	{
 		if (const TSharedPtr<FDMXGDTFFixtureType> FixtureType = GetFixtureType().Pin())
diff --git a/Engine/Plugins/VirtualProduction/DMX/DMXGDTF/Source/DMXGDTF/Private/GDTF/Geometries/DMXGDTFGeometryBreak.cpp b/Engine/Plugins/VirtualProduction/DMX/DMXGDTF/Source/DMXGDTF/Private/GDTF/Geometries/DMXGDTFGeometryBreak.cpp
new file mode 100644
--- /dev/null
+++ b/Engine/Plugins/VirtualProduction/DMX/DMXGDTF/Source/DMXGDTF/Private/GDTF/Geometries/DMXGDTFGeometryBreak.cpp
@@ -0,0 +1,29 @@
+// Copyright Epic Games, Inc. All Rights Reserved.
+
+#include "GDTF/Geometries/DMXGDTFGeometryBreak.h"
+
+#include "Serialization/DMXGDTFNodeInitializer.h"
+#include "Serialization/DMXGDTFXmlNodeBuilder.h"
+
+namespace UE::DMX::GDTF
+{
+	FDMXGDTFGeometryBreak::FDMXGDTFGeometryBreak(const TSharedRef<FDMXGDTFGeometryReference>& InGeometryReference)
+		: OuterGeometryReference(InGeometryReference)
+	{}
+
+	void FDMXGDTFGeometryBreak::Initialize(const FXmlNode& XmlNode)
+	{
+		FDMXGDTFNodeInitializer(SharedThis(this), XmlNode)
+			.GetAttribute(TEXT("DMXOffset"), DMXOffset)
+			.GetAttribute(TEXT("DMXBreak"), DMXBreak);
+	}
+
+	FXmlNode* FDMXGDTFGeometryBreak::CreateXmlNode(FXmlNode& Parent)
+	{
+		const FDMXGDTFXmlNodeBuilder ChildBuilder = FDMXGDTFXmlNodeBuilder(Parent, *this)
+			.SetAttribute(TEXT("DMXOffset"), DMXOffset)
+			.SetAttribute(TEXT("DMXBreak"), DMXBreak);
+
+		return ChildBuilder.GetIntermediateXmlNode();
+	}
+}
diff --git a/Engine/Plugins/VirtualProduction/DMX/DMXGDTF/Source/DMXGDTF/Private/GDTF/Geometries/DMXGDTFGeometryCollectBase.cpp b/Engine/Plugins/VirtualProduction/DMX/DMXGDTF/Source/DMXGDTF/Private/GDTF/Geometries/DMXGDTFGeometryCollectBase.cpp
--- a/Engine/Plugins/VirtualProduction/DMX/DMXGDTF/Source/DMXGDTF/Private/GDTF/Geometries/DMXGDTFGeometryCollectBase.cpp
+++ b/Engine/Plugins/VirtualProduction/DMX/DMXGDTF/Source/DMXGDTF/Private/GDTF/Geometries/DMXGDTFGeometryCollectBase.cpp
@@ -79,15 +79,49 @@ namespace UE::DMX::GDTF
 
 		return ChildBuilder.GetIntermediateXmlNode();
 	}
-	void FDMXGDTFGeometryCollectBase::FindGeometryByName(const TCHAR* InName, TSharedPtr<FDMXGDTFGeometry>& OutGeometry, TSharedPtr<FDMXGDTFGeometryReference>& OutGeometryReference) const
+
+	void FDMXGDTFGeometryCollectBase::GetGeometriesRecursive(TArray<TSharedPtr<FDMXGDTFGeometry>>& OutGeometries, TArray<TSharedPtr<FDMXGDTFGeometryReference>>& OutGeometryReferences)
 	{
-		const TSharedPtr<FDMXGDTFGeometryReference>* GeometryReferencePtr = Algo::FindBy(GeometryReferenceArray, InName, &FDMXGDTFGeometryReference::Name);
-		if (GeometryReferencePtr)
+		for (const TSharedPtr<FDMXGDTFGeometryReference>& GeometryReference : GeometryReferenceArray)
 		{
-			OutGeometryReference = *GeometryReferencePtr;
-			return;
+			OutGeometryReferences.Add(GeometryReference);
 		}
 
+		const auto GetChildrenRecursiveLambda = [this, &OutGeometries, &OutGeometryReferences](const auto& Geometries)
+			{
+				for (const TSharedPtr<FDMXGDTFGeometry>& Geometry : Geometries)
+				{
+					if (Geometry.IsValid())
+					{
+						OutGeometries.Add(Geometry);
+
+						Geometry->GetGeometriesRecursive(OutGeometries, OutGeometryReferences);
+					}
+				}
+			};
+
+		// Recursive for children
+		GetChildrenRecursiveLambda(GeometryArray);
+		GetChildrenRecursiveLambda(AxisArray);
+		GetChildrenRecursiveLambda(FilterBeamArray);
+		GetChildrenRecursiveLambda(FilterColorArray);
+		GetChildrenRecursiveLambda(FilterGoboArray);
+		GetChildrenRecursiveLambda(FilterShaperArray);
+		GetChildrenRecursiveLambda(BeamArray);
+		GetChildrenRecursiveLambda(MediaServerLayerArray);
+		GetChildrenRecursiveLambda(MediaServerCameraArray);
+		GetChildrenRecursiveLambda(MediaServerMasterArray);
+		GetChildrenRecursiveLambda(DisplayArray);
+		GetChildrenRecursiveLambda(LaserArray);
+		GetChildrenRecursiveLambda(WiringObjectArray);
+		GetChildrenRecursiveLambda(InventoryArray);
+		GetChildrenRecursiveLambda(StructureArray);
+		GetChildrenRecursiveLambda(SupportArray);
+		GetChildrenRecursiveLambda(MagnetArray);
+	}
+
+	TSharedPtr<FDMXGDTFGeometry> FDMXGDTFGeometryCollectBase::FindGeometryByName(const TCHAR* InName) const
+	{
 		// Helper to find geometry nodes in arrays of different types
 		auto FindInArrayLambda = [InName](auto InArray, TSharedPtr<FDMXGDTFGeometry>& OutGeometry) -> bool
 			{
@@ -97,22 +131,31 @@ namespace UE::DMX::GDTF
 			};
 
 		TSharedPtr<FDMXGDTFGeometry> Geometry;
-		if (FindInArrayLambda(GeometryArray, Geometry)) { OutGeometry = Geometry; }
-		else if (FindInArrayLambda(AxisArray, Geometry)) { OutGeometry = Geometry; }
-		else if (FindInArrayLambda(FilterBeamArray, Geometry)) { OutGeometry = Geometry; }
-		else if (FindInArrayLambda(FilterColorArray, Geometry)) { OutGeometry = Geometry; }
-		else if (FindInArrayLambda(FilterGoboArray, Geometry)) { OutGeometry = Geometry; }
-		else if (FindInArrayLambda(FilterShaperArray, Geometry)) { OutGeometry = Geometry; }
-		else if (FindInArrayLambda(BeamArray, Geometry)) { OutGeometry = Geometry; }
-		else if (FindInArrayLambda(MediaServerLayerArray, Geometry)) { OutGeometry = Geometry; }
-		else if (FindInArrayLambda(MediaServerCameraArray, Geometry)) { OutGeometry = Geometry; }
-		else if (FindInArrayLambda(MediaServerMasterArray, Geometry)) { OutGeometry = Geometry; }
-		else if (FindInArrayLambda(DisplayArray, Geometry)) { OutGeometry = Geometry; }
-		else if (FindInArrayLambda(LaserArray, Geometry)) { OutGeometry = Geometry; }
-		else if (FindInArrayLambda(WiringObjectArray, Geometry)) { OutGeometry = Geometry; }
-		else if (FindInArrayLambda(InventoryArray, Geometry)) { OutGeometry = Geometry; }
-		else if (FindInArrayLambda(StructureArray, Geometry)) { OutGeometry = Geometry; }
-		else if (FindInArrayLambda(SupportArray, Geometry)) { OutGeometry = Geometry; }
-		else if (FindInArrayLambda(MagnetArray, Geometry)) { OutGeometry = Geometry; }
+		if (FindInArrayLambda(GeometryArray, Geometry)) { return Geometry; }
+		else if (FindInArrayLambda(AxisArray, Geometry)) { return Geometry; }
+		else if (FindInArrayLambda(FilterBeamArray, Geometry)) { return Geometry; }
+		else if (FindInArrayLambda(FilterColorArray, Geometry)) { return Geometry; }
+		else if (FindInArrayLambda(FilterGoboArray, Geometry)) { return Geometry; }
+		else if (FindInArrayLambda(FilterShaperArray, Geometry)) { return Geometry; }
+		else if (FindInArrayLambda(BeamArray, Geometry)) { return Geometry; }
+		else if (FindInArrayLambda(MediaServerLayerArray, Geometry)) { return Geometry; }
+		else if (FindInArrayLambda(MediaServerCameraArray, Geometry)) { return Geometry; }
+		else if (FindInArrayLambda(MediaServerMasterArray, Geometry)) { return Geometry; }
+		else if (FindInArrayLambda(DisplayArray, Geometry)) { return Geometry; }
+		else if (FindInArrayLambda(LaserArray, Geometry)) { return Geometry; }
+		else if (FindInArrayLambda(WiringObjectArray, Geometry)) { return Geometry; }
+		else if (FindInArrayLambda(InventoryArray, Geometry)) { return Geometry; }
+		else if (FindInArrayLambda(StructureArray, Geometry)) { return Geometry; }
+		else if (FindInArrayLambda(SupportArray, Geometry)) { return Geometry; }
+		else if (FindInArrayLambda(MagnetArray, Geometry)) { return Geometry; }
+
+		return nullptr;
+	}
+
+	TSharedPtr<FDMXGDTFGeometryReference> FDMXGDTFGeometryCollectBase::FindGeometryReferenceByName(const TCHAR* InName) const
+	{
+		const TSharedPtr<FDMXGDTFGeometryReference>* GeometryReferencePtr = Algo::FindBy(GeometryReferenceArray, InName, &FDMXGDTFGeometryReference::Name);
+
+		return GeometryReferencePtr ? *GeometryReferencePtr : nullptr;
 	}
 }
diff --git a/Engine/Plugins/VirtualProduction/DMX/DMXGDTF/Source/DMXGDTF/Private/GDTF/Geometries/DMXGDTFGeometryReference.cpp b/Engine/Plugins/VirtualProduction/DMX/DMXGDTF/Source/DMXGDTF/Private/GDTF/Geometries/DMXGDTFGeometryReference.cpp
--- a/Engine/Plugins/VirtualProduction/DMX/DMXGDTF/Source/DMXGDTF/Private/GDTF/Geometries/DMXGDTFGeometryReference.cpp
+++ b/Engine/Plugins/VirtualProduction/DMX/DMXGDTF/Source/DMXGDTF/Private/GDTF/Geometries/DMXGDTFGeometryReference.cpp
@@ -4,8 +4,8 @@
 
 #include "Algo/Find.h"
 #include "GDTF/DMXGDTFFixtureType.h"
-#include "GDTF/Geometries/DMXGDTFBreak.h"
 #include "GDTF/Geometries/DMXGDTFGeometry.h"
+#include "GDTF/Geometries/DMXGDTFGeometryBreak.h"
 #include "GDTF/Geometries/DMXGDTFGeometryCollect.h"
 #include "GDTF/Models/DMXGDTFModel.h"
 #include "Serialization/DMXGDTFNodeInitializer.h"
@@ -43,7 +43,7 @@ namespace UE::DMX::GDTF
 
 	TSharedPtr<FDMXGDTFGeometry> FDMXGDTFGeometryReference::ResolveGeometry() const
 	{
-		// Only top level geometries are allowed
+		// Only top level geometries are allowed to be referenced 
 		const TSharedPtr<FDMXGDTFFixtureType> FixtureType = GetFixtureType().Pin();
 		const TSharedPtr<FDMXGDTFGeometryCollect> GeometryCollect = FixtureType.IsValid() ? FixtureType->GeometryCollect : nullptr;
 		if (GeometryCollect.IsValid())
diff --git a/Engine/Plugins/VirtualProduction/DMX/DMXGDTF/Source/DMXGDTF/Private/Serialization/DMXGDTFNodeInitializer.h b/Engine/Plugins/VirtualProduction/DMX/DMXGDTF/Source/DMXGDTF/Private/Serialization/DMXGDTFNodeInitializer.h
--- a/Engine/Plugins/VirtualProduction/DMX/DMXGDTF/Source/DMXGDTF/Private/Serialization/DMXGDTFNodeInitializer.h
+++ b/Engine/Plugins/VirtualProduction/DMX/DMXGDTF/Source/DMXGDTF/Private/Serialization/DMXGDTFNodeInitializer.h
@@ -102,7 +102,7 @@ namespace UE::DMX::GDTF
 			}
 			else if constexpr (std::is_same_v<AttributeType, FDMXGDTFDMXValue>)
 			{
-				OutValue.Value = StringValue;
+				OutValue = FDMXGDTFDMXValue(*StringValue);
 			}
 			else
 			{
@@ -396,7 +396,7 @@ namespace UE::DMX::GDTF
 			const FMatrix GDTFToUnrealMatrix = FMatrix(
 				FPlane(1.0, 0.0, 0.0, 0.0),
 				FPlane(0.0, 0.0, 1.0, 0.0),
-				FPlane(0.0, -1.0, 0.0, 0.0),
+				FPlane(0.0, 1.0, 0.0, 0.0),
 				FPlane(0.0, 0.0, 0.0, 1.0)
 			);
 
diff --git a/Engine/Plugins/VirtualProduction/DMX/DMXGDTF/Source/DMXGDTF/Private/Serialization/DMXGDTFXmlNodeBuilder.cpp b/Engine/Plugins/VirtualProduction/DMX/DMXGDTF/Source/DMXGDTF/Private/Serialization/DMXGDTFXmlNodeBuilder.cpp
--- a/Engine/Plugins/VirtualProduction/DMX/DMXGDTF/Source/DMXGDTF/Private/Serialization/DMXGDTFXmlNodeBuilder.cpp
+++ b/Engine/Plugins/VirtualProduction/DMX/DMXGDTF/Source/DMXGDTF/Private/Serialization/DMXGDTFXmlNodeBuilder.cpp
@@ -139,12 +139,12 @@ namespace UE::DMX::GDTF
 
 	FDMXGDTFXmlNodeBuilder& FDMXGDTFXmlNodeBuilder::SetAttribute(const FString& AttributeName, const FDMXGDTFDMXValue& DMXValue, const TOptional<FDMXGDTFDMXValue> OptionalDefault)
 	{
-		if (OptionalDefault.IsSet() && DMXValue.Value == OptionalDefault.GetValue().Value)
+		if (OptionalDefault.IsSet() && DMXValue == OptionalDefault.GetValue())
 		{
 			return *this;
 		}
 
-		Attributes.Add(FXmlAttribute(AttributeName, DMXValue.Value));
+		Attributes.Add(FXmlAttribute(AttributeName, DMXValue.AsString()));
 		return *this;
 	}
 
@@ -203,7 +203,7 @@ namespace UE::DMX::GDTF
 		const FMatrix GDTFToUnrealMatrix = FMatrix(
 			FPlane(1.0, 0.0, 0.0, 0.0),
 			FPlane(0.0, 0.0, 1.0, 0.0),
-			FPlane(0.0, -1.0, 0.0, 0.0),
+			FPlane(0.0, 1.0, 0.0, 0.0),
 			FPlane(0.0, 0.0, 0.0, 1.0)
 		);
 
