DMX - Rework DMX Library's fixture patch list to use the fixture ID of the patch instead of the general scene description

diff --git a/Engine/Plugins/VirtualProduction/DMX/DMXEngine/Source/DMXEditor/Private/Widgets/FixturePatch/DMXFixturePatchListItem.cpp b/Engine/Plugins/VirtualProduction/DMX/DMXEngine/Source/DMXEditor/Private/Widgets/FixturePatch/DMXFixturePatchListItem.cpp
new file mode 100644
--- /dev/null
+++ b/Engine/Plugins/VirtualProduction/DMX/DMXEngine/Source/DMXEditor/Private/Widgets/FixturePatch/DMXFixturePatchListItem.cpp
@@ -0,0 +1,239 @@
+// Copyright Epic Games, Inc. All Rights Reserved.
+
+#include "DMXFixturePatchListItem.h"
+
+#include "Algo/MinElement.h"
+#include "DMXEditor.h"
+#include "DMXFixturePatchSharedData.h"
+#include "Library/DMXEntityFixturePatch.h"
+#include "Library/DMXEntityReference.h"
+#include "Library/DMXLibrary.h"
+#include "ScopedTransaction.h"
+
+#define LOCTEXT_NAMESPACE "DMXFixturePatchListItem"
+
+FDMXFixturePatchListItem::FDMXFixturePatchListItem(TWeakPtr<FDMXEditor> InDMXEditor, UDMXEntityFixturePatch* InFixturePatch)
+	: WeakFixturePatch(InFixturePatch)
+	, WeakDMXEditor(InDMXEditor)
+{}
+
+FGuid FDMXFixturePatchListItem::GetMVRUUID() const
+{
+	return WeakFixturePatch.IsValid() ? WeakFixturePatch->GetMVRFixtureUUID() : FGuid();
+}
+
+FLinearColor FDMXFixturePatchListItem::GetBackgroundColor() const
+{
+	if (!ErrorStatusText.IsEmpty())
+	{
+		return FLinearColor::Red;
+	}
+
+	if (UDMXEntityFixturePatch* FixturePatch = WeakFixturePatch.Get())
+	{
+		return FixturePatch->EditorColor;
+	}
+
+	return FLinearColor::Red;
+}
+
+FString FDMXFixturePatchListItem::GetFixturePatchName() const
+{
+	if (const UDMXEntityFixturePatch* FixturePatch = GetFixturePatch())
+	{
+		return FixturePatch->Name;
+	}
+
+	return FString();
+}
+
+void FDMXFixturePatchListItem::SetFixturePatchName(const FString& InDesiredName, FString& OutNewName)
+{
+	if (UDMXEntityFixturePatch* FixturePatch = WeakFixturePatch.Get())
+	{
+		if (FixturePatch->Name == InDesiredName)
+		{
+			OutNewName = InDesiredName;
+			return;
+		}
+
+		const FScopedTransaction SetFixturePatchNameTransaction(LOCTEXT("SetFixturePatchNameTransaction", "Set Fixture Patch Name"));
+		FixturePatch->PreEditChange(UDMXEntityFixturePatch::StaticClass()->FindPropertyByName(GET_MEMBER_NAME_CHECKED(UDMXEntity, Name)));
+
+		FixturePatch->SetName(InDesiredName);
+		OutNewName = FixturePatch->Name;
+
+		FixturePatch->PostEditChange();
+	}
+}
+
+FString FDMXFixturePatchListItem::GetFixtureID() const
+{
+	return WeakFixturePatch.IsValid() ? FString::FromInt(WeakFixturePatch->GetFixtureID()) : TEXT("Invalid");
+}
+
+void FDMXFixturePatchListItem::SetFixtureID(int32 InFixtureID)
+{
+	if (UDMXEntityFixturePatch* FixturePatch = WeakFixturePatch.Get())
+	{
+		const FScopedTransaction SetFixturePatchNameTransaction(LOCTEXT("SetFixturePatchFixtureIDTransaction", "Set Fixture ID"));
+		FixturePatch->PreEditChange(UDMXEntityFixturePatch::StaticClass()->FindPropertyByName(UDMXEntityFixturePatch::GetFixtureIDPropertyNameChecked()));
+
+		FixturePatch->GenerateFixtureID(InFixtureID);
+
+		FixturePatch->PostEditChange();
+	}
+}
+
+UDMXEntityFixtureType* FDMXFixturePatchListItem::GetFixtureType() const
+{
+	if (UDMXEntityFixturePatch* FixturePatch = WeakFixturePatch.Get())
+	{
+		return FixturePatch->GetFixtureType();
+	}
+
+	return nullptr;
+}
+
+void FDMXFixturePatchListItem::SetFixtureType(UDMXEntityFixtureType* FixtureType)
+{
+	UDMXLibrary* DMXLibrary = GetDMXLibrary();
+	UDMXEntityFixturePatch* FixturePatch = WeakFixturePatch.Get();
+
+	if (DMXLibrary && FixturePatch)
+	{
+		if (FixturePatch->GetFixtureType() == FixtureType)
+		{
+			return;
+		}
+
+		const FScopedTransaction SetFixtureTypeTransaction(LOCTEXT("SetFixtureTypeTransaction", "Set Fixture Type of Patch"));
+		FixturePatch->PreEditChange(UDMXEntityFixturePatch::StaticClass()->FindPropertyByName(UDMXEntityFixturePatch::GetParentFixtureTypeTemplatePropertyNameChecked()));
+
+		FixturePatch->SetFixtureType(FixtureType);
+
+		FixturePatch->PostEditChange();
+	}
+}
+
+int32 FDMXFixturePatchListItem::GetModeIndex() const
+{
+	if (UDMXEntityFixturePatch* FixturePatch = WeakFixturePatch.Get())
+	{
+		return FixturePatch->GetActiveModeIndex();
+	}
+
+	return INDEX_NONE;
+}
+
+void FDMXFixturePatchListItem::SetModeIndex(int32 ModeIndex)
+{
+	UDMXLibrary* DMXLibrary = GetDMXLibrary();
+	UDMXEntityFixturePatch* FixturePatch = WeakFixturePatch.Get();
+
+	if (DMXLibrary && FixturePatch)
+	{
+		if (ModeIndex == FixturePatch->GetActiveModeIndex())
+		{
+			return;
+		}
+		
+		// If all should be changed, just change the patch
+		const FScopedTransaction SetModeTransaction(LOCTEXT("SetModeTransaction", "Set Mode of Patch"));
+		FixturePatch->PreEditChange(UDMXEntityFixturePatch::StaticClass()->FindPropertyByName(UDMXEntityFixturePatch::GetParentFixtureTypeTemplatePropertyNameChecked()));
+
+		FixturePatch->SetActiveModeIndex(ModeIndex);
+			
+		FixturePatch->PostEditChange();
+	}
+}
+
+int32 FDMXFixturePatchListItem::GetUniverse() const
+{
+	if (UDMXEntityFixturePatch* FixturePatch = WeakFixturePatch.Get())
+	{
+		return FixturePatch->GetUniverseID();
+	}
+
+	return -1;
+}
+
+int32 FDMXFixturePatchListItem::GetAddress() const
+{
+	if (UDMXEntityFixturePatch* FixturePatch = WeakFixturePatch.Get())
+	{
+		return FixturePatch->GetStartingChannel();
+	}
+
+	return -1;
+}
+
+void FDMXFixturePatchListItem::SetAddresses(int32 Universe, int32 Address)
+{
+	UDMXLibrary* DMXLibrary = GetDMXLibrary();
+	UDMXEntityFixturePatch* FixturePatch = WeakFixturePatch.Get();
+	const TSharedPtr<FDMXEditor> DMXEditor = WeakDMXEditor.Pin();
+	const TSharedPtr<FDMXFixturePatchSharedData> SharedData = DMXEditor.IsValid() ? DMXEditor->GetFixturePatchSharedData() : nullptr;
+
+	if (DMXLibrary && FixturePatch && SharedData.IsValid())
+	{
+		if (FixturePatch->GetUniverseID() == Universe &&
+			FixturePatch->GetStartingChannel() == Address)
+		{
+			return;
+		}
+
+		// Only valid values
+		const FDMXFixtureMode* ModePtr = FixturePatch->GetActiveMode();
+		const int32 MaxAddress = ModePtr ? DMX_MAX_ADDRESS - ModePtr->ChannelSpan + 1 : DMX_MAX_ADDRESS;
+		if (Universe < 0 || 
+			Universe > DMX_MAX_UNIVERSE || 
+			Address < 1 || 
+			Address > MaxAddress)
+		{
+			return;
+		}
+		
+		const FScopedTransaction SetAddressesTransaction(LOCTEXT("SetAddressesTransaction", "Set Addresses of Patch"));
+		FixturePatch->PreEditChange(UDMXEntityFixturePatch::StaticClass()->FindPropertyByName(UDMXEntityFixturePatch::GetParentFixtureTypeTemplatePropertyNameChecked()));
+
+		FixturePatch->SetUniverseID(Universe);
+		FixturePatch->SetStartingChannel(Address);
+
+		FixturePatch->PostEditChange();
+
+		// Select the universe in Fixture Patch Shared Data
+		SharedData->SelectUniverse(Universe);
+	}
+}
+
+int32 FDMXFixturePatchListItem::GetNumChannels() const
+{
+	if (UDMXEntityFixturePatch* FixturePatch = WeakFixturePatch.Get())
+	{
+		return FixturePatch->GetChannelSpan();
+	}
+
+	return -1;
+}
+
+UDMXEntityFixturePatch* FDMXFixturePatchListItem::GetFixturePatch() const
+{
+	return WeakFixturePatch.Get();
+}
+
+UDMXLibrary* FDMXFixturePatchListItem::GetDMXLibrary() const
+{
+	if (const TSharedPtr<FDMXEditor> DMXEditor = WeakDMXEditor.Pin())
+	{
+		UDMXLibrary* DMXLibrary = DMXEditor->GetDMXLibrary();
+		if (DMXLibrary)
+		{
+			return DMXLibrary;
+		}
+	}
+
+	return nullptr;
+}
+
+#undef LOCTEXT_NAMESPACE
diff --git a/Engine/Plugins/VirtualProduction/DMX/DMXEngine/Source/DMXEditor/Private/Widgets/FixturePatch/DMXFixturePatchListItem.h b/Engine/Plugins/VirtualProduction/DMX/DMXEngine/Source/DMXEditor/Private/Widgets/FixturePatch/DMXFixturePatchListItem.h
new file mode 100644
--- /dev/null
+++ b/Engine/Plugins/VirtualProduction/DMX/DMXEngine/Source/DMXEditor/Private/Widgets/FixturePatch/DMXFixturePatchListItem.h
@@ -0,0 +1,85 @@
+// Copyright Epic Games, Inc. All Rights Reserved.
+
+#pragma once
+
+#include "MVR/DMXMVRGeneralSceneDescription.h"
+#include "UObject/WeakObjectPtr.h"
+
+class FDMXEditor;
+class FDMXFixturePatchSharedData;
+class UDMXEntityFixturePatch;
+class UDMXEntityFixtureType;
+class UDMXLibrary;
+
+/** An MVR Fixture as an Item in a List. A Primary Items is the first MVR Fixture in a Patch. Secondary Items are subsequent MVR Fixtures */
+class FDMXFixturePatchListItem
+	: public TSharedFromThis<FDMXFixturePatchListItem>
+{
+public:	
+	/** Constructor */
+	FDMXFixturePatchListItem(TWeakPtr<FDMXEditor> InDMXEditor, UDMXEntityFixturePatch* InFixturePatch);
+
+	/** Returns the MVR UUID of the MVR Fixture */
+	FGuid GetMVRUUID() const;
+
+	/** Returns the background color of this Item */
+	FLinearColor GetBackgroundColor() const;
+
+	/** Returns the Name of the Fixture Patch */
+	FString GetFixturePatchName() const;
+
+	/** Sets the Name of the Fixture Patch */
+	void SetFixturePatchName(const FString& InDesiredName, FString& OutNewName);
+
+	/** Returns the Fixture ID */
+	FString GetFixtureID() const;
+
+	/** Sets the Fixture ID. Note, as other hard- and software, we allow to set integer values only. */
+	void SetFixtureID(int32 InFixtureID);
+
+	/** Returns the Fixture Type of the MVR Fixture */
+	UDMXEntityFixtureType* GetFixtureType() const;
+	
+	/** Sets the fixture type of the MVR Fixture */
+	void SetFixtureType(UDMXEntityFixtureType* FixtureType);
+
+	/** Returns the Mode Index of the MVR Fixture */
+	int32 GetModeIndex() const;
+
+	/** Sets the Mode Index for the MVR */
+	void SetModeIndex(int32 ModeIndex);
+
+	/** Returns the Universe the MVR Fixture resides in */
+	int32 GetUniverse() const;
+
+	/** Returns the Address of the MVR Fixture */
+	int32 GetAddress() const;
+
+	/** Sets the Addresses of the MVR Fixture */
+	void SetAddresses(int32 Universe, int32 Address);
+
+	/** Returns Num Channels the MVR Fixture spans */
+	int32 GetNumChannels() const;
+
+	/** Returns the Fixture Patch. Const on purpose, should be edited via this class's methods */
+	UDMXEntityFixturePatch* GetFixturePatch() const;
+
+	/** Returns the DMX Library in which the MVR Fixture resides */
+	UDMXLibrary* GetDMXLibrary() const;
+
+	/** Warning Status Text of the Item */
+	FText WarningStatusText;
+
+	/** Error Status Text of the Item */
+	FText ErrorStatusText;
+
+private:
+	/** True if this is in an even group */
+	bool bIsEvenGroup = false;
+
+	/** The fixture patch the MVR Fixture UUID is assigned to */
+	TWeakObjectPtr<UDMXEntityFixturePatch> WeakFixturePatch;
+
+	/** The DMX Editor that owns this Item */
+	TWeakPtr<FDMXEditor> WeakDMXEditor;
+};
diff --git a/Engine/Plugins/VirtualProduction/DMX/DMXEngine/Source/DMXEditor/Private/Widgets/FixturePatch/DMXFixturePatchNode.cpp b/Engine/Plugins/VirtualProduction/DMX/DMXEngine/Source/DMXEditor/Private/Widgets/FixturePatch/DMXFixturePatchNode.cpp
--- a/Engine/Plugins/VirtualProduction/DMX/DMXEngine/Source/DMXEditor/Private/Widgets/FixturePatch/DMXFixturePatchNode.cpp
+++ b/Engine/Plugins/VirtualProduction/DMX/DMXEngine/Source/DMXEditor/Private/Widgets/FixturePatch/DMXFixturePatchNode.cpp
@@ -5,16 +5,13 @@
 #include "DMXEditor.h"
 #include "DMXFixturePatchEditorDefinitions.h"
 #include "DMXFixturePatchSharedData.h"
-#include "SDMXFixturePatchFragment.h"
-#include "SDMXPatchedUniverse.h"
 #include "Library/DMXEntityFixturePatch.h"
 #include "Library/DMXLibrary.h"
-#include "MVR/Types/DMXMVRFixtureNode.h"
-
 #include "ScopedTransaction.h"
+#include "SDMXFixturePatchFragment.h"
+#include "SDMXPatchedUniverse.h"
 #include "Widgets/Layout/SGridPanel.h"
 
-
 #define LOCTEXT_NAMESPACE "DMXFixturePatchNode"
 
 /** A fixture patch fragment in a grid. CreateFragments creates DMXFixturePatchFragment from a DMXFixturePatchNode. */
@@ -316,62 +313,22 @@ bool FDMXFixturePatchNode::IsSelected() const
 
 FString FDMXFixturePatchNode::GetFixtureID() const
 {
-	if (!MVRFixtureNode.IsValid())
-	{
-		UDMXLibrary* DMXLibrary = WeakDMXEditor.IsValid() ? WeakDMXEditor.Pin()->GetDMXLibrary() : nullptr;
-		if (!DMXLibrary || !FixturePatch.IsValid())
-		{
-			return FString();
-		}
-
-		TArray<UDMXMVRFixtureNode*> MVRFixtureNodes;
-		DMXLibrary->GetLazyGeneralSceneDescription()->GetFixtureNodes(MVRFixtureNodes);
-		UDMXMVRFixtureNode* const* MVRFixtureNodePtr = Algo::FindByPredicate(MVRFixtureNodes, [this](const UDMXMVRFixtureNode* Node)
-			{
-				return Node->UUID == FixturePatch->GetMVRFixtureUUID();
-			});
-		if (MVRFixtureNodePtr)
-		{
-			MVRFixtureNode = TWeakObjectPtr<UDMXMVRFixtureNode>(*MVRFixtureNodePtr);
-		}
-	}
-
-	if (MVRFixtureNode.IsValid())
-	{
-		return MVRFixtureNode->FixtureID;
-	}
-
-	return FString();
+	return FixturePatch.IsValid() ? FString::FromInt(FixturePatch->GetFixtureID()) : TEXT("Invalid");
 }
 
 int32 FDMXFixturePatchNode::GetUniverseID() const
 {
-	if (FixturePatch.IsValid())
-	{
-		return FixturePatch->GetUniverseID();
-	}
-
-	return -1;
+	return FixturePatch.IsValid() ? FixturePatch->GetUniverseID() : -1;
 }
 
 int32 FDMXFixturePatchNode::GetStartingChannel() const
 {
-	if (FixturePatch.IsValid())
-	{
-		return FixturePatch->GetStartingChannel();
-	}
-
-	return -1;
+	return FixturePatch.IsValid() ? FixturePatch->GetStartingChannel() : -1;
 }
 
 int32 FDMXFixturePatchNode::GetChannelSpan() const
 {
-	if (FixturePatch.IsValid())
-	{
-		return FixturePatch->GetChannelSpan();
-	}
-
-	return -1;
+	return FixturePatch.IsValid() ? FixturePatch->GetChannelSpan() : -1;
 }
 
 void FDMXFixturePatchNode::UpdateIsSelected()
diff --git a/Engine/Plugins/VirtualProduction/DMX/DMXEngine/Source/DMXEditor/Private/Widgets/FixturePatch/DMXFixturePatchNode.h b/Engine/Plugins/VirtualProduction/DMX/DMXEngine/Source/DMXEditor/Private/Widgets/FixturePatch/DMXFixturePatchNode.h
--- a/Engine/Plugins/VirtualProduction/DMX/DMXEngine/Source/DMXEditor/Private/Widgets/FixturePatch/DMXFixturePatchNode.h
+++ b/Engine/Plugins/VirtualProduction/DMX/DMXEngine/Source/DMXEditor/Private/Widgets/FixturePatch/DMXFixturePatchNode.h
@@ -2,7 +2,6 @@
 
 #pragma once
 
-#include "CoreMinimal.h"
 #include "Templates/SharedPointer.h"
 #include "UObject/WeakObjectPtr.h"
 
@@ -10,8 +9,6 @@ class FDMXEditor;
 class SDMXFixturePatchFragment;
 class SDMXPatchedUniverse;
 class UDMXEntityFixturePatch;
-class UDMXMVRFixtureNode;
-
 
 /** 
  * A fixture patch in a grid, consists of several fragments, displayed as SDMXFixturePatchNodes (see GenerateWidgets). 
@@ -80,9 +77,6 @@ private:
 	/** The Fixture Patch this Node stands for */
 	TWeakObjectPtr<UDMXEntityFixturePatch> FixturePatch;
 
-	/** MVR Fixture Node that corrresponds to the patch */
-	mutable TWeakObjectPtr<UDMXMVRFixtureNode> MVRFixtureNode;
-
 	/** Weak DMXEditor refrence */
 	TWeakPtr<FDMXEditor> WeakDMXEditor;
 };
diff --git a/Engine/Plugins/VirtualProduction/DMX/DMXEngine/Source/DMXEditor/Private/Widgets/FixturePatch/DMXMVRFixtureListItem.cpp b/Engine/Plugins/VirtualProduction/DMX/DMXEngine/Source/DMXEditor/Private/Widgets/FixturePatch/DMXMVRFixtureListItem.cpp
deleted file mode 100644
diff --git a/Engine/Plugins/VirtualProduction/DMX/DMXEngine/Source/DMXEditor/Private/Widgets/FixturePatch/DMXMVRFixtureListItem.h b/Engine/Plugins/VirtualProduction/DMX/DMXEngine/Source/DMXEditor/Private/Widgets/FixturePatch/DMXMVRFixtureListItem.h
deleted file mode 100644
diff --git a/Engine/Plugins/VirtualProduction/DMX/DMXEngine/Source/DMXEditor/Private/Widgets/FixturePatch/SDMXFixturePatchEditor.cpp b/Engine/Plugins/VirtualProduction/DMX/DMXEngine/Source/DMXEditor/Private/Widgets/FixturePatch/SDMXFixturePatchEditor.cpp
--- a/Engine/Plugins/VirtualProduction/DMX/DMXEngine/Source/DMXEditor/Private/Widgets/FixturePatch/SDMXFixturePatchEditor.cpp
+++ b/Engine/Plugins/VirtualProduction/DMX/DMXEngine/Source/DMXEditor/Private/Widgets/FixturePatch/SDMXFixturePatchEditor.cpp
@@ -2,24 +2,22 @@
 
 #include "Widgets/FixturePatch/SDMXFixturePatchEditor.h"
 
+#include "Customizations/DMXEntityFixturePatchDetails.h"
 #include "DMXEditor.h"
 #include "DMXEditorSettings.h"
 #include "DMXFixturePatchSharedData.h"
 #include "DMXSubsystem.h"
-#include "SDMXFixturePatcher.h"
-#include "Customizations/DMXEntityFixturePatchDetails.h"
-#include "IO/DMXPortManager.h"
 #include "IO/DMXOutputPort.h"
 #include "IO/DMXOutputPortReference.h"
-#include "Library/DMXLibrary.h"
+#include "IO/DMXPortManager.h"
 #include "Library/DMXEntityFixturePatch.h"
-#include "Widgets/FixturePatch/SDMXMVRFixtureList.h"
-
-#include "PropertyEditorModule.h"
+#include "Library/DMXLibrary.h"
 #include "Modules/ModuleManager.h"
+#include "PropertyEditorModule.h"
+#include "SDMXFixturePatcher.h"
+#include "Widgets/FixturePatch/SDMXFixturePatchList.h"
 #include "Widgets/Layout/SSplitter.h"
 
-
 #define LOCTEXT_NAMESPACE "SDMXFixturePatcher"
 
 SDMXFixturePatchEditor::~SDMXFixturePatchEditor()
@@ -36,6 +34,10 @@ void SDMXFixturePatchEditor::Construct(const FArguments& InArgs)
 	SDMXEntityEditor::Construct(SDMXEntityEditor::FArguments());
 
 	DMXEditorPtr = InArgs._DMXEditor;
+	if (!DMXEditorPtr.IsValid())
+	{
+		return;
+	}
 	FixturePatchSharedData = DMXEditorPtr.Pin()->GetFixturePatchSharedData();
 
 	SetCanTick(false);
@@ -58,7 +60,7 @@ void SDMXFixturePatchEditor::Construct(const FArguments& InArgs)
 		+ SSplitter::Slot()	
 		.Value(LeftSideWidth)
 		[
-			SAssignNew(MVRFixtureList, SDMXMVRFixtureList, DMXEditorPtr)
+			SAssignNew(FixturePatchList, SDMXFixturePatchList, DMXEditorPtr)
 		]
 
 		// Right, Fixture Patcher and Details
@@ -93,42 +95,36 @@ void SDMXFixturePatchEditor::Construct(const FArguments& InArgs)
 
 FReply SDMXFixturePatchEditor::OnKeyDown(const FGeometry& MyGeometry, const FKeyEvent& InKeyEvent)
 {
-	return MVRFixtureList->ProcessCommandBindings(InKeyEvent);
+	return FixturePatchList->ProcessCommandBindings(InKeyEvent);
 }
 
 void SDMXFixturePatchEditor::RequestRenameOnNewEntity(const UDMXEntity* InEntity, ESelectInfo::Type SelectionType)
 {
-	if (MVRFixtureList.IsValid())
+	if (FixturePatchList.IsValid())
 	{
-		MVRFixtureList->EnterFixturePatchNameEditingMode();
+		FixturePatchList->EnterFixturePatchNameEditingMode();
 	}
 }
 
 void SDMXFixturePatchEditor::SelectEntity(UDMXEntity* InEntity, ESelectInfo::Type InSelectionType /*= ESelectInfo::Type::Direct*/)
 {
-	if (TSharedPtr<FDMXEditor> DMXEditor = DMXEditorPtr.Pin())
+	if (UDMXEntityFixturePatch* FixturePatch = Cast<UDMXEntityFixturePatch>(InEntity))
 	{
-		if (UDMXEntityFixturePatch* FixturePatch = Cast<UDMXEntityFixturePatch>(InEntity))
-		{
-			DMXEditor->GetFixturePatchSharedData()->SelectFixturePatch(FixturePatch);
-		}
+		FixturePatchSharedData->SelectFixturePatch(FixturePatch);
 	}
 }
 
 void SDMXFixturePatchEditor::SelectEntities(const TArray<UDMXEntity*>& InEntities, ESelectInfo::Type InSelectionType /*= ESelectInfo::Type::Direct*/)
 {
-	if (TSharedPtr<FDMXEditor> DMXEditor = DMXEditorPtr.Pin())
+	TArray<TWeakObjectPtr<UDMXEntityFixturePatch>> FixturePatches;
+	for (UDMXEntity* Entity : InEntities)
 	{
-		TArray<TWeakObjectPtr<UDMXEntityFixturePatch>> FixturePatches;
-		for (UDMXEntity* Entity : InEntities)
+		if (UDMXEntityFixturePatch* FixturePatch = Cast<UDMXEntityFixturePatch>(Entity))
 		{
-			if (UDMXEntityFixturePatch* FixturePatch = Cast<UDMXEntityFixturePatch>(Entity))
-			{
-				FixturePatches.Add(FixturePatch);
-			}
+			FixturePatches.Add(FixturePatch);
 		}
-		DMXEditor->GetFixturePatchSharedData()->SelectFixturePatches(FixturePatches);
 	}
+	FixturePatchSharedData->SelectFixturePatches(FixturePatches);
 }
 
 TArray<UDMXEntity*> SDMXFixturePatchEditor::GetSelectedEntities() const
@@ -148,12 +144,12 @@ TArray<UDMXEntity*> SDMXFixturePatchEditor::GetSelectedEntities() const
 
 void SDMXFixturePatchEditor::SelectUniverse(int32 UniverseID)
 {
-	check(UniverseID >= 0 && UniverseID <= DMX_MAX_UNIVERSE);
-
-	if (TSharedPtr<FDMXEditor> DMXEditor = DMXEditorPtr.Pin())
+	if (!ensureMsgf(UniverseID >= 0 && UniverseID <= DMX_MAX_UNIVERSE, TEXT("Invalid Universe when trying to select Universe %i."), UniverseID))
 	{
-		FixturePatchSharedData->SelectUniverse(UniverseID);
+		return;
 	}
+
+	FixturePatchSharedData->SelectUniverse(UniverseID);
 }
 
 void SDMXFixturePatchEditor::OnFixturePatchesSelected()
diff --git a/Engine/Plugins/VirtualProduction/DMX/DMXEngine/Source/DMXEditor/Private/Widgets/FixturePatch/SDMXFixturePatchEditor.h b/Engine/Plugins/VirtualProduction/DMX/DMXEngine/Source/DMXEditor/Private/Widgets/FixturePatch/SDMXFixturePatchEditor.h
--- a/Engine/Plugins/VirtualProduction/DMX/DMXEngine/Source/DMXEditor/Private/Widgets/FixturePatch/SDMXFixturePatchEditor.h
+++ b/Engine/Plugins/VirtualProduction/DMX/DMXEngine/Source/DMXEditor/Private/Widgets/FixturePatch/SDMXFixturePatchEditor.h
@@ -10,7 +10,7 @@ class FDMXEditor;
 class FDMXFixturePatchSharedData;
 class SDMXFixturePatcher;
 class SDMXFixturePatchTree;
-class SDMXMVRFixtureList;
+class SDMXFixturePatchList;
 class UDMXEntityFixturePatch;
 
 struct FPropertyChangedEvent;
@@ -62,7 +62,7 @@ private:
 	TSharedRef<IDetailsView> GenerateFixturePatchDetailsView() const;
 
 	/** List of Fixture Patches as MVR Fixtures */
-	TSharedPtr<SDMXMVRFixtureList> MVRFixtureList;
+	TSharedPtr<SDMXFixturePatchList> FixturePatchList;
 
 	/** Details View for the selected Fixture Patches */
 	TSharedPtr<IDetailsView> FixturePatchDetailsView;
diff --git a/Engine/Plugins/VirtualProduction/DMX/DMXEngine/Source/DMXEditor/Private/Widgets/FixturePatch/SDMXFixturePatchList.cpp b/Engine/Plugins/VirtualProduction/DMX/DMXEngine/Source/DMXEditor/Private/Widgets/FixturePatch/SDMXFixturePatchList.cpp
new file mode 100644
--- /dev/null
+++ b/Engine/Plugins/VirtualProduction/DMX/DMXEngine/Source/DMXEditor/Private/Widgets/FixturePatch/SDMXFixturePatchList.cpp
@@ -0,0 +1,1409 @@
+// Copyright Epic Games, Inc. All Rights Reserved.
+
+#include "SDMXFixturePatchList.h"
+
+#include "Algo/Copy.h"
+#include "Algo/MaxElement.h"
+#include "Algo/MinElement.h"
+#include "Algo/Sort.h"
+#include "Algo/Transform.h"
+#include "Commands/DMXEditorCommands.h"
+#include "DMXEditor.h"
+#include "DMXEditorSettings.h"
+#include "DMXEditorUtils.h"
+#include "DMXFixturePatchSharedData.h"
+#include "DMXRuntimeUtils.h"
+#include "Exporters/Exporter.h"
+#include "Factories.h"
+#include "FixturePatchAutoAssignUtility.h"
+#include "Framework/Application/SlateApplication.h"
+#include "Framework/Commands/GenericCommands.h"
+#include "Framework/Commands/UICommandList.h"
+#include "Framework/MultiBox/MultiBoxBuilder.h"
+#include "HAL/PlatformApplicationMisc.h"
+#include "Library/DMXEntityFixturePatch.h"
+#include "Library/DMXEntityFixtureType.h"
+#include "Library/DMXLibrary.h"
+#include "ScopedTransaction.h"
+#include "Styling/AppStyle.h"
+#include "TimerManager.h"
+#include "UnrealExporter.h"
+#include "Widgets/FixturePatch/DMXFixturePatchListItem.h"
+#include "Widgets/FixturePatch/SDMXFixturePatchListRow.h"
+#include "Widgets/FixturePatch/SDMXFixturePatchListToolbar.h"
+#include "Widgets/SBoxPanel.h"
+#include "Widgets/Views/SHeaderRow.h"
+#include "Widgets/Views/SListView.h"
+
+#define LOCTEXT_NAMESPACE "SDMXFixturePatchList"
+
+namespace UE::DMX::SDMXFixturePatchList::Private
+{
+	/** Copies a fixture patch as text to the clipboard */
+	static void ClipboardCopyFixturePatches(const TArray<UDMXEntityFixturePatch*>& FixturePatches)
+	{
+		// Clear the mark state for saving.
+		UnMarkAllObjects(EObjectMark(OBJECTMARK_TagExp | OBJECTMARK_TagImp));
+
+		const FExportObjectInnerContext Context;
+		FStringOutputDevice Archive;
+
+		// Export the component object(s) to text for copying
+		for (UDMXEntityFixturePatch* FixturePatch : FixturePatches)
+		{
+			if (!FixturePatch)
+			{
+				continue;
+			}
+
+			// Export the entity object to the given string
+			UExporter::ExportToOutputDevice(&Context, FixturePatch, nullptr, Archive, TEXT("copy"), 0, PPF_ExportsNotFullyQualified | PPF_Copy | PPF_Delimited, false, GetTransientPackage());
+		}
+
+		// Copy text to clipboard
+		FString ExportedText = Archive;
+
+		// Avoid exporting the OnFixturePatchReceived Binding
+		TArray<FString> ExportedTextLines;
+		constexpr bool bCullEmpty = false;
+		ExportedText.ParseIntoArrayLines(ExportedTextLines, bCullEmpty);
+		FString ExportedTextWithoutOnFixturePatchReceivedBinding;
+		for (const FString& String : ExportedTextLines)
+		{
+			if (String.Contains(TEXT("OnFixturePatchReceivedDMX")))
+			{
+				continue;
+			}
+			ExportedTextWithoutOnFixturePatchReceivedBinding.Append(String + LINE_TERMINATOR);
+		}
+
+		FPlatformApplicationMisc::ClipboardCopy(*ExportedTextWithoutOnFixturePatchReceivedBinding);
+	}
+
+	/** Duplicates an existing patch */
+	UDMXEntityFixturePatch* DuplicatePatch(UDMXLibrary* DMXLibrary, UDMXEntityFixturePatch* FixturePatchToDuplicate)
+	{
+		if (!DMXLibrary || !FixturePatchToDuplicate)
+		{
+			return nullptr;
+		}
+
+
+		// Duplicate
+		DMXLibrary->PreEditChange(UDMXLibrary::StaticClass()->FindPropertyByName(UDMXLibrary::GetEntitiesPropertyName()));
+		UDMXEntityFixtureType* FixtureTypeOfPatchToDuplicate = FixturePatchToDuplicate->GetFixtureType();
+		if (FixtureTypeOfPatchToDuplicate && FixtureTypeOfPatchToDuplicate->GetParentLibrary() != DMXLibrary)
+		{
+			FDMXEntityFixtureTypeConstructionParams FixtureTypeConstructionParams;
+			FixtureTypeConstructionParams.DMXCategory = FixtureTypeOfPatchToDuplicate->DMXCategory;
+			FixtureTypeConstructionParams.Modes = FixtureTypeOfPatchToDuplicate->Modes;
+			FixtureTypeConstructionParams.ParentDMXLibrary = DMXLibrary;
+
+			constexpr bool bMarkLibraryDirty = false;
+			UDMXEntityFixtureType::CreateFixtureTypeInLibrary(FixtureTypeConstructionParams, FixtureTypeOfPatchToDuplicate->Name, bMarkLibraryDirty);
+		}
+
+		// Duplicate the Fixture Patch
+		FDMXEntityFixturePatchConstructionParams ConstructionParams;
+		ConstructionParams.FixtureTypeRef = FixturePatchToDuplicate->GetFixtureType();
+		ConstructionParams.ActiveMode = FixturePatchToDuplicate->GetActiveModeIndex();
+		ConstructionParams.UniverseID = FixturePatchToDuplicate->GetUniverseID();
+		ConstructionParams.StartingAddress = FixturePatchToDuplicate->GetStartingChannel();
+
+		constexpr bool bMarkLibraryDirty = false;
+		UDMXEntityFixturePatch* NewFixturePatch = UDMXEntityFixturePatch::CreateFixturePatchInLibrary(ConstructionParams, FixturePatchToDuplicate->Name, bMarkLibraryDirty);
+
+		// Use the same color as the duplicated patch
+		NewFixturePatch->EditorColor = FixturePatchToDuplicate->EditorColor;
+
+		DMXLibrary->PostEditChange();
+
+		return NewFixturePatch;
+	}
+
+	/** Text object factory for pasting DMX Fixture Patches */
+	struct FDMXFixturePatchObjectTextFactory final
+		: public FCustomizableTextObjectFactory
+	{
+		/** Constructor */
+		FDMXFixturePatchObjectTextFactory(UDMXLibrary* InDMXLibrary)
+			: FCustomizableTextObjectFactory(GWarn)
+			, WeakDMXLibrary(InDMXLibrary)
+		{}
+
+		/** Returns true if Fixture Patches can be constructed from the Text Buffer */
+		static bool CanCreate(const FString& InTextBuffer, UDMXLibrary* InDMXLibrary)
+		{
+			TSharedRef<FDMXFixturePatchObjectTextFactory> Factory = MakeShared<FDMXFixturePatchObjectTextFactory>(InDMXLibrary);
+
+			// Create new objects if we're allowed to
+			return Factory->CanCreateObjectsFromText(InTextBuffer);
+		}
+
+		/**
+		 * Constructs a new object factory from the given text buffer. Returns the factor or nullptr if no factory can be created.
+		 * An updated General Scene Description of the library needs be passed explicitly to avoid recurring update calls.
+		 */
+		static bool Create(const FString& InTextBuffer, UDMXLibrary* InDMXLibrary, TArray<UDMXEntityFixturePatch*>& OutNewFixturePatches)
+		{
+			if (!InDMXLibrary)
+			{
+				return false;
+			}
+
+			OutNewFixturePatches.Reset();
+
+			const TSharedRef<FDMXFixturePatchObjectTextFactory> Factory = MakeShared < FDMXFixturePatchObjectTextFactory>(InDMXLibrary);
+
+			// Create new objects if we're allowed to
+			if (Factory->CanCreateObjectsFromText(InTextBuffer))
+			{
+				Factory->WeakDMXLibrary = InDMXLibrary;
+
+				EObjectFlags ObjectFlags = RF_Transactional;
+				Factory->ProcessBuffer(InDMXLibrary, ObjectFlags, InTextBuffer);
+
+				OutNewFixturePatches = Factory->NewFixturePatches;
+			}
+
+			return true;
+		}
+
+	protected:
+		//~ Begin FCustomizableTextObjectFactory implementation
+		virtual bool CanCreateClass(UClass* ObjectClass, bool& bOmitSubObjs) const override
+		{
+			return ObjectClass->IsChildOf(UDMXEntityFixturePatch::StaticClass());
+		}
+
+		virtual void ProcessConstructedObject(UObject* NewObject) override
+		{
+			UDMXLibrary* DMXLibrary = WeakDMXLibrary.Get();
+			UDMXEntityFixturePatch* NewFixturePatch = Cast<UDMXEntityFixturePatch>(NewObject);
+			if (DMXLibrary && NewFixturePatch)
+			{
+				const FScopedTransaction Transaction(TransactionText);
+
+				const bool bIsDuplicating = Algo::FindBy(DMXLibrary->GetEntitiesTypeCast<UDMXEntityFixturePatch>(), NewFixturePatch->GetMVRFixtureUUID(), &UDMXEntityFixturePatch::GetMVRFixtureUUID) != nullptr;
+				if (bIsDuplicating)
+				{					
+					// Create a duplicate of the new patch that is properly initialized
+					using namespace UE::DMX::SDMXFixturePatchList::Private;
+					UDMXEntityFixturePatch* DuplicatedFixturePatch = DuplicatePatch(DMXLibrary, NewFixturePatch);
+					NewFixturePatches.Add(DuplicatedFixturePatch);
+
+					// Remove the patch created by the factory
+					UDMXEntityFixturePatch::RemoveFixturePatchFromLibrary(NewFixturePatch);
+				}
+				else
+				{
+					// Simply assign the new patch to the library
+					NewFixturePatch->Rename(*MakeUniqueObjectName(DMXLibrary, NewFixturePatch->GetClass()).ToString(), DMXLibrary, REN_DoNotDirty | REN_DontCreateRedirectors);
+					NewFixturePatch->SetName(FDMXRuntimeUtils::FindUniqueEntityName(DMXLibrary, NewFixturePatch->GetClass(), NewFixturePatch->GetDisplayName()));
+					NewFixturePatch->SetParentLibrary(DMXLibrary);
+					NewFixturePatch->RefreshID();
+
+					NewFixturePatches.Add(NewFixturePatch);
+				}
+			}
+		}
+		//~ End FCustomizableTextObjectFactory implementation
+
+	private:
+		/** Instantiated Fixture Patches */
+		TArray<UDMXEntityFixturePatch*> NewFixturePatches;
+
+		/** Transaction text displayed when pasting */
+		FText TransactionText;
+
+		/** Weak DMX Editor in which the operation should occur */
+		TWeakObjectPtr<UDMXLibrary> WeakDMXLibrary;
+	};
+}
+
+
+/** Helper to generate Status Text for MVR Fixture List Items */
+class FDMXFixturePatchListStatusTextGenerator
+{
+public:
+	FDMXFixturePatchListStatusTextGenerator(const TArray<TSharedPtr<FDMXFixturePatchListItem>>& InItems)
+		: Items(InItems)
+	{}
+
+	/** Generates warning texts. Returns a map of those Items that need a warning set along with the warning Text */
+	TMap<TSharedPtr<FDMXFixturePatchListItem>, FText> GenerateWarningTexts() const
+	{
+		TMap<TSharedPtr<FDMXFixturePatchListItem>, FText> AccumulatedConflicts;
+
+		TMap<TSharedPtr<FDMXFixturePatchListItem>, FText> FixtureTypeIssues = GetFixtureTypeIssues();
+		AppendConflictTexts(FixtureTypeIssues, AccumulatedConflicts);
+
+		TMap<TSharedPtr<FDMXFixturePatchListItem>, FText> FixtureIDIssues = GetFixtureIDIssues();
+		AppendConflictTexts(FixtureIDIssues, AccumulatedConflicts);
+
+		TMap<TSharedPtr<FDMXFixturePatchListItem>, FText> FixtureIDConflicts = GetFixtureIDConflicts();
+		AppendConflictTexts(FixtureIDConflicts, AccumulatedConflicts);
+
+		TMap<TSharedPtr<FDMXFixturePatchListItem>, FText> ChannelExcessConflicts = GetChannelExcessConflicts();
+		AppendConflictTexts(ChannelExcessConflicts, AccumulatedConflicts);
+
+		TMap<TSharedPtr<FDMXFixturePatchListItem>, FText> ChannelOverlapConflicts = GetChannelOverlapConflicts();
+		AppendConflictTexts(ChannelOverlapConflicts, AccumulatedConflicts);
+
+		return AccumulatedConflicts;
+	}
+
+private:
+	void AppendConflictTexts(const TMap<TSharedPtr<FDMXFixturePatchListItem>, FText>& InItemToConflictTextMap, TMap<TSharedPtr<FDMXFixturePatchListItem>, FText>& InOutConflictTexts) const
+	{
+		for (const TTuple<TSharedPtr<FDMXFixturePatchListItem>, FText>& ItemToConflictTextPair : InItemToConflictTextMap)
+		{
+			if (InOutConflictTexts.Contains(ItemToConflictTextPair.Key))
+			{
+				const FText LineTerminator = FText::FromString(LINE_TERMINATOR);
+				const FText AccumulatedErrorText = FText::Format(FText::FromString(TEXT("{0}{1}{2}{3}")), InOutConflictTexts[ItemToConflictTextPair.Key], LineTerminator, LineTerminator, ItemToConflictTextPair.Value);
+				InOutConflictTexts[ItemToConflictTextPair.Key] = AccumulatedErrorText;
+			}
+			else
+			{
+				InOutConflictTexts.Add(ItemToConflictTextPair);
+			}
+		}
+	}
+
+	/** The patch of an item. Useful to Get Conflicts with Other */
+	struct FItemPatch
+	{
+		FItemPatch(const TSharedPtr<FDMXFixturePatchListItem>& InItem)
+			: Item(InItem)
+		{
+			Universe = Item->GetUniverse();
+			AddressRange = TRange<int32>(Item->GetAddress(), Item->GetAddress() + Item->GetNumChannels());
+		}
+
+		/** Returns a conflict text if this item conflicts with Other */
+		FText GetConfictsWithOther(const FItemPatch& Other) const
+		{
+			// No conflict with self
+			if (Other.Item == Item)
+			{
+				return FText::GetEmpty();
+			}
+
+			// No conflict with the same patch
+			if (Item->GetFixturePatch() == Other.Item->GetFixturePatch())
+			{
+				return FText::GetEmpty();
+			}
+
+			// No conflict if not in the same universe
+			if (Other.Universe != Universe)
+			{
+				return FText::GetEmpty();
+			}
+
+			// No conflict if channels don't overlap
+			if (!AddressRange.Overlaps(Other.AddressRange))
+			{
+				return FText::GetEmpty();
+			}
+
+			// No conflict if patches are functionally equal
+			if (AddressRange.GetLowerBound() == Other.AddressRange.GetLowerBound() &&
+				Item->GetFixtureType() == Other.Item->GetFixtureType() &&
+				Item->GetModeIndex() == Other.Item->GetModeIndex())
+			{
+				return FText::GetEmpty();
+			}
+
+			const FText FixtureIDText = MakeBeautifulItemText(Item);
+			const FText OtherFixtureIDText = MakeBeautifulItemText(Other.Item);
+			if (AddressRange.GetLowerBound() == Other.AddressRange.GetLowerBound() &&
+				Item->GetFixtureType() == Other.Item->GetFixtureType())
+			{
+				// Modes confict
+				check(Item->GetModeIndex() != Other.Item->GetModeIndex());
+				return FText::Format(LOCTEXT("ModeConflict", "Uses same Address and Fixture Type as Fixture {1}, but Modes differ."), FixtureIDText, OtherFixtureIDText);
+			}
+			else if (AddressRange.GetLowerBound() == Other.AddressRange.GetLowerBound())
+			{
+				// Fixture Types conflict
+				check(Item->GetFixtureType() != Other.Item->GetFixtureType());
+				return FText::Format(LOCTEXT("FixtureTypeConflict", "Uses same Address as Fixture {1}, but Fixture Types differ."), FixtureIDText, OtherFixtureIDText);
+			}
+			else
+			{
+				// Addresses conflict
+				return FText::Format(LOCTEXT("AddressConflict", "Overlaps Addresses with Fixture {1}"), FixtureIDText, OtherFixtureIDText);
+			}
+		}
+
+		FORCEINLINE const TSharedPtr<FDMXFixturePatchListItem>& GetItem() const { return Item; };
+
+	private:
+		int32 Universe = -1;
+		TRange<int32> AddressRange;
+
+		TSharedPtr<FDMXFixturePatchListItem> Item;
+	};
+
+	/** Returns a Map of Items to Channels that have Fixture Types with issues set */
+	TMap<TSharedPtr<FDMXFixturePatchListItem>, FText> GetFixtureTypeIssues() const
+	{
+		TMap<TSharedPtr<FDMXFixturePatchListItem>, FText> ItemToIssueMap;
+		for (const TSharedPtr<FDMXFixturePatchListItem>& Item : Items)
+		{
+			if (!Item->GetFixtureType())
+			{
+				const FText IssueText = LOCTEXT("NoFixtureTypeIssue", "No Fixture Type selected.");
+				ItemToIssueMap.Add(Item, IssueText);
+			}
+			else if (Item->GetFixtureType()->Modes.IsEmpty())
+			{
+				const FText IssueText = LOCTEXT("NoModesIssue", "Fixture Type has no Modes defined.");
+				ItemToIssueMap.Add(Item, IssueText);
+			}
+			else if (Item->GetFixturePatch()->GetActiveMode() && 
+				!Item->GetFixturePatch()->GetActiveMode()->bFixtureMatrixEnabled &&
+				Item->GetFixturePatch()->GetActiveMode()->Functions.IsEmpty())
+			{
+				const FText IssueText = LOCTEXT("ActiveModeHasNoFunctionsIssue", "Mode does not define any Functions.");
+				ItemToIssueMap.Add(Item, IssueText);
+			}
+		}
+
+		return ItemToIssueMap;
+	}
+
+	/** Returns a Map of Items to Channels exceeding the DMX address range Texts */
+	TMap<TSharedPtr<FDMXFixturePatchListItem>, FText> GetChannelExcessConflicts() const
+	{
+		TMap<TSharedPtr<FDMXFixturePatchListItem>, FText> ItemToConflictMap;
+		for (const TSharedPtr<FDMXFixturePatchListItem>& Item : Items)
+		{
+			const int32 EndingAddress = Item->GetAddress() + Item->GetNumChannels() - 1;
+			if (Item->GetAddress() < 1 &&
+				EndingAddress > DMX_MAX_ADDRESS)
+			{
+				const FText ConflictText = FText::Format(LOCTEXT("ChannelExceedsMinAndMaxChannelConflict", "Exceeds available DMX Address range. Staring Address is {0} but min Address is 1. Ending Address is {1} but max Address is 512."), Item->GetAddress(), EndingAddress);
+				ItemToConflictMap.Add(Item, ConflictText);
+			}
+			else if (Item->GetAddress() < 1)
+			{
+				const FText ConflictText = FText::Format(LOCTEXT("ChannelExceedsMinChannelNumberConflict", "Exceeds available DMX Address range. Staring Address is {0} but min Address is 1."), Item->GetAddress());
+				ItemToConflictMap.Add(Item, ConflictText);
+			}
+			else if (EndingAddress > DMX_MAX_ADDRESS)
+			{
+				const FText ConflictText = FText::Format(LOCTEXT("ChannelExeedsMaxChannelNumberConflict", "Exceeds available DMX Address range. Ending Address is {0} but max Address is 512."), EndingAddress);
+				ItemToConflictMap.Add(Item, ConflictText);
+			}			
+		}
+
+		return ItemToConflictMap;
+	}
+
+	/** Returns a Map of Items to overlapping Channel conflict Texts */
+	TMap<TSharedPtr<FDMXFixturePatchListItem>, FText> GetChannelOverlapConflicts() const
+	{
+		TArray<FItemPatch> ItemPatches;
+		ItemPatches.Reserve(Items.Num());
+		for (const TSharedPtr<FDMXFixturePatchListItem>& Item : Items)
+		{
+			FItemPatch ItemPatch(Item);
+			ItemPatches.Add(ItemPatch);
+		}
+
+		TMap<TSharedPtr<FDMXFixturePatchListItem>, FText> ItemToConflictMap;
+		for (const FItemPatch& ItemPatch : ItemPatches)
+		{
+			for (const FItemPatch& Other : ItemPatches)
+			{
+				const FText ConflictWithOtherText = ItemPatch.GetConfictsWithOther(Other);
+				if (!ConflictWithOtherText.IsEmpty())
+				{
+					if (ItemToConflictMap.Contains(ItemPatch.GetItem()))
+					{
+						FText AppendConflictText = FText::Format(FText::FromString(TEXT("{0}{1}{2}")), ItemToConflictMap[ItemPatch.GetItem()], FText::FromString(FString(LINE_TERMINATOR)), ConflictWithOtherText);
+						ItemToConflictMap[ItemPatch.GetItem()] = AppendConflictText;
+					}
+					else
+					{
+						ItemToConflictMap.Add(ItemPatch.GetItem(), ConflictWithOtherText);
+					}
+				}
+			}
+		}
+
+		return ItemToConflictMap;
+	}
+
+	/** Returns an Map of Items to Fixture IDs issues Texts */
+	TMap<TSharedPtr<FDMXFixturePatchListItem>, FText> GetFixtureIDIssues() const
+	{
+		TMap<TSharedPtr<FDMXFixturePatchListItem>, FText> Result;
+		for (const TSharedPtr<FDMXFixturePatchListItem>& Item : Items)
+		{
+			int32 FixtureIDNumerical;
+			if (!LexTryParseString(FixtureIDNumerical, *Item->GetFixtureID()))
+			{
+				Result.Add(Item, LOCTEXT("FixtureIDNotNumericalIssueText", "FID has to be a number."));
+			}
+		}
+		return Result;
+	}
+
+	/** Returns an Map of Items to Fixture IDs conflict Texts */
+	TMap<TSharedPtr<FDMXFixturePatchListItem>, FText> GetFixtureIDConflicts() const
+	{
+		TMap<FString, TArray<TSharedPtr<FDMXFixturePatchListItem>>> FixtureIDMap;
+		FixtureIDMap.Reserve(Items.Num());
+		for (const TSharedPtr<FDMXFixturePatchListItem>& Item : Items)
+		{
+			FixtureIDMap.FindOrAdd(Item->GetFixtureID()).Add(Item);
+		}
+		TArray<TArray<TSharedPtr<FDMXFixturePatchListItem>>> FixtureIDConflicts;
+		FixtureIDMap.GenerateValueArray(FixtureIDConflicts);
+		FixtureIDConflicts.RemoveAll([](const TArray<TSharedPtr<FDMXFixturePatchListItem>>& ConflictingItems)
+			{
+				return ConflictingItems.Num() < 2;
+			});
+		
+		TMap<TSharedPtr<FDMXFixturePatchListItem>, FText> ItemToConflictMap;
+		for (TArray<TSharedPtr<FDMXFixturePatchListItem>>& ConflictingItems : FixtureIDConflicts)
+		{
+			ConflictingItems.Sort([](const TSharedPtr<FDMXFixturePatchListItem>& ItemA, const TSharedPtr<FDMXFixturePatchListItem>& ItemB)
+				{
+					return ItemA->GetFixtureID() < ItemB->GetFixtureID();
+				});
+
+			check(ConflictingItems.Num() > 0);
+			FText ConflictText = FText::Format(LOCTEXT("BaseFixtureIDConflictText", "Ambiguous FIDs in {0}"), MakeBeautifulItemText(ConflictingItems[0]));
+			for (int32 ConflictingItemIndex = 1; ConflictingItemIndex < ConflictingItems.Num(); ConflictingItemIndex++)
+			{
+				ConflictText = FText::Format(LOCTEXT("AppendFixtureIDConflictText", "{0}, {1}"), ConflictText, MakeBeautifulItemText(ConflictingItems[ConflictingItemIndex]));
+			}
+			
+			for (const TSharedPtr<FDMXFixturePatchListItem>& Item : ConflictingItems)
+			{
+				ItemToConflictMap.Add(Item, ConflictText);
+			}
+		}
+
+		return ItemToConflictMap;
+	}
+
+	static FText MakeBeautifulItemText(const TSharedPtr<FDMXFixturePatchListItem>& Item)
+	{
+		const FString AddressesString = FString::FromInt(Item->GetUniverse()) + TEXT(".") + FString::FromInt(Item->GetAddress());
+		const FString ItemNameString = TEXT("'") + Item->GetFixturePatchName() + TEXT("'");
+		const FString BeautifulItemString = ItemNameString + TEXT(" (") + AddressesString + TEXT(")");;
+		return FText::FromString(BeautifulItemString);
+	}
+
+	/** The items the class handles */
+	TArray<TSharedPtr<FDMXFixturePatchListItem>> Items;
+};
+
+const FName FDMXFixturePatchListCollumnID::EditorColor = "EditorColor";
+const FName FDMXFixturePatchListCollumnID::FixturePatchName = "FixturePatchName";
+const FName FDMXFixturePatchListCollumnID::Status = "Status";
+const FName FDMXFixturePatchListCollumnID::FixtureID = "FixtureID";
+const FName FDMXFixturePatchListCollumnID::FixtureType = "FixtureType";
+const FName FDMXFixturePatchListCollumnID::Mode = "Mode";
+const FName FDMXFixturePatchListCollumnID::Patch = "Patch";
+
+SDMXFixturePatchList::SDMXFixturePatchList()
+	: SortMode(EColumnSortMode::Ascending)
+	, SortedByColumnID(FDMXFixturePatchListCollumnID::FixtureID)
+{}
+
+SDMXFixturePatchList::~SDMXFixturePatchList()
+{
+	if (HeaderRow.IsValid())
+	{
+		UDMXEditorSettings* EditorSettings = GetMutableDefault<UDMXEditorSettings>();
+		for (const SHeaderRow::FColumn& Column : HeaderRow->GetColumns())
+		{
+			if (Column.ColumnId == FDMXFixturePatchListCollumnID::FixtureID)
+			{
+				EditorSettings->MVRFixtureListSettings.FixtureIDColumnWidth = Column.Width.Get();
+			}
+			if (Column.ColumnId == FDMXFixturePatchListCollumnID::FixtureType)
+			{
+				EditorSettings->MVRFixtureListSettings.FixtureTypeColumnWidth = Column.Width.Get();
+			}
+			else if (Column.ColumnId == FDMXFixturePatchListCollumnID::Mode)
+			{
+				EditorSettings->MVRFixtureListSettings.ModeColumnWidth = Column.Width.Get();
+			}
+			else if (Column.ColumnId == FDMXFixturePatchListCollumnID::Patch)
+			{
+				EditorSettings->MVRFixtureListSettings.PatchColumnWidth = Column.Width.Get();
+			}
+		}
+
+		EditorSettings->SaveConfig();
+	}
+}
+
+void SDMXFixturePatchList::PostUndo(bool bSuccess)
+{
+	RequestListRefresh();
+}
+
+void SDMXFixturePatchList::PostRedo(bool bSuccess)
+{
+	RequestListRefresh();
+}
+
+void SDMXFixturePatchList::Construct(const FArguments& InArgs, TWeakPtr<FDMXEditor> InDMXEditor)
+{
+	if (!InDMXEditor.IsValid())
+	{
+		return;
+	}
+	
+	WeakDMXEditor = InDMXEditor;
+	FixturePatchSharedData = InDMXEditor.Pin()->GetFixturePatchSharedData();
+
+	const UDMXEditorSettings* EditorSettings = GetDefault<UDMXEditorSettings>();
+	SortedByColumnID = EditorSettings->MVRFixtureListSettings.SortByCollumnID;
+	SortMode = static_cast<EColumnSortMode::Type>(EditorSettings->MVRFixtureListSettings.SortPriorityEnumIndex);
+
+	// Handle Entity changes
+	UDMXLibrary::GetOnEntitiesAdded().AddSP(this, &SDMXFixturePatchList::OnEntityAddedOrRemoved);
+	UDMXLibrary::GetOnEntitiesRemoved().AddSP(this, &SDMXFixturePatchList::OnEntityAddedOrRemoved);
+	UDMXEntityFixturePatch::GetOnFixturePatchChanged().AddSP(this, &SDMXFixturePatchList::OnFixturePatchChanged);
+	UDMXEntityFixtureType::GetOnFixtureTypeChanged().AddSP(this, &SDMXFixturePatchList::OnFixtureTypeChanged);
+
+	// Handle Shared Data selection changes
+	FixturePatchSharedData->OnFixturePatchSelectionChanged.AddSP(this, &SDMXFixturePatchList::OnFixturePatchSharedDataSelectedFixturePatches);
+
+	static const FTableViewStyle TableViewStyle = FAppStyle::Get().GetWidgetStyle<FTableViewStyle>("TreeView");
+
+	ChildSlot
+	[
+		SNew(SVerticalBox)
+
+		+ SVerticalBox::Slot()
+		.HAlign(HAlign_Fill)
+		.AutoHeight()
+		[
+			SAssignNew(Toolbar, SDMXFixturePatchListToolbar, WeakDMXEditor)
+			.OnSearchChanged(this, &SDMXFixturePatchList::OnSearchChanged)
+		]
+
+		+ SVerticalBox::Slot()
+		.HAlign(HAlign_Fill)
+		.FillHeight(1.f)
+		[
+			SAssignNew(ListView, SDMXFixturePatchListType)
+			.ListViewStyle(&TableViewStyle)
+			.HeaderRow(GenerateHeaderRow())
+			.ListItemsSource(&ListSource)
+			.OnGenerateRow(this, &SDMXFixturePatchList::OnGenerateRow)
+			.OnSelectionChanged(this, &SDMXFixturePatchList::OnSelectionChanged)
+			.OnContextMenuOpening(this, &SDMXFixturePatchList::OnContextMenuOpening)
+		]
+	];
+
+	RegisterCommands();
+	RefreshList();
+
+	AdoptSelectionFromFixturePatchSharedData();
+
+	// Make an initial selection if nothing was selected from Fixture Patch Shared Data, as if the user clicked it
+	if (ListView->GetSelectedItems().IsEmpty() && !ListSource.IsEmpty())
+	{
+		ListView->SetSelection(ListSource[0], ESelectInfo::OnMouseClick);
+	}
+}
+
+void SDMXFixturePatchList::RequestListRefresh()
+{
+	if (RequestListRefreshTimerHandle.IsValid())
+	{
+		return;
+	}
+
+	RequestListRefreshTimerHandle = GEditor->GetTimerManager()->SetTimerForNextTick(FTimerDelegate::CreateSP(this, &SDMXFixturePatchList::RefreshList));
+}
+
+void SDMXFixturePatchList::EnterFixturePatchNameEditingMode()
+{
+	if (ListView.IsValid())
+	{
+		const TArray<TSharedPtr<FDMXFixturePatchListItem>> SelectedItems = ListView->GetSelectedItems();
+		if (SelectedItems.Num() == 0)
+		{
+			const TSharedPtr<SDMXFixturePatchListRow>* SelectedRowPtr = Rows.FindByPredicate([&SelectedItems](const TSharedPtr<SDMXFixturePatchListRow>& Row)
+				{
+					return Row->GetItem() == SelectedItems[0];
+				});
+			if (SelectedRowPtr)
+			{
+				(*SelectedRowPtr)->EnterFixturePatchNameEditingMode();
+			}
+		}
+	}
+}	
+
+FReply SDMXFixturePatchList::ProcessCommandBindings(const FKeyEvent& InKeyEvent)
+{
+	if (CommandList->ProcessCommandBindings(InKeyEvent))
+	{
+		return FReply::Handled();
+	}
+
+	return FReply::Unhandled();
+}
+
+FReply SDMXFixturePatchList::OnKeyDown(const FGeometry& MyGeometry, const FKeyEvent& InKeyEvent)
+{
+	return ProcessCommandBindings(InKeyEvent);
+}
+
+void SDMXFixturePatchList::OnSearchChanged()
+{
+	RequestListRefresh();
+}
+
+void SDMXFixturePatchList::RefreshList()
+{	
+	RequestListRefreshTimerHandle.Invalidate();
+
+	// Clear cached data
+	Rows.Reset();
+	ListSource.Reset();
+
+	const TSharedPtr<FDMXEditor> DMXEditor = WeakDMXEditor.Pin();
+	UDMXLibrary* DMXLibrary = DMXEditor.IsValid() ? DMXEditor->GetDMXLibrary() : nullptr;
+	if (!IsValid(DMXLibrary))
+	{
+		return;
+	}
+
+	// Remove all that are no longer a patch
+	const TArray<UDMXEntityFixturePatch*> FixturePatches = DMXLibrary->GetEntitiesTypeCast<UDMXEntityFixturePatch>();
+
+	Algo::TransformIf(FixturePatches, ListSource,
+		[](const UDMXEntityFixturePatch* FixturePatch)
+		{
+			return FixturePatch != nullptr;
+		},
+		[&DMXEditor](UDMXEntityFixturePatch* FixturePatch)
+		{
+			return MakeShared<FDMXFixturePatchListItem>(DMXEditor.ToSharedRef(), FixturePatch);
+		});
+	SortListSource(EColumnSortPriority::Max, SortedByColumnID, SortMode);
+
+	// Generate status texts
+	GenereateStatusText();
+
+	// Apply search filters. Relies on up-to-date status to find conflicts.
+	if (Toolbar.IsValid())
+	{
+		ListSource = Toolbar->FilterItems(ListSource);
+	}
+	
+	// Update and sort the list and its widgets
+	ListView->RebuildList();
+
+	AdoptSelectionFromFixturePatchSharedData();
+}
+
+void SDMXFixturePatchList::GenereateStatusText()
+{
+	for (const TSharedPtr<FDMXFixturePatchListItem>& Item : ListSource)
+	{
+		Item->WarningStatusText = FText::GetEmpty();
+		Item->ErrorStatusText = FText::GetEmpty();
+	}
+
+	FDMXFixturePatchListStatusTextGenerator StatusTextGenerator(ListSource);
+
+	const TMap<TSharedPtr<FDMXFixturePatchListItem>, FText> WarningTextMap = StatusTextGenerator.GenerateWarningTexts();
+	for (const TTuple<TSharedPtr<FDMXFixturePatchListItem>, FText>& ItemToWarningTextPair : WarningTextMap)
+	{
+		ItemToWarningTextPair.Key->WarningStatusText = ItemToWarningTextPair.Value;
+	}
+}
+
+TSharedRef<ITableRow> SDMXFixturePatchList::OnGenerateRow(TSharedPtr<FDMXFixturePatchListItem> InItem, const TSharedRef<STableViewBase>& OwnerTable)
+{
+	const TSharedRef<SDMXFixturePatchListRow> NewRow =
+		SNew(SDMXFixturePatchListRow, OwnerTable, InItem.ToSharedRef())
+		.OnRowRequestsListRefresh(this, &SDMXFixturePatchList::RequestListRefresh)
+		.OnRowRequestsStatusRefresh(this, &SDMXFixturePatchList::GenereateStatusText)
+		.IsSelected_Lambda([this, InItem]()
+			{
+				return ListView->IsItemSelected(InItem);
+			});
+
+	Rows.Add(NewRow);
+
+	return NewRow;
+}
+
+void SDMXFixturePatchList::OnSelectionChanged(TSharedPtr<FDMXFixturePatchListItem> InItem, ESelectInfo::Type SelectInfo)
+{
+	if (SelectInfo == ESelectInfo::Direct)
+	{
+		return;
+	}
+
+	LastSelectedItem = InItem;
+ 
+	const TArray<TSharedPtr<FDMXFixturePatchListItem>> SelectedItems = ListView->GetSelectedItems();
+	TArray<TWeakObjectPtr<UDMXEntityFixturePatch>> FixturePatchesToSelect;
+	for (const TSharedPtr<FDMXFixturePatchListItem>& Item : SelectedItems)
+	{
+		if (UDMXEntityFixturePatch* FixturePatch = Item->GetFixturePatch())
+		{
+			FixturePatchesToSelect.AddUnique(FixturePatch);
+		}
+	}
+
+	FixturePatchSharedData->SelectFixturePatches(FixturePatchesToSelect);
+
+	if (FixturePatchesToSelect.Num() > 0)
+	{
+		const int32 SelectedUniverse = FixturePatchSharedData->GetSelectedUniverse();
+		const int32 UniverseOfFirstItem = FixturePatchesToSelect[0]->GetUniverseID();
+		if (SelectedUniverse != UniverseOfFirstItem)
+		{
+			FixturePatchSharedData->SelectUniverse(UniverseOfFirstItem);
+		}
+	}
+}
+
+void SDMXFixturePatchList::OnEntityAddedOrRemoved(UDMXLibrary* DMXLibrary, TArray<UDMXEntity*> Entities)
+{
+	if (bChangingDMXLibrary)
+	{
+		return;
+	}
+
+	RequestListRefresh();
+}
+
+void SDMXFixturePatchList::OnFixturePatchChanged(const UDMXEntityFixturePatch* FixturePatch)
+{
+	if (bChangingDMXLibrary)
+	{
+		return;
+	}
+
+	// Refresh only if the fixture patch is in the library this editor handles
+	const UDMXLibrary* DMXLibrary = WeakDMXEditor.IsValid() ? WeakDMXEditor.Pin()->GetDMXLibrary() : nullptr;
+	if (FixturePatch && FixturePatch->GetParentLibrary() == DMXLibrary)
+	{
+		RequestListRefresh();
+	}
+}
+
+void SDMXFixturePatchList::OnFixtureTypeChanged(const UDMXEntityFixtureType* FixtureType)
+{
+	if (bChangingDMXLibrary)
+	{
+		return;
+	}
+
+	// Refresh only if the fixture type is in the library this editor handles
+	const UDMXLibrary* DMXLibrary = WeakDMXEditor.IsValid() ? WeakDMXEditor.Pin()->GetDMXLibrary() : nullptr;
+	if (FixtureType && FixtureType->GetParentLibrary() == DMXLibrary)
+	{
+		RequestListRefresh();
+	}
+}
+
+void SDMXFixturePatchList::OnFixturePatchSharedDataSelectedFixturePatches()
+{
+	if (!bChangingDMXLibrary)
+	{
+		TArray<TWeakObjectPtr<UDMXEntityFixturePatch>> SelectedFixturePatches = FixturePatchSharedData->GetSelectedFixturePatches();
+		SelectedFixturePatches.RemoveAll([](TWeakObjectPtr<UDMXEntityFixturePatch> FixturePatch)
+			{
+				return !FixturePatch.IsValid();
+			});
+
+		TArray<TSharedPtr<FDMXFixturePatchListItem>> NewSelection;
+		for (const TSharedPtr<FDMXFixturePatchListItem>& Item : ListSource)
+		{
+			if (SelectedFixturePatches.Contains(Item->GetFixturePatch()))
+			{
+				NewSelection.Add(Item);
+			}
+		}
+
+		if (NewSelection.Num() > 0)
+		{
+			ListView->ClearSelection();
+			ListView->SetItemSelection(NewSelection, true, ESelectInfo::OnMouseClick);
+		}
+		else
+		{
+			ListView->ClearSelection();
+		}
+	}
+}
+
+void SDMXFixturePatchList::AdoptSelectionFromFixturePatchSharedData()
+{
+	if (!ListView.IsValid())
+	{
+		return;
+	}
+
+	const TArray<TWeakObjectPtr<UDMXEntityFixturePatch>> SelectedFixturePatches = FixturePatchSharedData->GetSelectedFixturePatches();
+
+	TArray<TSharedPtr<FDMXFixturePatchListItem>> NewSelection;
+	for (const TWeakObjectPtr<UDMXEntityFixturePatch> SelectedFixturePatch : SelectedFixturePatches)
+	{
+		const TSharedPtr<FDMXFixturePatchListItem>* SelectedItemPtr = ListSource.FindByPredicate([SelectedFixturePatch](const TSharedPtr<FDMXFixturePatchListItem>& Item)
+			{
+				return SelectedFixturePatch.IsValid() && Item->GetFixturePatch() == SelectedFixturePatch;
+			});
+
+		if (SelectedItemPtr)
+		{
+			NewSelection.Add(*SelectedItemPtr);
+		}
+	}
+
+	if (NewSelection.Num() > 0)
+	{
+		ListView->ClearSelection();
+
+		constexpr bool bSelected = true;
+		ListView->SetItemSelection(NewSelection, bSelected, ESelectInfo::OnMouseClick);
+		ListView->RequestScrollIntoView(NewSelection[0]);
+	}
+}
+
+void SDMXFixturePatchList::AutoAssignFixturePatches(UE::DMXEditor::AutoAssign::EAutoAssignMode Mode)
+{
+	if (!WeakDMXEditor.IsValid())
+	{
+		return;
+	}
+
+	TArray<UDMXEntityFixturePatch*> FixturePatchesToAutoAssign;
+	const TArray<TWeakObjectPtr<UDMXEntityFixturePatch>> SelectedFixturePatches = FixturePatchSharedData->GetSelectedFixturePatches();
+	for (TWeakObjectPtr<UDMXEntityFixturePatch> WeakFixturePatch : SelectedFixturePatches)
+	{
+		if (UDMXEntityFixturePatch* FixturePatch = WeakFixturePatch.Get())
+		{
+			FixturePatch->PreEditChange(UDMXEntityFixturePatch::StaticClass()->FindPropertyByName(UDMXEntityFixturePatch::GetUniverseIDPropertyNameChecked()));
+			FixturePatch->PreEditChange(UDMXEntityFixturePatch::StaticClass()->FindPropertyByName(UDMXEntityFixturePatch::GetStartingChannelPropertyNameChecked()));
+			FixturePatchesToAutoAssign.Add(FixturePatch);
+		}
+	}
+
+	if (FixturePatchesToAutoAssign.IsEmpty())
+	{
+		return;
+	}
+
+	using namespace UE::DMXEditor::AutoAssign;
+	
+	const int32 DesiredUniverse = 1;
+	const int32 DesiredChannel = 1;
+	FAutoAssignUtility::AutoAssign(Mode, WeakDMXEditor.Pin().ToSharedRef(), FixturePatchesToAutoAssign, DesiredUniverse, DesiredChannel);
+	FixturePatchSharedData->SelectUniverse(FixturePatchesToAutoAssign[0]->GetUniverseID());
+
+	for (UDMXEntityFixturePatch* FixturePatch : FixturePatchesToAutoAssign)
+	{
+		// Post edit change for both properties
+		FixturePatch->PostEditChange();
+		FixturePatch->PostEditChange();
+	}
+
+	RequestListRefresh();
+}
+
+bool SDMXFixturePatchList::DoesDMXLibraryHaveReachableUniverses() const
+{
+	UDMXLibrary* DMXLibrary = WeakDMXEditor.IsValid() ? WeakDMXEditor.Pin()->GetDMXLibrary() : nullptr;
+	if (DMXLibrary)
+	{
+		return !DMXLibrary->GetInputPorts().IsEmpty() && !DMXLibrary->GetOutputPorts().IsEmpty();
+	}
+	return false;
+}
+
+void SDMXFixturePatchList::SetKeyboardFocus()
+{
+	FSlateApplication::Get().SetKeyboardFocus(AsShared());
+}
+
+TSharedRef<SHeaderRow> SDMXFixturePatchList::GenerateHeaderRow()
+{
+	const UDMXEditorSettings* EditorSettings = GetDefault<UDMXEditorSettings>();
+
+	const float StatusColumnWidth = FMath::Max(FAppStyle::GetBrush("Icons.Warning")->GetImageSize().X + 6.f, FAppStyle::GetBrush("Icons.Error")->GetImageSize().X + 6.f);
+	const float PatchColumnWidth = EditorSettings->MVRFixtureListSettings.PatchColumnWidth > .02f ? EditorSettings->MVRFixtureListSettings.PatchColumnWidth : .1f;
+	const float EditorColorColumnWidth = StatusColumnWidth;
+	const float FixtureIDColumnWidth = EditorSettings->MVRFixtureListSettings.FixtureIDColumnWidth > .01f ? EditorSettings->MVRFixtureListSettings.FixtureIDColumnWidth : .1f;
+	const float FixtureTypeColumnWidth = EditorSettings->MVRFixtureListSettings.FixtureTypeColumnWidth > .02f ? EditorSettings->MVRFixtureListSettings.FixtureTypeColumnWidth : .1f;
+	const float ModeColumnWidth = EditorSettings->MVRFixtureListSettings.ModeColumnWidth > .02f ? EditorSettings->MVRFixtureListSettings.ModeColumnWidth : .1f;
+
+	HeaderRow = SNew(SHeaderRow);
+
+	HeaderRow->AddColumn(
+		SHeaderRow::FColumn::FArguments()
+		.ColumnId(FDMXFixturePatchListCollumnID::EditorColor)
+		.DefaultLabel(LOCTEXT("EditorColorColumnLabel", ""))
+		.FixedWidth(EditorColorColumnWidth)
+	);
+
+	HeaderRow->AddColumn(
+		SHeaderRow::FColumn::FArguments()
+		.ColumnId(FDMXFixturePatchListCollumnID::FixturePatchName)
+		.SortMode(this, &SDMXFixturePatchList::GetColumnSortMode, FDMXFixturePatchListCollumnID::FixturePatchName)
+		.OnSort(this, &SDMXFixturePatchList::SortList)
+		.DefaultLabel(LOCTEXT("FixturePatchNameColumnLabel", "Fixture Patch"))
+		.FillWidth(PatchColumnWidth)
+	);
+
+	HeaderRow->AddColumn(
+		SHeaderRow::FColumn::FArguments()
+		.ColumnId(FDMXFixturePatchListCollumnID::Status)
+		.DefaultLabel(LOCTEXT("StatusColumnLabel", ""))
+		.FixedWidth(StatusColumnWidth)
+	);
+
+	HeaderRow->AddColumn(
+		SHeaderRow::FColumn::FArguments()
+		.ColumnId(FDMXFixturePatchListCollumnID::FixtureID)
+		.SortMode(this, &SDMXFixturePatchList::GetColumnSortMode, FDMXFixturePatchListCollumnID::FixtureID)
+		.OnSort(this, &SDMXFixturePatchList::SortList)
+		.DefaultLabel(LOCTEXT("FixtureIDColumnLabel", "FID"))
+		.FillWidth(FixtureIDColumnWidth)
+	);
+
+	HeaderRow->AddColumn(
+		SHeaderRow::FColumn::FArguments()
+		.ColumnId(FDMXFixturePatchListCollumnID::FixtureType)
+		.SortMode(this, &SDMXFixturePatchList::GetColumnSortMode, FDMXFixturePatchListCollumnID::FixtureType)
+		.OnSort(this, &SDMXFixturePatchList::SortList)
+		.DefaultLabel(LOCTEXT("FixtureTypeColumnLabel", "FixtureType"))
+		.FillWidth(FixtureTypeColumnWidth)
+	);
+
+	HeaderRow->AddColumn(
+		SHeaderRow::FColumn::FArguments()
+		.ColumnId(FDMXFixturePatchListCollumnID::Mode)
+		.SortMode(this, &SDMXFixturePatchList::GetColumnSortMode, FDMXFixturePatchListCollumnID::Mode)
+		.OnSort(this, &SDMXFixturePatchList::SortList)
+		.DefaultLabel(LOCTEXT("ModeColumnLabel", "Mode"))
+		.FillWidth(ModeColumnWidth)
+	);
+
+	HeaderRow->AddColumn(
+		SHeaderRow::FColumn::FArguments()
+		.ColumnId(FDMXFixturePatchListCollumnID::Patch)
+		.SortMode(this, &SDMXFixturePatchList::GetColumnSortMode, FDMXFixturePatchListCollumnID::Patch)
+		.OnSort(this, &SDMXFixturePatchList::SortList)
+		.DefaultLabel(LOCTEXT("PatchColumnLabel", "Patch"))
+		.FillWidth(0.1f)
+	);
+
+	return HeaderRow.ToSharedRef();
+}
+
+EColumnSortMode::Type SDMXFixturePatchList::GetColumnSortMode(const FName ColumnId) const
+{
+	return SortMode;
+}
+
+void SDMXFixturePatchList::SortListSource(const EColumnSortPriority::Type SortPriority, const FName& ColumnId, const EColumnSortMode::Type InSortMode)
+{
+	SortedByColumnID = ColumnId;
+	SortMode = InSortMode;
+
+	const bool bAscending = InSortMode == EColumnSortMode::Ascending ? true : false;
+	if (ColumnId == FDMXFixturePatchListCollumnID::FixturePatchName)
+	{
+		Algo::StableSort(ListSource, [bAscending](const TSharedPtr<FDMXFixturePatchListItem>& ItemA, const TSharedPtr<FDMXFixturePatchListItem>& ItemB)
+			{
+				const FString FixturePatchNameA = ItemA->GetFixturePatchName();
+				const FString FixturePatchNameB = ItemB->GetFixturePatchName();
+
+				const bool bIsGreater = FixturePatchNameA >= FixturePatchNameB;
+				return bAscending ? !bIsGreater : bIsGreater;
+			});
+	}
+	else if (ColumnId == FDMXFixturePatchListCollumnID::FixtureID)
+	{
+		Algo::StableSort(ListSource, [bAscending](const TSharedPtr<FDMXFixturePatchListItem>& ItemA, const TSharedPtr<FDMXFixturePatchListItem>& ItemB)
+			{
+				bool bIsGreater = [ItemA, ItemB]()
+				{
+					const FString FixtureIDStringA = ItemA->GetFixtureID();
+					const FString FixtureIDStringB = ItemB->GetFixtureID();
+
+					int32 FixtureIDA = 0;
+					int32 FixtureIDB = 0;
+
+					const bool bCanParseA = LexTryParseString(FixtureIDA, *FixtureIDStringA);
+					const bool bCanParseB = LexTryParseString(FixtureIDB, *FixtureIDStringB);
+					
+					const bool bIsNumeric = bCanParseA && bCanParseB;
+					if (bIsNumeric)
+					{
+						return FixtureIDA >= FixtureIDB;
+					}
+					else
+					{
+						return FixtureIDStringA >= FixtureIDStringB;
+					}
+				}();
+
+				return bAscending ? !bIsGreater : bIsGreater;
+			});
+	}
+	else if (ColumnId == FDMXFixturePatchListCollumnID::FixtureType)
+	{
+		Algo::StableSort(ListSource, [bAscending](const TSharedPtr<FDMXFixturePatchListItem>& ItemA, const TSharedPtr<FDMXFixturePatchListItem>& ItemB)
+			{
+				const FString FixtureTypeA = ItemA->GetFixtureType()->Name;
+				const FString FixtureTypeB = ItemB->GetFixtureType()->Name;
+
+				const bool bIsGreater = FixtureTypeA >= FixtureTypeB;
+				return bAscending ? !bIsGreater : bIsGreater;
+			});
+	}
+	else if (ColumnId == FDMXFixturePatchListCollumnID::Mode)
+	{
+		Algo::StableSort(ListSource, [bAscending](const TSharedPtr<FDMXFixturePatchListItem>& ItemA, const TSharedPtr<FDMXFixturePatchListItem>& ItemB)
+			{
+				const bool bIsGreater = ItemA->GetModeIndex() >= ItemB->GetModeIndex();
+				return bAscending ? !bIsGreater : bIsGreater;
+			});
+	}
+	else if (ColumnId == FDMXFixturePatchListCollumnID::Patch)
+	{
+		Algo::StableSort(ListSource, [bAscending](const TSharedPtr<FDMXFixturePatchListItem>& ItemA, const TSharedPtr<FDMXFixturePatchListItem>& ItemB)
+			{
+				const UDMXEntityFixturePatch* FixturePatchA = ItemA->GetFixturePatch();
+				const UDMXEntityFixturePatch* FixturePatchB = ItemB->GetFixturePatch();
+
+				const bool bIsUniverseIDGreater = ItemA->GetUniverse() > ItemB->GetUniverse();
+				const bool bIsSameUniverse = ItemA->GetUniverse() == ItemB->GetUniverse();
+				const bool bAreAddressesGreater = ItemA->GetAddress() > ItemB->GetAddress();
+
+				const bool bIsGreater = bIsUniverseIDGreater || (bIsSameUniverse && bAreAddressesGreater);
+				return bAscending ? !bIsGreater : bIsGreater;
+			});
+	}
+}
+
+void SDMXFixturePatchList::SortList(const EColumnSortPriority::Type SortPriority, const FName& ColumnId, const EColumnSortMode::Type InSortMode)
+{
+	SortListSource(SortPriority, ColumnId, InSortMode);
+
+	if (ListView.IsValid())
+	{
+		ListView->RequestListRefresh();
+	}
+
+	UDMXEditorSettings* EditorSettings = GetMutableDefault<UDMXEditorSettings>();
+	EditorSettings->MVRFixtureListSettings.SortByCollumnID = SortedByColumnID;
+	EditorSettings->MVRFixtureListSettings.SortPriorityEnumIndex = static_cast<int32>(SortMode);
+
+	EditorSettings->SaveConfig();
+}
+
+TSharedPtr<SWidget> SDMXFixturePatchList::OnContextMenuOpening()
+{
+	if (!ListView.IsValid())
+	{
+		return SNullWidget::NullWidget;
+	}
+
+	const bool bCloseWindowAfterMenuSelection = true;
+	FMenuBuilder MenuBuilder(bCloseWindowAfterMenuSelection, CommandList);
+	MenuBuilder.BeginSection("BasicOperationsSection", LOCTEXT("BasicOperationsSection", "Basic Operations"));
+	{
+		MenuBuilder.AddMenuEntry(FGenericCommands::Get().Cut);
+		MenuBuilder.AddMenuEntry(FGenericCommands::Get().Copy);
+		MenuBuilder.AddMenuEntry(FGenericCommands::Get().Paste);
+		MenuBuilder.AddMenuEntry(FGenericCommands::Get().Duplicate);
+		MenuBuilder.AddMenuEntry(FGenericCommands::Get().Delete);
+	}
+	MenuBuilder.EndSection();
+
+	using namespace UE::DMXEditor::AutoAssign;
+	MenuBuilder.BeginSection("AutoAssignSection", LOCTEXT("AutoAssignActionsSection", "Auto-Assign"));
+	{
+		MenuBuilder.AddMenuEntry(FDMXEditorCommands::Get().AutoAssignSelectedUniverse);
+	}
+	MenuBuilder.EndSection();
+
+	return MenuBuilder.MakeWidget();
+}
+
+void SDMXFixturePatchList::RegisterCommands()
+{
+	if (CommandList.IsValid())
+	{
+		return;
+	}
+
+	CommandList = MakeShared<FUICommandList>();
+	CommandList->MapAction(FGenericCommands::Get().Cut, 
+		FUIAction(
+			FExecuteAction::CreateSP(this, &SDMXFixturePatchList::OnCutSelectedItems),
+			FCanExecuteAction::CreateSP(this, &SDMXFixturePatchList::CanCutItems)
+		)
+	);
+	CommandList->MapAction(FGenericCommands::Get().Copy,
+		FUIAction(
+			FExecuteAction::CreateSP(this, &SDMXFixturePatchList::OnCopySelectedItems),
+			FCanExecuteAction::CreateSP(this, &SDMXFixturePatchList::CanCopyItems)
+		)
+	);
+	CommandList->MapAction(FGenericCommands::Get().Paste,
+		FUIAction(
+			FExecuteAction::CreateSP(this, &SDMXFixturePatchList::OnPasteItems),
+			FCanExecuteAction::CreateSP(this, &SDMXFixturePatchList::CanPasteItems)
+		)
+	);
+	CommandList->MapAction(FGenericCommands::Get().Duplicate,
+		FUIAction(
+			FExecuteAction::CreateSP(this, &SDMXFixturePatchList::OnDuplicateItems),
+			FCanExecuteAction::CreateSP(this, &SDMXFixturePatchList::CanDuplicateItems)
+		)
+	);
+	CommandList->MapAction(FGenericCommands::Get().Delete,
+		FUIAction(
+			FExecuteAction::CreateSP(this, &SDMXFixturePatchList::OnDeleteItems),
+			FCanExecuteAction::CreateSP(this, &SDMXFixturePatchList::CanDeleteItems),
+			EUIActionRepeatMode::RepeatEnabled
+		)
+	);
+
+	using namespace UE::DMXEditor::AutoAssign;
+	CommandList->MapAction
+	(
+		FDMXEditorCommands::Get().AutoAssignSelectedUniverse,
+		FExecuteAction::CreateSP(this, &SDMXFixturePatchList::AutoAssignFixturePatches, EAutoAssignMode::SelectedUniverse)
+	);
+}
+
+bool SDMXFixturePatchList::CanCutItems() const
+{
+	return CanCopyItems() && CanDeleteItems() && !GIsTransacting;
+}
+
+void SDMXFixturePatchList::OnCutSelectedItems()
+{
+	const TArray<TSharedPtr<FDMXFixturePatchListItem>> SelectedItems = ListView->GetSelectedItems();
+
+	const FScopedTransaction Transaction(SelectedItems.Num() > 1 ? LOCTEXT("CutFixturePatches", "Cut Fixtures") : LOCTEXT("CutFixturePatche", "Cut Fixture"));
+
+	OnCopySelectedItems();
+	OnDeleteItems();
+}
+
+bool SDMXFixturePatchList::CanCopyItems() const
+{
+	return FixturePatchSharedData->GetSelectedFixturePatches().Num() > 0 && !GIsTransacting;
+}
+
+void SDMXFixturePatchList::OnCopySelectedItems()
+{
+	const TArray<TSharedPtr<FDMXFixturePatchListItem>> SelectedItems = ListView->GetSelectedItems();
+	TArray<UDMXEntityFixturePatch*> FixturePatchesToCopy;
+	for (const TSharedPtr<FDMXFixturePatchListItem>& Item : SelectedItems)
+	{
+		FixturePatchesToCopy.Add(Item->GetFixturePatch());
+	}
+
+	using namespace UE::DMX::SDMXFixturePatchList::Private;
+	ClipboardCopyFixturePatches(FixturePatchesToCopy);
+}
+
+bool SDMXFixturePatchList::CanPasteItems() const
+{
+	using namespace UE::DMX::SDMXFixturePatchList::Private;
+
+	UDMXLibrary* DMXLibrary = WeakDMXEditor.IsValid() ? WeakDMXEditor.Pin()->GetDMXLibrary() : nullptr;
+	if (!DMXLibrary)
+	{
+		return false;
+	}
+
+	// Get the text from the clipboard
+	FString TextToImport;
+	FPlatformApplicationMisc::ClipboardPaste(TextToImport);
+
+	return FDMXFixturePatchObjectTextFactory::CanCreate(TextToImport, DMXLibrary) && !GIsTransacting;
+}
+
+void SDMXFixturePatchList::OnPasteItems()
+{
+	using namespace UE::DMX::SDMXFixturePatchList::Private;
+
+	const TSharedPtr<FDMXEditor> DMXEditor = WeakDMXEditor.Pin();
+	UDMXLibrary* DMXLibrary = DMXEditor.IsValid() ? DMXEditor->GetDMXLibrary() : nullptr;
+	if (!DMXEditor.IsValid() || !DMXLibrary)
+	{
+		return;
+	}
+
+	TGuardValue Guard(bChangingDMXLibrary, true);
+	const FText TransactionText = LOCTEXT("PasteFixturePatchesTransaction", "Paste Fixture Patches");
+	const FScopedTransaction PasteTransaction(TransactionText);
+
+	FString TextToImport;
+	FPlatformApplicationMisc::ClipboardPaste(TextToImport);
+
+	TArray<UDMXEntityFixturePatch*> PastedFixturePatches;
+	if(FDMXFixturePatchObjectTextFactory::Create(TextToImport, DMXLibrary, PastedFixturePatches))
+	{
+		TArray<TWeakObjectPtr<UDMXEntityFixturePatch>> WeakPastedFixturePatches;
+		for (UDMXEntityFixturePatch* FixturePatch : PastedFixturePatches)
+		{
+			WeakPastedFixturePatches.Add(FixturePatch);
+		}
+
+		// Assign
+		using namespace UE::DMXEditor::AutoAssign;
+		int32 AssignedToUniverse = FAutoAssignUtility::AutoAssign(EAutoAssignMode::SelectedUniverse, DMXEditor.ToSharedRef(), PastedFixturePatches);
+
+		FixturePatchSharedData->SelectUniverse(AssignedToUniverse);
+		FixturePatchSharedData->SelectFixturePatches(WeakPastedFixturePatches);
+	
+		RequestListRefresh();
+	}
+}
+
+bool SDMXFixturePatchList::CanDuplicateItems() const
+{
+	return FixturePatchSharedData->GetSelectedFixturePatches().Num() > 0 && !GIsTransacting;
+}
+
+void SDMXFixturePatchList::OnDuplicateItems()
+{
+	const TSharedPtr<FDMXEditor> DMXEditor = WeakDMXEditor.Pin();
+	UDMXLibrary* DMXLibrary = DMXEditor.IsValid() ? DMXEditor->GetDMXLibrary() : nullptr;
+	if (!DMXEditor.IsValid() || !DMXLibrary)
+	{
+		return;
+	}
+
+	TGuardValue Guard(bChangingDMXLibrary, true);
+
+	const FText TransactionText = LOCTEXT("DuplicateFixturePatchesTransaction", "Duplicate Fixture Patches");
+	const FScopedTransaction PasteTransaction(TransactionText);
+	DMXLibrary->PreEditChange(nullptr);
+
+	TArray<TWeakObjectPtr<UDMXEntityFixturePatch>> SelectedWeakFixturePatches = FixturePatchSharedData->GetSelectedFixturePatches();
+	SelectedWeakFixturePatches.Remove(nullptr);
+	if (SelectedWeakFixturePatches.IsEmpty())
+	{
+		return;
+	}
+
+	TArray<UDMXEntityFixturePatch*> SelectedFixturePatches;
+	Algo::TransformIf(SelectedWeakFixturePatches, SelectedFixturePatches, 
+		[](const TWeakObjectPtr<UDMXEntityFixturePatch>& WeakFixturePatch) { return WeakFixturePatch.IsValid(); },
+		[](const TWeakObjectPtr<UDMXEntityFixturePatch>& WeakFixturePatch) { return WeakFixturePatch.Get(); });
+
+	// Duplicate in order of patch
+	Algo::StableSortBy(SelectedFixturePatches, [this](UDMXEntityFixturePatch* FixturePatch)
+		{
+			return FixturePatch->GetUniverseID() * DMX_UNIVERSE_SIZE + FixturePatch->GetStartingChannel();
+		});
+
+	TArray<UDMXEntityFixturePatch*> NewFixturePatches;
+	for (UDMXEntityFixturePatch* FixturePatch : SelectedFixturePatches)
+	{
+		if(!FixturePatch->GetParentLibrary())
+		{
+			continue;
+		}
+
+		using namespace UE::DMX::SDMXFixturePatchList::Private;
+		UDMXEntityFixturePatch* NewFixturePatch = DuplicatePatch(DMXLibrary, FixturePatch);
+		NewFixturePatches.Add(NewFixturePatch);
+	}
+
+	using namespace UE::DMXEditor::AutoAssign;
+	const int32 AssignedToUniverse = FAutoAssignUtility::AutoAssign(EAutoAssignMode::AfterLastPatchedUniverse, DMXEditor.ToSharedRef(), NewFixturePatches);
+
+	DMXLibrary->PostEditChange();
+
+	TArray<TWeakObjectPtr<UDMXEntityFixturePatch>> NewWeakFixturePatches;
+	Algo::Copy(NewFixturePatches, NewWeakFixturePatches);
+	FixturePatchSharedData->SelectFixturePatches(NewWeakFixturePatches);
+	FixturePatchSharedData->SelectUniverse(AssignedToUniverse);
+
+	RequestListRefresh();
+}
+
+bool SDMXFixturePatchList::CanDeleteItems() const
+{
+	return FixturePatchSharedData->GetSelectedFixturePatches().Num() > 0 && !GIsTransacting;
+}
+
+void SDMXFixturePatchList::OnDeleteItems()
+{
+	TGuardValue Guard(bChangingDMXLibrary, true);
+
+	const TArray<TSharedPtr<FDMXFixturePatchListItem>> SelectedItems = ListView->GetSelectedItems();
+
+	if (SelectedItems.Num() == 0)
+	{
+		return;
+	}
+	
+	UDMXLibrary* DMXLibrary = SelectedItems[0]->GetDMXLibrary();
+	if (!DMXLibrary)
+	{
+		return;
+	}
+	
+	const FText DeleteFixturePatchesTransactionText = FText::Format(LOCTEXT("DeleteFixturePatchesTransaction", "Delete Fixture {0}|plural(one=Patch, other=Patches)"), SelectedItems.Num() > 1);
+	const FScopedTransaction DeleteFixturePatchTransaction(DeleteFixturePatchesTransactionText);
+
+	DMXLibrary->PreEditChange(nullptr);
+	for (const TSharedPtr<FDMXFixturePatchListItem>& Item : SelectedItems)
+	{
+		if (UDMXEntityFixturePatch* FixturePatch = Item->GetFixturePatch())
+		{
+			UDMXEntityFixturePatch::RemoveFixturePatchFromLibrary(FixturePatch);
+		}
+	}
+	DMXLibrary->PostEditChange();
+
+	// Make a meaningful selection invariant to ordering of the List
+	TSharedPtr<FDMXFixturePatchListItem> NewSelection;
+	for (int32 ItemIndex = 0; ItemIndex < ListSource.Num(); ItemIndex++)
+	{
+		if (SelectedItems.Contains(ListSource[ItemIndex]))
+		{
+			if (ListSource.IsValidIndex(ItemIndex + 1) && !SelectedItems.Contains(ListSource[ItemIndex + 1]))
+			{
+				NewSelection = ListSource[ItemIndex + 1];
+				break;
+			}
+			else if (ListSource.IsValidIndex(ItemIndex - 1) && !SelectedItems.Contains(ListSource[ItemIndex - 1]))
+			{
+				NewSelection = ListSource[ItemIndex - 1];
+				break;
+			}
+		}
+	}
+	if (NewSelection.IsValid())
+	{
+		ListView->SetSelection(NewSelection, ESelectInfo::OnMouseClick);
+	}
+
+	RequestListRefresh();
+}
+
+#undef LOCTEXT_NAMESPACE
diff --git a/Engine/Plugins/VirtualProduction/DMX/DMXEngine/Source/DMXEditor/Private/Widgets/FixturePatch/SDMXFixturePatchList.h b/Engine/Plugins/VirtualProduction/DMX/DMXEngine/Source/DMXEditor/Private/Widgets/FixturePatch/SDMXFixturePatchList.h
new file mode 100644
--- /dev/null
+++ b/Engine/Plugins/VirtualProduction/DMX/DMXEngine/Source/DMXEditor/Private/Widgets/FixturePatch/SDMXFixturePatchList.h
@@ -0,0 +1,191 @@
+// Copyright Epic Games, Inc. All Rights Reserved.
+
+#pragma once
+
+#include "EditorUndoClient.h"
+#include "Engine/EngineTypes.h"
+#include "Widgets/SCompoundWidget.h"
+#include "Widgets/Views/SHeaderRow.h"
+
+class FDMXEditor;
+class FDMXFixturePatchListItem;
+class FDMXFixturePatchSharedData;
+class FUICommandList;
+class ITableRow;
+class SBorder;
+class SDMXFixturePatchListRow;
+class SDMXFixturePatchListToolbar;
+class SHeaderRow;
+template <typename ItemType> class SListView;
+class STableViewBase;
+class UDMXEntity;
+class UDMXEntityFixturePatch;
+class UDMXEntityFixtureType;
+class UDMXLibrary;
+class UDMXMVRGeneralSceneDescription;
+namespace UE::DMXEditor::AutoAssign { enum class EAutoAssignMode : uint8; }
+
+/** Collumn IDs in the Fixture Patch List */
+struct FDMXFixturePatchListCollumnID
+{
+	static const FName EditorColor;
+	static const FName FixturePatchName;
+	static const FName Status;
+	static const FName FixtureID;
+	static const FName FixtureType;
+	static const FName Mode;
+	static const FName Patch;
+};
+
+using SDMXFixturePatchListType = SListView<TSharedPtr<FDMXFixturePatchListItem>>;
+
+/** Sortable, editable List of Fixture Patches in the library */
+class SDMXFixturePatchList
+	: public SCompoundWidget
+	, public FSelfRegisteringEditorUndoClient
+{
+public:
+	SLATE_BEGIN_ARGS(SDMXFixturePatchList)
+	{}
+
+	SLATE_END_ARGS()
+
+	SDMXFixturePatchList();
+	virtual ~SDMXFixturePatchList();
+
+	//~Begin EditorUndoClient interface
+	virtual void PostUndo(bool bSuccess) override;
+	virtual void PostRedo(bool bSuccess) override;
+	//~End EditorUndoClient interface
+
+	/** Constructs this widget */
+	void Construct(const FArguments& InArgs, TWeakPtr<FDMXEditor> InDMXEditor);
+
+	/** Updates the List on the next tick */
+	void RequestListRefresh();
+
+	/** If a single Row is selected, enters Editing Mode for the Fixture Patch Name of the selected Row */
+	void EnterFixturePatchNameEditingMode();
+
+	/** Processes command bindings for the Key Event */
+	FReply ProcessCommandBindings(const FKeyEvent& InKeyEvent);
+
+protected:
+	//~ Begin SWidget Interface
+	virtual bool SupportsKeyboardFocus() const override { return true; }
+	virtual FReply OnKeyDown(const FGeometry& MyGeometry, const FKeyEvent& InKeyEvent) override;
+	//~ End SWidget Interface
+
+private:
+	/** Called when the search changed */
+	void OnSearchChanged();
+
+	/** Refreshes the List */
+	void RefreshList();
+
+	/** Generates Status Text for each Item and sets it */
+	void GenereateStatusText();
+
+	/** Called when a row in the List gets generated */
+	TSharedRef<ITableRow> OnGenerateRow(TSharedPtr<FDMXFixturePatchListItem> InItem, const TSharedRef<STableViewBase>& OwnerTable);
+
+	/** Called when the Selection in the list changed */
+	void OnSelectionChanged(TSharedPtr<FDMXFixturePatchListItem> InItem, ESelectInfo::Type SelectInfo);
+
+	/** Called when entities were added or removed from the DMX Library */
+	void OnEntityAddedOrRemoved(UDMXLibrary* DMXLibrary, TArray<UDMXEntity*> Entities);
+
+	/** Called when a Fixture Patch changed */
+	void OnFixturePatchChanged(const UDMXEntityFixturePatch* FixturePatch);
+
+	/** Called when a Fixture Type changed */
+	void OnFixtureTypeChanged(const UDMXEntityFixtureType* FixtureType);
+
+	/** Called when Fixture Patch Shared Data selected Fixture Patches */
+	void OnFixturePatchSharedDataSelectedFixturePatches();
+
+	/** Selects what is selected in Fixture Patch Shared Data */
+	void AdoptSelectionFromFixturePatchSharedData();
+
+	/** Auto assigns selected Fixture Patches */
+	void AutoAssignFixturePatches(UE::DMXEditor::AutoAssign::EAutoAssignMode Mode);
+
+	/** Returns true if the DMX Library has reachable Universes */
+	bool DoesDMXLibraryHaveReachableUniverses() const;
+
+	/** Sets keyboard focus on this widget */
+	void SetKeyboardFocus();
+
+	/** Generates the Header Row of the List */
+	TSharedRef<SHeaderRow> GenerateHeaderRow();
+
+	/** Returns the column sort mode for the list*/
+	EColumnSortMode::Type GetColumnSortMode(const FName ColumnId) const;
+
+	/** Sorts the list source */
+	void SortListSource(const EColumnSortPriority::Type SortPriority, const FName& ColumnId, const EColumnSortMode::Type InSortMode);
+
+	/** Sort the list */
+	void SortList(const EColumnSortPriority::Type SortPriority, const FName& ColumnId, const EColumnSortMode::Type InSortMode);
+
+	/** The current Sort Mode */
+	EColumnSortMode::Type SortMode;
+
+	/** By which column ID the List is sorted */
+	FName SortedByColumnID;
+
+	/** Set to true while changing the DMX Library and only redraw the list when this is false, e.g. for Duplicate, Paste etc. */
+	bool bChangingDMXLibrary = false;
+
+	/** Source array for the Fixture Patch List */
+	TArray<TSharedPtr<FDMXFixturePatchListItem>> ListSource;
+
+	/** List source when filtered by search */
+	TArray<TSharedPtr<FDMXFixturePatchListItem>> FilteredListSource;
+
+	/** The item that was last selected in the list */
+	TSharedPtr<FDMXFixturePatchListItem> LastSelectedItem;
+
+	/** The Search Bar for the List */
+	TSharedPtr<SDMXFixturePatchListToolbar> Toolbar;
+
+	/** The the list of Fixture Patches */
+	TSharedPtr<SDMXFixturePatchListType> ListView;
+
+	/** The Header Row of the List */
+	TSharedPtr<SHeaderRow> HeaderRow;
+
+	/** Rows of Mode widgets in the List */
+	TArray<TSharedPtr<SDMXFixturePatchListRow>> Rows;
+
+	/** Timer handle for the Request List Refresh method */
+	FTimerHandle RequestListRefreshTimerHandle;
+
+	/** Shared Data for Fixture Patch Editors */
+	TSharedPtr<FDMXFixturePatchSharedData> FixturePatchSharedData;
+
+	/** The Fixture Type Editor that owns this widget */
+	TWeakPtr<FDMXEditor> WeakDMXEditor;
+
+private:
+	///////////////////////////////////////////////////
+	// Context menu Commands related
+
+	TSharedPtr<SWidget> OnContextMenuOpening();
+
+	/** Registers the commands */
+	void RegisterCommands();
+
+	bool CanCutItems() const;
+	void OnCutSelectedItems();
+	bool CanCopyItems() const;
+	void OnCopySelectedItems();
+	bool CanPasteItems() const;
+	void OnPasteItems();
+	bool CanDuplicateItems() const;
+	void OnDuplicateItems();
+	bool CanDeleteItems() const;
+	void OnDeleteItems();
+
+	TSharedPtr<FUICommandList> CommandList;
+};
diff --git a/Engine/Plugins/VirtualProduction/DMX/DMXEngine/Source/DMXEditor/Private/Widgets/FixturePatch/SDMXFixturePatchListRow.cpp b/Engine/Plugins/VirtualProduction/DMX/DMXEngine/Source/DMXEditor/Private/Widgets/FixturePatch/SDMXFixturePatchListRow.cpp
new file mode 100644
--- /dev/null
+++ b/Engine/Plugins/VirtualProduction/DMX/DMXEngine/Source/DMXEditor/Private/Widgets/FixturePatch/SDMXFixturePatchListRow.cpp
@@ -0,0 +1,768 @@
+// Copyright Epic Games, Inc. All Rights Reserved.
+
+#include "SDMXFixturePatchListRow.h"
+
+#include "DMXEditor.h"
+#include "DMXEditorStyle.h"
+#include "Engine/EngineTypes.h"
+#include "Library/DMXEntityFixturePatch.h"
+#include "Library/DMXEntityFixtureType.h"
+#include "Library/DMXLibrary.h"
+#include "MVR/DMXMVRGeneralSceneDescription.h"
+#include "SSearchableComboBox.h"
+#include "Styling/AppStyle.h"
+#include "Widgets/FixturePatch/DMXFixturePatchListItem.h"
+#include "Widgets/FixturePatch/SDMXFixturePatchList.h"
+#include "Widgets/Input/STextComboBox.h"
+#include "Widgets/Text/SInlineEditableTextBlock.h"
+
+#define LOCTEXT_NAMESPACE "SDMXFixturePatchListRow"
+
+/////////////////////////////////////////////////////
+// SDMXFixturePatchFixtureTypePicker
+
+/** Widget to pick a fixture type for a Fixture Patch */
+class SDMXFixturePatchFixtureTypePicker
+	: public SCompoundWidget
+{
+	DECLARE_DELEGATE_OneParam(FDMXFixturePatchListRowOnFixtureTypeSelectedDelegate, UDMXEntityFixtureType* /** Selected Fixture Type */);
+
+public:
+	SLATE_BEGIN_ARGS(SDMXFixturePatchFixtureTypePicker)
+	{}
+		/** Called when the combo box selection changed */
+		SLATE_EVENT(FDMXFixturePatchListRowOnFixtureTypeSelectedDelegate, OnFixtureTypeSelected)
+
+	SLATE_END_ARGS()
+
+	/** Constructs this widget */
+	void Construct(const FArguments& InArgs, UDMXLibrary* InDMXLibrary)
+	{
+		if (!InDMXLibrary)
+		{
+			return;
+		}
+
+		WeakDMXLibrary = InDMXLibrary;
+		OnFixtureTypeSelectedDelegate = InArgs._OnFixtureTypeSelected;
+
+
+		InDMXLibrary->GetOnEntitiesAdded().AddSP(this, &SDMXFixturePatchFixtureTypePicker::OnEntitiesAddedOrRemoved);
+		InDMXLibrary->GetOnEntitiesRemoved().AddSP(this, &SDMXFixturePatchFixtureTypePicker::OnEntitiesAddedOrRemoved);
+		UDMXEntityFixtureType::GetOnFixtureTypeChanged().AddSP(this, &SDMXFixturePatchFixtureTypePicker::OnFixtureTypeChanged);
+
+		ChildSlot
+			[
+				SAssignNew(ComboBox, SSearchableComboBox)
+				.OptionsSource(&FixtureTypeNames)
+				.OnGenerateWidget(this, &SDMXFixturePatchFixtureTypePicker::OnGenerateWidget)
+				.OnSelectionChanged(this, &SDMXFixturePatchFixtureTypePicker::OnSelectionChanged)
+				.Content()						
+				[
+					SNew(STextBlock)
+					.Text(this, &SDMXFixturePatchFixtureTypePicker::GetSelectedItemText)
+					.Font(FAppStyle::GetFontStyle(TEXT("PropertyWindow.NormalFont")))
+				]
+			];
+
+		RefreshInternal();
+	}
+
+	/** Sets the currently selected Fixture Type */
+	void SetSelection(UDMXEntityFixtureType* FixtureType)
+	{
+		if (FixtureType)
+		{
+			TSharedPtr<FString> const* FixtureTypeNameToSelectPtr = FixtureTypeNames.FindByPredicate([FixtureType](const TSharedPtr<FString>& FixtureTypeName)
+				{
+					return *FixtureTypeName == FixtureType->Name;
+				});
+			if (ensureAlwaysMsgf(FixtureTypeNameToSelectPtr, TEXT("Trying to select a fixture type but Fixture Type %s is not present."), *FixtureType->Name))
+			{
+				ComboBox->SetSelectedItem(*FixtureTypeNameToSelectPtr);
+			}
+		}
+		else
+		{
+			checkf(FixtureTypeNames.Num() > 0, TEXT("Expected at least a 'None' entry in FitureTypesName, but the array is empty."));
+			ComboBox->SetSelectedItem(FixtureTypeNames[0]);
+		}
+	}
+
+	/** Requests to refresh the widget on the next tick */
+	void RequestRefresh()
+	{
+		if (!RefreshTimerHandle.IsValid())
+		{
+			RefreshTimerHandle = GEditor->GetTimerManager()->SetTimerForNextTick(FTimerDelegate::CreateSP(this, &SDMXFixturePatchFixtureTypePicker::RefreshInternal));
+		}
+	}
+
+private:
+	/** Called when the combo box generates a widget */
+	TSharedRef<SWidget> OnGenerateWidget(TSharedPtr<FString> InItem)
+	{
+		return
+			SNew(STextBlock)
+			.Text(FText::FromString(*InItem));
+	}
+
+	/** Called when the combo box selection changed */
+	void OnSelectionChanged(TSharedPtr<FString> NewSelection, ESelectInfo::Type SelectInfo)
+	{
+		if (SelectInfo == ESelectInfo::Direct)
+		{
+			return;
+		}
+
+		if (UDMXEntityFixtureType* SelectedFixtureType = GetSelectedFixtureType())
+		{
+			OnFixtureTypeSelectedDelegate.ExecuteIfBound(SelectedFixtureType);
+		}
+	}
+
+	UDMXEntityFixtureType* GetSelectedFixtureType() const
+	{
+		UDMXEntityFixtureType* SelectedFixtureType = nullptr;
+		const TSharedPtr<FString> SelectedItem = ComboBox->GetSelectedItem();
+		if (TWeakObjectPtr<UDMXEntityFixtureType> const* FixtureTypePtr = FixtureTypeNameToFixtureTypeMap.Find(SelectedItem))
+		{
+			SelectedFixtureType = FixtureTypePtr->Get();
+		}
+
+		return SelectedFixtureType;
+	}
+
+	/** Returns the text of the selected item */
+	FText GetSelectedItemText() const
+	{
+		if (ComboBox.IsValid())
+		{
+			if (const TSharedPtr<FString> SelectedItem = ComboBox->GetSelectedItem())
+			{
+				return FText::FromString(*SelectedItem);
+			}
+		}
+
+		return FText::GetEmpty();
+	}
+
+	/** Refreshes the Combo Box */
+	void RefreshInternal()
+	{
+		RefreshTimerHandle.Invalidate();
+
+		FixtureTypeNames.Reset();
+		const UDMXLibrary* DMXLibrary = WeakDMXLibrary.Get();
+		if (!DMXLibrary)
+		{
+			return;
+		}
+
+		const TArray<UDMXEntityFixtureType*> FixtureTypes = DMXLibrary->GetEntitiesTypeCast<UDMXEntityFixtureType>();
+
+		if (FixtureTypes.IsEmpty())
+		{
+			const TSharedRef<FString> NoneItem = MakeShared<FString>(TEXT("None"));
+			FixtureTypeNames.Add(NoneItem);
+			ComboBox->SetSelectedItem(NoneItem);
+		}
+		else
+		{
+			for (UDMXEntityFixtureType* FixtureType : FixtureTypes)
+			{
+				const TSharedPtr<FString> FixtureTypeName = MakeShared<FString>(FixtureType->Name);
+
+				FixtureTypeNames.Add(FixtureTypeName);
+				FixtureTypeNameToFixtureTypeMap.Add(FixtureTypeName, FixtureType);
+			}
+		}
+
+		if (!ensureMsgf(!FixtureTypeNames.IsEmpty(), TEXT("No combo box option for Fixture Type.")))
+		{
+			return;
+		}
+
+		const FString SelectedFixtureTypeName = ComboBox->GetSelectedItem() ? *ComboBox->GetSelectedItem() : TEXT("");
+		const TSharedPtr<FString>* SelectionPtr = Algo::FindByPredicate(FixtureTypeNames, [SelectedFixtureTypeName](const TSharedPtr<FString>& FixtureTypeName)
+			{
+				return *FixtureTypeName == SelectedFixtureTypeName;
+			});
+		if (!SelectionPtr)
+		{
+			ComboBox->SetSelectedItem(FixtureTypeNames[0], ESelectInfo::Direct);
+		}
+
+		ComboBox->RefreshOptions();
+	}
+
+	/** Called when the Entities in the DMX Library were added or removed */
+	void OnEntitiesAddedOrRemoved(UDMXLibrary* DMXLibrary, TArray<UDMXEntity*> AddedOrRemovedEntities)
+	{
+		if (DMXLibrary == WeakDMXLibrary)
+		{
+			RequestRefresh();
+		}
+	}
+
+	/** Called when a Fixture Type changed */
+	void OnFixtureTypeChanged(const UDMXEntityFixtureType* ChangedFixtureType)
+	{
+		if (ChangedFixtureType && ChangedFixtureType->GetParentLibrary() == WeakDMXLibrary)
+		{
+			RequestRefresh();
+		}
+	}
+
+	/** Names of Fixture Types in the Combo Box */
+	TArray<TSharedPtr<FString>> FixtureTypeNames;
+
+	/** Map from the Fixture Type Name to the actual Fixture Type object */
+	TMap<TSharedPtr<FString>, TWeakObjectPtr<UDMXEntityFixtureType>> FixtureTypeNameToFixtureTypeMap;
+
+	/** The combo box widget in use */
+	TSharedPtr<SSearchableComboBox> ComboBox;
+
+	/** The DMX Library from which to select fixture types */
+	TWeakObjectPtr<UDMXLibrary> WeakDMXLibrary;
+
+	/** Timer handle used for the RequestRefresh method */
+	FTimerHandle RefreshTimerHandle;
+
+	// Slate args
+	FDMXFixturePatchListRowOnFixtureTypeSelectedDelegate OnFixtureTypeSelectedDelegate;
+};
+
+
+/////////////////////////////////////////////////////
+// SDMXFixturePatchModePicker
+
+/** Widget to pick a fixture type for an MVR Fixture */
+class SDMXFixturePatchModePicker
+	: public SCompoundWidget
+{
+	DECLARE_DELEGATE_OneParam(FDMXFixturePatchListRowOnModeSelectedDelegate, int32 /** Selected Mode Index */);
+
+public:
+	SLATE_BEGIN_ARGS(SDMXFixturePatchModePicker)
+	{}
+		/** Called when the combo box selection changed */
+		SLATE_EVENT(FDMXFixturePatchListRowOnModeSelectedDelegate, OnModeSelected)
+
+	SLATE_END_ARGS()
+
+	/** Constructs this widget */
+	void Construct(const FArguments& InArgs)
+	{
+		OnModeSelectedDelegate = InArgs._OnModeSelected;
+
+		UDMXEntityFixtureType::GetOnFixtureTypeChanged().AddSP(this, &SDMXFixturePatchModePicker::OnFixtureTypeChanged);
+
+		ChildSlot
+			[
+				SNew(SVerticalBox)
+				.IsEnabled_Lambda([this]
+					{
+						return bEnabled;
+					})
+				
+				+ SVerticalBox::Slot()
+				[
+					SAssignNew(ComboBox, STextComboBox)
+					.Visibility_Lambda([this]()
+						{
+							return WeakFixtureType.IsValid() ? EVisibility::Visible : EVisibility::Hidden;
+						})
+					.OptionsSource(&ModeNames)
+					.OnSelectionChanged(this, &SDMXFixturePatchModePicker::OnSelectionChanged)
+					.Font(FAppStyle::GetFontStyle(TEXT("PropertyWindow.NormalFont")))
+				]
+
+				+ SVerticalBox::Slot()
+				[
+					SNew(STextBlock)
+					.Visibility_Lambda([this]()
+						{
+							return WeakFixtureType.IsValid() ? EVisibility::Collapsed : EVisibility::Visible;
+						})
+					.Text(LOCTEXT("NoModeBecauseNoFixtureTypeSelectedInfo", "No Fixture Type Selected"))
+					.Font(FAppStyle::GetFontStyle(TEXT("PropertyWindow.NormalFont")))
+				]
+			];
+
+		RefreshOptions();
+	}
+
+	/** Sets the selection */
+	void SetSelection(UDMXEntityFixtureType* FixtureType, int32 ModeIndex)
+	{
+		WeakFixtureType = FixtureType;		
+		RefreshOptions();
+
+		if (FixtureType && FixtureType->Modes.IsValidIndex(ModeIndex))
+		{
+			if (ensureAlwaysMsgf(ModeNames.IsValidIndex(ModeIndex), TEXT("Trying to select a Mode but its index is not valid.")))
+			{
+				ComboBox->SetSelectedItem(ModeNames[ModeIndex]);
+			}
+		}
+	}
+
+private:
+	/** Called when the combo box selection changed */
+	void OnSelectionChanged(TSharedPtr<FString> NewSelection, ESelectInfo::Type SelectInfo)
+	{
+		UDMXEntityFixtureType* FixtureType = WeakFixtureType.Get();
+		if (ensureAlwaysMsgf(FixtureType, TEXT("Tried to select Mode, but Fixture Type is no longer valid.")))
+		{
+			const TSharedPtr<FString> SelectedItem = ComboBox->GetSelectedItem();
+			int32 const* ModeIndexPtr = ModeNameToModeIndexMap.Find(SelectedItem);
+			if (ensureAlwaysMsgf(ModeIndexPtr, TEXT("Tried to select Mode, but the Mode Index is no longer valid.")))
+			{
+				OnModeSelectedDelegate.ExecuteIfBound(*ModeIndexPtr);
+			}
+		}
+	}
+
+	/** Returns the text of the selected item */
+	FText GetSelectedItemText() const
+	{
+		const TSharedPtr<FString> SelectedItem = ComboBox->GetSelectedItem();
+		return FText::FromString(*SelectedItem);
+	}
+
+	/** Refreshes the Combo Box */
+	void RefreshOptions()
+	{
+		// Remember selection
+		const FString PreviousSelectedModeName = ComboBox->GetSelectedItem().IsValid() ? *ComboBox->GetSelectedItem() : FString();
+
+		ModeNames.Reset();
+
+		if (UDMXEntityFixtureType* FixtureType = WeakFixtureType.Get())
+		{
+			for (int32 ModeIndex = 0; ModeIndex < FixtureType->Modes.Num(); ModeIndex++)
+			{
+				const TSharedPtr<FString> ModeName = MakeShared<FString>(FixtureType->Modes[ModeIndex].ModeName);
+
+				ModeNames.Add(ModeName);
+				ModeNameToModeIndexMap.Add(ModeName, ModeIndex);
+			}
+		}
+
+		if (ModeNames.Num() == 0)
+		{
+			const FText NoModeAvailableText = LOCTEXT("NoModeAvailableInModeComboBoxText", "No Mode available");
+			ModeNames.Add(MakeShared<FString>(NoModeAvailableText.ToString()));
+
+			bEnabled = false;
+		}
+		else
+		{
+			bEnabled = true;
+		}
+
+		ComboBox->RefreshOptions();
+
+		// Restore selection
+		TSharedPtr<FString>* PreviousSelectionPtr = Algo::FindByPredicate(ModeNames, [PreviousSelectedModeName](const TSharedPtr<FString>& ModeName)
+			{
+				return ModeName.IsValid() && *ModeName == PreviousSelectedModeName;
+			});
+		if (PreviousSelectionPtr)
+		{
+			ComboBox->SetSelectedItem(*PreviousSelectionPtr);
+		}
+	}
+
+	/** Called when the Entities in the DMX Library were added or removed */
+	void OnFixtureTypeChanged(const UDMXEntityFixtureType* ChangedFixtureType)
+	{
+		if (ChangedFixtureType == WeakFixtureType)
+		{
+			RefreshOptions();
+		}
+	}
+
+	/** True if the Combo Box should be enabled */
+	bool bEnabled = true;
+
+	/** Names of Modes in the Combo Box */
+	TArray<TSharedPtr<FString>> ModeNames;
+
+	/** Map from the Fixture Type Name to the actual Mode Index */
+	TMap<TSharedPtr<FString>, int32> ModeNameToModeIndexMap;
+
+	/** The combo box widget in use */
+	TSharedPtr<STextComboBox> ComboBox;
+
+	/** The Fixture Type for which the Modes are currently displayed */
+	TWeakObjectPtr<UDMXEntityFixtureType> WeakFixtureType;
+
+	// Slate args
+	FDMXFixturePatchListRowOnModeSelectedDelegate OnModeSelectedDelegate;
+};
+
+
+/////////////////////////////////////////////////////
+// SDMXFixturePatchListRow
+
+void SDMXFixturePatchListRow::Construct(const FArguments& InArgs, const TSharedRef<STableViewBase>& InOwnerTable, const TSharedRef<FDMXFixturePatchListItem>& InItem)
+{
+	Item = InItem;
+	OnRowRequestsStatusRefresh = InArgs._OnRowRequestsStatusRefresh;
+	OnRowRequestsListRefresh = InArgs._OnRowRequestsListRefresh;
+	IsSelected = InArgs._IsSelected;
+
+	SMultiColumnTableRow<TSharedPtr<FDMXFixturePatchListItem>>::Construct(
+		FSuperRowType::FArguments()
+		.Style(&FDMXEditorStyle::Get().GetWidgetStyle<FTableRowStyle>("FixturePatchList.Row")),
+		InOwnerTable);
+}
+
+void SDMXFixturePatchListRow::EnterFixturePatchNameEditingMode()
+{
+	if (FixturePatchNameTextBlock.IsValid())
+	{
+		FixturePatchNameTextBlock->EnterEditingMode();
+	}
+}
+
+TSharedRef<SWidget> SDMXFixturePatchListRow::GenerateWidgetForColumn(const FName& ColumnName)
+{
+	if (ColumnName == FDMXFixturePatchListCollumnID::EditorColor)
+	{
+		return GenerateEditorColorWidget();
+	}
+	if (ColumnName == FDMXFixturePatchListCollumnID::FixturePatchName)
+	{
+		return GenerateFixturePatchNameWidget();
+	}
+	else if (ColumnName == FDMXFixturePatchListCollumnID::Status)
+	{
+		return GenerateStatusWidget();
+	}
+	else if (ColumnName == FDMXFixturePatchListCollumnID::FixtureID)
+	{
+		return GenerateFixtureIDWidget();
+	}
+	else if (ColumnName == FDMXFixturePatchListCollumnID::FixtureType)
+	{
+		return GenerateFixtureTypeWidget();
+	}
+	else if (ColumnName == FDMXFixturePatchListCollumnID::Mode)
+	{
+		return GenerateModeWidget();
+	}
+	else if (ColumnName == FDMXFixturePatchListCollumnID::Patch)
+	{
+		return GeneratePatchWidget();
+	}
+
+	return SNullWidget::NullWidget;
+}
+
+TSharedRef<SWidget> SDMXFixturePatchListRow::GenerateEditorColorWidget()
+{
+	return
+		SNew(SBorder)
+		.HAlign(HAlign_Fill)
+		.Padding(5.f, 2.f)
+		.BorderImage(FAppStyle::GetBrush("NoBorder"))
+		[
+			SNew(SImage)
+			.Image(FDMXEditorStyle::Get().GetBrush("DMXEditor.WhiteRoundedPropertyBorder"))
+			.ColorAndOpacity_Lambda([this]()
+			{
+				return Item->GetBackgroundColor();
+			})
+		];
+}
+
+TSharedRef<SWidget> SDMXFixturePatchListRow::GenerateFixturePatchNameWidget()
+{
+	return
+		SAssignNew(FixturePatchNameBorder, SBorder)
+		.HAlign(HAlign_Fill)
+		.VAlign(VAlign_Center)
+		.Padding(4.f)
+		.BorderImage(FAppStyle::GetBrush("NoBorder"))
+		.OnMouseDoubleClick(this, &SDMXFixturePatchListRow::OnFixturePatchNameBorderDoubleClicked)
+		[
+			SNew(SBorder)
+			.HAlign(HAlign_Fill)
+			.VAlign(VAlign_Center)
+			.BorderImage(FDMXEditorStyle::Get().GetBrush("DMXEditor.RoundedPropertyBorder"))
+			[
+				SAssignNew(FixturePatchNameTextBlock, SInlineEditableTextBlock)
+				.Text_Lambda([this]()
+				{
+					const FString FixturePatchName = Item->GetFixturePatchName();
+					return FText::FromString(FixturePatchName);
+				})
+				.Font(FAppStyle::GetFontStyle(TEXT("PropertyWindow.NormalFont")))
+				.OnTextCommitted(this, &SDMXFixturePatchListRow::OnFixturePatchNameCommitted)
+				.IsSelected(IsSelected)
+			]
+		];
+}
+
+FReply SDMXFixturePatchListRow::OnFixturePatchNameBorderDoubleClicked(const FGeometry& InMyGeometry, const FPointerEvent& InMouseEvent)
+{
+	if (InMouseEvent.GetEffectingButton() == EKeys::LeftMouseButton)
+	{
+		if (FixturePatchNameTextBlock.IsValid())
+		{
+			FixturePatchNameTextBlock->EnterEditingMode();
+		}
+	}
+
+	return FReply::Handled();
+}
+
+void SDMXFixturePatchListRow::OnFixturePatchNameCommitted(const FText& InNewText, ETextCommit::Type InTextCommit)
+{
+	if (InNewText.IsEmpty())
+	{
+		return;
+	}
+
+	FString ResultingName;
+	Item->SetFixturePatchName(InNewText.ToString(), ResultingName);
+	FixturePatchNameTextBlock->SetText(FText::FromString(ResultingName));
+}
+
+TSharedRef<SWidget> SDMXFixturePatchListRow::GenerateStatusWidget()
+{
+	return
+		SNew(SBox)
+		.HAlign(HAlign_Center)
+		.VAlign(VAlign_Center)
+		[
+			SNew(SImage)
+			.Image_Lambda([this]()
+				{
+					if (!Item->ErrorStatusText.IsEmpty() || !Item->WarningStatusText.IsEmpty())
+					{
+						return FDMXEditorStyle::Get().GetBrush("Icons.WarningExclamationMark");
+					}
+
+					static const FSlateBrush EmptyBrush = FSlateNoResource();
+					return &EmptyBrush;
+				})
+			.ToolTipText_Lambda([this]()
+				{
+					if (!Item->ErrorStatusText.IsEmpty())
+					{
+						return Item->ErrorStatusText;
+					}
+					else if (!Item->WarningStatusText.IsEmpty())
+					{
+						return Item->WarningStatusText;
+					}
+
+					return FText::GetEmpty();
+				})
+		];
+}
+
+TSharedRef<SWidget> SDMXFixturePatchListRow::GenerateFixtureIDWidget()
+{
+	return
+		SNew(SBorder)
+		.HAlign(HAlign_Fill)
+		.VAlign(VAlign_Center)
+		.Padding(4.f)
+		.BorderImage(FAppStyle::GetBrush("NoBorder"))
+		.OnMouseDoubleClick(this, &SDMXFixturePatchListRow::OnFixtureIDBorderDoubleClicked)
+		[
+			SNew(SBorder)
+			.HAlign(HAlign_Fill)
+			.VAlign(VAlign_Center)
+			.BorderImage(FDMXEditorStyle::Get().GetBrush("DMXEditor.RoundedPropertyBorder"))
+			[
+				SAssignNew(FixtureIDTextBlocK, SInlineEditableTextBlock)
+				.Text_Lambda([this]()
+				{
+					return FText::FromString(Item->GetFixtureID());
+				})
+				.Font(FAppStyle::GetFontStyle(TEXT("PropertyWindow.NormalFont")))
+				.OnTextCommitted(this, &SDMXFixturePatchListRow::OnFixtureIDCommitted)
+				.IsSelected(IsSelected)
+			]
+		];
+}
+
+FReply SDMXFixturePatchListRow::OnFixtureIDBorderDoubleClicked(const FGeometry& InMyGeometry, const FPointerEvent& InMouseEvent)
+{
+	if (InMouseEvent.GetEffectingButton() == EKeys::LeftMouseButton)
+	{
+		if (FixtureIDTextBlocK.IsValid())
+		{
+			FixtureIDTextBlocK->EnterEditingMode();
+		}
+	}
+
+	return FReply::Handled();
+}
+
+void SDMXFixturePatchListRow::OnFixtureIDCommitted(const FText& InNewText, ETextCommit::Type InTextCommit)
+{
+	const FString StringValue = InNewText.ToString();
+	int32 NewFixtureID;
+	if (LexTryParseString<int32>(NewFixtureID, *StringValue))
+	{
+		Item->SetFixtureID(NewFixtureID);
+
+		const FString ParsedFixtureIDString = FString::FromInt(NewFixtureID);
+		FixtureIDTextBlocK->SetText(FText::FromString(ParsedFixtureIDString));
+
+		OnRowRequestsStatusRefresh.ExecuteIfBound();
+	}
+}
+
+TSharedRef<SWidget> SDMXFixturePatchListRow::GenerateFixtureTypeWidget()
+{
+	UDMXLibrary* DMXLibrary = Item->GetDMXLibrary();
+	if (!ensureAlwaysMsgf(DMXLibrary, TEXT("Tried to set fixture type for MVR Fixture, but fixture type is invalid.")))
+	{
+		return SNullWidget::NullWidget;
+	}
+
+	const TSharedRef<SDMXFixturePatchFixtureTypePicker> FixtureTypePicker =
+		SNew(SDMXFixturePatchFixtureTypePicker, DMXLibrary)
+		.OnFixtureTypeSelected(this, &SDMXFixturePatchListRow::OnFixtureTypeSelected);
+
+	UDMXEntityFixtureType* SelectedFixtureType = Item->GetFixtureType();
+	FixtureTypePicker->SetSelection(SelectedFixtureType);
+		
+	return 
+		SNew(SBorder)
+		.HAlign(HAlign_Fill)
+		.VAlign(VAlign_Center)
+		.Padding(4.f)
+		.BorderImage(FAppStyle::GetBrush("NoBorder"))
+		[
+			FixtureTypePicker
+		];
+}
+
+void SDMXFixturePatchListRow::OnFixtureTypeSelected(UDMXEntityFixtureType* SelectedFixtureType)
+{
+	Item->SetFixtureType(SelectedFixtureType);
+}
+
+TSharedRef<SWidget> SDMXFixturePatchListRow::GenerateModeWidget()
+{
+	const TSharedRef<SDMXFixturePatchModePicker> ModePicker =
+		SNew(SDMXFixturePatchModePicker)
+		.OnModeSelected(this, &SDMXFixturePatchListRow::OnModeSelected);
+
+	UDMXEntityFixtureType* SelectedFixtureType = Item->GetFixtureType();
+	const int32 SelectedModeIndex = Item->GetModeIndex();
+
+	ModePicker->SetSelection(SelectedFixtureType, SelectedModeIndex);
+
+	return
+		SNew(SBorder)
+		.HAlign(HAlign_Fill)
+		.VAlign(VAlign_Center)
+		.Padding(4.f)
+		.BorderImage(FAppStyle::GetBrush("NoBorder"))
+		[
+			ModePicker
+		];
+}
+
+void SDMXFixturePatchListRow::OnModeSelected(int32 SelectedModeIndex)
+{
+	Item->SetModeIndex(SelectedModeIndex);
+}
+
+TSharedRef<SWidget> SDMXFixturePatchListRow::GeneratePatchWidget()
+{
+	return
+		SNew(SBorder)
+		.HAlign(HAlign_Fill)
+		.VAlign(VAlign_Center)
+		.Padding(4.f)
+		.BorderImage(FAppStyle::GetBrush("NoBorder"))
+		.OnMouseDoubleClick(this, &SDMXFixturePatchListRow::OnPatchBorderDoubleClicked)
+		[
+			SNew(SBorder)
+			.HAlign(HAlign_Fill)
+			.VAlign(VAlign_Center)
+			.BorderImage(FDMXEditorStyle::Get().GetBrush("DMXEditor.RoundedPropertyBorder"))
+			[
+				SAssignNew(PatchTextBlock, SInlineEditableTextBlock)
+				.Text_Lambda([this]()
+					{						
+						const int32 UniverseID = Item->GetUniverse();
+						const int32 StartingAddress = Item->GetAddress();
+						if (UniverseID > 0 && StartingAddress > 0)
+						{
+							return FText::Format(LOCTEXT("AddressesText", "{0}.{1}"), UniverseID, StartingAddress);
+						}
+						else
+						{
+							return LOCTEXT("NotPatchedText", "Not patched");
+						}
+					})
+				.Font(FAppStyle::GetFontStyle(TEXT("PropertyWindow.NormalFont")))
+				.OnTextCommitted(this, &SDMXFixturePatchListRow::OnPatchCommitted)
+				.IsSelected(IsSelected)
+			]
+		];
+}
+
+FReply SDMXFixturePatchListRow::OnPatchBorderDoubleClicked(const FGeometry& InMyGeometry, const FPointerEvent& InMouseEvent)
+{
+	if (InMouseEvent.GetEffectingButton() == EKeys::LeftMouseButton)
+	{
+		if (PatchTextBlock.IsValid())
+		{
+			PatchTextBlock->EnterEditingMode();
+		}
+	}
+
+	return FReply::Handled();
+}
+
+void SDMXFixturePatchListRow::OnPatchCommitted(const FText& InNewText, ETextCommit::Type InTextCommit)
+{
+	const FString PatchString = InNewText.ToString();
+	static const TCHAR* ParamDelimiters[] =
+	{
+		TEXT("."),
+		TEXT(","),
+		TEXT(":"),
+		TEXT(";")
+	};
+
+	TArray<FString> ValueStringArray;
+	PatchString.ParseIntoArray(ValueStringArray, ParamDelimiters, 4);
+	if (ValueStringArray.Num() == 2)
+	{
+		int32 Universe;
+		if (!LexTryParseString(Universe, *ValueStringArray[0]))
+		{
+			return;
+		}
+		
+		int32 Address;
+		if (!LexTryParseString(Address, *ValueStringArray[1]))
+		{
+			return;
+		}
+
+		Item->SetAddresses(Universe, Address);
+
+		const FString UniverseString = FString::FromInt(Universe);
+		const FString AddressString = FString::FromInt(Address);
+		PatchTextBlock->SetText(FText::Format(LOCTEXT("UniverseDotAddressText", "{0}.{1}"), FText::FromString(UniverseString), FText::FromString(AddressString)));
+
+		OnRowRequestsListRefresh.ExecuteIfBound();
+	}
+}
+
+#undef LOCTEXT_NAMESPACE
diff --git a/Engine/Plugins/VirtualProduction/DMX/DMXEngine/Source/DMXEditor/Private/Widgets/FixturePatch/SDMXFixturePatchListRow.h b/Engine/Plugins/VirtualProduction/DMX/DMXEngine/Source/DMXEditor/Private/Widgets/FixturePatch/SDMXFixturePatchListRow.h
new file mode 100644
--- /dev/null
+++ b/Engine/Plugins/VirtualProduction/DMX/DMXEngine/Source/DMXEditor/Private/Widgets/FixturePatch/SDMXFixturePatchListRow.h
@@ -0,0 +1,117 @@
+// Copyright Epic Games, Inc. All Rights Reserved.
+
+#pragma once
+
+#include "Widgets/Views/SListView.h"
+#include "Widgets/Views/STableRow.h"
+
+enum class EDMXFixturePatchListEditMode : uint8;
+class FDMXFixturePatchListItem;
+class SInlineEditableTextBlock;
+class UDMXEntityFixturePatch;
+class UDMXEntityFixtureType;
+
+/** MVR Fixture view as a row in a list */
+class SDMXFixturePatchListRow
+	: public SMultiColumnTableRow<TSharedPtr<FDMXFixturePatchListItem>>
+{
+public:
+	SLATE_BEGIN_ARGS(SDMXFixturePatchListRow)
+	{}
+		/** Delegate executed when the row requests to refresh the statuses */
+		SLATE_EVENT(FSimpleDelegate, OnRowRequestsStatusRefresh)
+
+		/** Delegate executed when the row requests to refresh the whole list */
+		SLATE_EVENT(FSimpleDelegate, OnRowRequestsListRefresh)
+
+		/** Callback to check if the row is selected (should be hooked up if a parent widget is handling selection or focus) */
+		SLATE_EVENT(FIsSelected, IsSelected)
+
+	SLATE_END_ARGS()
+			
+	/** Constructs this widget */
+	void Construct(const FArguments& InArgs, const TSharedRef<STableViewBase>& InOwnerTable, const TSharedRef<FDMXFixturePatchListItem>& InItem);
+
+	/** Enters editing mode for the Fixture Patch Name */
+	void EnterFixturePatchNameEditingMode();
+
+	/** Returns the Item of this row */
+	TSharedPtr<FDMXFixturePatchListItem> GetItem() const { return Item; };
+
+protected:
+	//~ Begin SMultiColumnTableRow interface
+	virtual TSharedRef<SWidget> GenerateWidgetForColumn(const FName& ColumnName) override;
+	//~ End SMultiColumnTableRow interface
+
+private:
+	/** Generates the Widget that displays the Editor Color */
+	TSharedRef<SWidget> GenerateEditorColorWidget();
+
+	/** Generates the Widget that displays the Fixture Patch Name */
+	TSharedRef<SWidget> GenerateFixturePatchNameWidget();
+
+	/** Called when the Fixture Patch Name Border was double-clicked */
+	FReply OnFixturePatchNameBorderDoubleClicked(const FGeometry& InMyGeometry, const FPointerEvent& InMouseEvent);
+
+	/** Called when a Fixture Patch Name was committed */
+	void OnFixturePatchNameCommitted(const FText& InNewText, ETextCommit::Type InTextCommit);
+
+	/** Generates the widget that displays the Status */
+	TSharedRef<SWidget> GenerateStatusWidget();
+
+	/** Generates the Widget that displays the Fixture ID */
+	TSharedRef<SWidget> GenerateFixtureIDWidget();
+
+	/** Called when the Fixture ID Border was double-clicked */
+	FReply OnFixtureIDBorderDoubleClicked(const FGeometry& InMyGeometry, const FPointerEvent& InMouseEvent);
+
+	/** Called when a Fixture ID was committed */
+	void OnFixtureIDCommitted(const FText& InNewText, ETextCommit::Type InTextCommit);
+
+	/** Generates the Widget that displays the Fixture Type */
+	TSharedRef<SWidget> GenerateFixtureTypeWidget();
+
+	/** Called when a Fixture Type was selected */
+	void OnFixtureTypeSelected(UDMXEntityFixtureType* SelectedFixtureType);
+
+	/** Generates the Widget that displays the Mode */
+	TSharedRef<SWidget> GenerateModeWidget();
+
+	/** Called when a Mode was selected */
+	void OnModeSelected(int32 SelectedModeIndex);
+
+	/** Generates the Widget that displays the Patch */
+	TSharedRef<SWidget> GeneratePatchWidget();
+
+	/** Called when the Patch Border was double-clicked */
+	FReply OnPatchBorderDoubleClicked(const FGeometry& InMyGeometry, const FPointerEvent& InMouseEvent);
+
+	/** Called when a Patch was committed */
+	void OnPatchCommitted(const FText& InNewText, ETextCommit::Type InTextCommit);
+
+	/** The Edit Mode the Widget should present */
+	EDMXFixturePatchListEditMode EditMode;
+
+	/** The outermost border around the the Fixture Patch Name Column */
+	TSharedPtr<SBorder> FixturePatchNameBorder;
+
+	/** The text block to edit the Fixture Patch Name */
+	TSharedPtr<SInlineEditableTextBlock> FixturePatchNameTextBlock;
+
+	/** The text block to edit the Fixture ID */
+	TSharedPtr<SInlineEditableTextBlock> FixtureIDTextBlocK;
+
+	/** The text block to edit the Name */
+	TSharedPtr<SInlineEditableTextBlock> NameTextBlock;
+
+	/** The text block to edit the Name */
+	TSharedPtr<SInlineEditableTextBlock> PatchTextBlock;
+
+	/** The MVR Fixture List Item this row displays */
+	TSharedPtr<FDMXFixturePatchListItem> Item;
+
+	// Slate arguments
+	FSimpleDelegate OnRowRequestsStatusRefresh;
+	FSimpleDelegate OnRowRequestsListRefresh;
+	FIsSelected IsSelected;
+};
diff --git a/Engine/Plugins/VirtualProduction/DMX/DMXEngine/Source/DMXEditor/Private/Widgets/FixturePatch/SDMXFixturePatchListToolbar.cpp b/Engine/Plugins/VirtualProduction/DMX/DMXEngine/Source/DMXEditor/Private/Widgets/FixturePatch/SDMXFixturePatchListToolbar.cpp
new file mode 100644
--- /dev/null
+++ b/Engine/Plugins/VirtualProduction/DMX/DMXEngine/Source/DMXEditor/Private/Widgets/FixturePatch/SDMXFixturePatchListToolbar.cpp
@@ -0,0 +1,236 @@
+// Copyright Epic Games, Inc. All Rights Reserved.
+
+#include "SDMXFixturePatchListToolbar.h"
+
+#include "Algo/MaxElement.h"
+#include "Commands/DMXEditorCommands.h"
+#include "DMXEditor.h"
+#include "DMXEditorUtils.h"
+#include "DMXFixturePatchListItem.h"
+#include "DMXFixturePatchSharedData.h"
+#include "EditorStyleSet.h"
+#include "FixturePatchAutoAssignUtility.h"
+#include "Library/DMXEntityFixturePatch.h"
+#include "Library/DMXEntityFixtureType.h"
+#include "SAddFixturePatchMenu.h"
+#include "ScopedTransaction.h"
+#include "Widgets/Input/SCheckBox.h"
+#include "Widgets/Input/SEditableTextBox.h"
+#include "Widgets/Input/SSearchBox.h"
+#include "Widgets/Layout/SSeparator.h"
+#include "Widgets/Layout/SWrapBox.h"
+#include "Widgets/SBoxPanel.h"
+#include "Widgets/SDMXEntityDropdownMenu.h"
+
+#define LOCTEXT_NAMESPACE "SDMXFixturePatchListToolbar"
+
+void SDMXFixturePatchListToolbar::Construct(const FArguments& InArgs, TWeakPtr<FDMXEditor> InDMXEditor)
+{
+	WeakDMXEditor = InDMXEditor;
+	OnSearchChanged = InArgs._OnSearchChanged;
+
+	ChildSlot
+		.HAlign(HAlign_Fill)
+		.VAlign(VAlign_Fill)
+		[
+			SNew(SBorder)
+			.Padding(FMargin(8.f, 8.f))
+			.BorderBackgroundColor(FLinearColor(0.6f, 0.6f, 0.6f))
+			.BorderImage(FAppStyle::GetBrush("DetailsView.CategoryTop"))
+			[
+				SNew(SWrapBox)
+				.InnerSlotPadding(FVector2D(14.f, 8.f))
+				.UseAllottedSize(true)
+
+				// Add Fixture Button
+				+ SWrapBox::Slot()
+				.HAlign(HAlign_Left)
+				.VAlign(VAlign_Center)
+				[
+					GenerateFixtureTypeDropdownMenu()
+				]
+										
+				+ SWrapBox::Slot()
+				.HAlign(HAlign_Left)
+				.VAlign(VAlign_Center)
+				[
+					SNew(SSeparator)
+					.Orientation(EOrientation::Orient_Vertical)
+				]
+
+				// Search
+				+ SWrapBox::Slot()
+				.HAlign(HAlign_Left)
+				.VAlign(VAlign_Center)
+				[
+					SNew(SSearchBox)
+					.MinDesiredWidth(400.f)
+					.OnTextChanged(this, &SDMXFixturePatchListToolbar::OnSearchTextChanged)
+					.ToolTipText(LOCTEXT("SearchBarTooltip", "Examples:\n\n* PatchName\n* FixtureTypeName\n* SomeMode\n* 1.\n* 1.1\n* Universe 1\n* Uni 1-3\n* Uni 1, 3\n* Uni 1, 4-5'."))
+				]
+									
+				+ SWrapBox::Slot()
+				.HAlign(HAlign_Left)
+				.VAlign(VAlign_Center)
+				[
+					SNew(SSeparator)
+					.Orientation(EOrientation::Orient_Vertical)
+				]
+
+				// Show Conflicts Only option
+				+ SWrapBox::Slot()
+				.HAlign(HAlign_Left)
+				.VAlign(VAlign_Center)
+				[
+					SNew(STextBlock)
+					.Text(LOCTEXT("ShowConflictsOnlyCheckBoxLabel", "Show Conflicts only"))
+					.Font(FAppStyle::GetFontStyle(TEXT("PropertyWindow.NormalFont")))
+				]
+
+				+ SWrapBox::Slot()
+				.HAlign(HAlign_Left)
+				.VAlign(VAlign_Center)
+				[
+					SNew(SCheckBox)
+					.IsChecked(false)
+					.OnCheckStateChanged(this, &SDMXFixturePatchListToolbar::OnShowConflictsOnlyCheckStateChanged)
+				]
+			]
+		];
+}
+
+TArray<TSharedPtr<FDMXFixturePatchListItem>> SDMXFixturePatchListToolbar::FilterItems(const TArray<TSharedPtr<FDMXFixturePatchListItem>>& Items)
+{
+	// Apply 'conflicts only' if enabled
+	TArray<TSharedPtr<FDMXFixturePatchListItem>> Result = Items;
+	if (bShowConfictsOnly)
+	{
+		Result.RemoveAll([](const TSharedPtr<FDMXFixturePatchListItem>& Item)
+			{
+				return
+					Item->ErrorStatusText.IsEmpty() &&
+					Item->WarningStatusText.IsEmpty();
+			});
+	}
+
+	// Filter and return in order of precendence
+	if (SearchString.IsEmpty())
+	{
+		return Result;
+	}
+
+	const TArray<int32> Universes = FDMXEditorUtils::ParseUniverses(SearchString);
+	if(!Universes.IsEmpty())
+	{
+		Result.RemoveAll([Universes](const TSharedPtr<FDMXFixturePatchListItem>& Item)
+			{
+				return !Universes.Contains(Item->GetUniverse());
+			});
+
+		return Result;
+	}
+
+	int32 Address;
+	if (FDMXEditorUtils::ParseAddress(SearchString, Address))
+	{
+		Result.RemoveAll([Address](const TSharedPtr<FDMXFixturePatchListItem>& Item)
+			{
+				return Item->GetAddress() != Address;
+			});
+
+		return Result;
+	}
+	
+	const TArray<int32> FixtureIDs = FDMXEditorUtils::ParseFixtureIDs(SearchString);
+	for (int32 FixtureID : FixtureIDs)
+	{
+		TArray<TSharedPtr<FDMXFixturePatchListItem>> FixtureIDsOnlyResult = Result;
+		FixtureIDsOnlyResult.RemoveAll([FixtureID](const TSharedPtr<FDMXFixturePatchListItem>& Item)
+			{
+				int32 OtherFixtureIDNumerical;
+				if (FDMXEditorUtils::ParseFixtureID(Item->GetFixtureID(), OtherFixtureIDNumerical))
+				{
+					return OtherFixtureIDNumerical != FixtureID;
+				}
+				return true;
+			});
+
+		if (FixtureIDsOnlyResult.Num() > 0)
+		{
+			return FixtureIDsOnlyResult;
+		}
+	}
+
+	Result.RemoveAll([this](const TSharedPtr<FDMXFixturePatchListItem>& Item)
+		{
+			return !Item->GetFixturePatchName().Contains(SearchString);
+		});
+
+	return Result;
+}
+
+TSharedRef<SWidget> SDMXFixturePatchListToolbar::GenerateFixtureTypeDropdownMenu()
+{
+	FText AddButtonLabel = FDMXEditorCommands::Get().AddNewEntityFixturePatch->GetLabel();
+	FText AddButtonToolTip = FDMXEditorCommands::Get().AddNewEntityFixturePatch->GetDescription();
+
+	TSharedRef<SComboButton> AddComboButton = 
+		SNew(SComboButton)
+			.ButtonContent()
+			[
+				SNew(SHorizontalBox)
+				+ SHorizontalBox::Slot()
+				.AutoWidth()
+				.Padding(FMargin(0, 1))
+				[
+					SNew(SImage)
+					.Image(FAppStyle::GetBrush("Plus"))
+				]
+				+ SHorizontalBox::Slot()
+				.VAlign(VAlign_Center)
+				.AutoWidth()
+				.Padding(FMargin(2, 0, 2, 0))
+				[
+					SNew(STextBlock)
+					.Text(AddButtonLabel)
+				]
+			]
+			.MenuContent()
+			[
+				SNew(SBox)
+				.Padding(4.f)
+				.MinDesiredWidth(460.f)
+				[
+					SAssignNew(AddFixturePatchMenu, UE::DMXEditor::FixturePatchEditor::SAddFixturePatchMenu, WeakDMXEditor)
+				]
+			]
+			.IsFocusable(true)
+			.ContentPadding(FMargin(4.0f, 4.0f))
+			.ComboButtonStyle(FAppStyle::Get(), "ToolbarComboButton")
+			.ButtonStyle(FAppStyle::Get(), "FlatButton.Success")
+			.ForegroundColor(FLinearColor::White)
+			.ToolTipText(AddButtonToolTip)
+			.OnComboBoxOpened_Lambda([this]()
+				{
+					if (AddFixturePatchMenu.IsValid())
+					{
+						AddFixturePatchMenu->RequestRefresh();
+					}
+				});
+
+	return AddComboButton;
+}
+
+void SDMXFixturePatchListToolbar::OnSearchTextChanged(const FText& SearchText)
+{
+	SearchString = SearchText.ToString();
+	OnSearchChanged.ExecuteIfBound();
+}
+
+void SDMXFixturePatchListToolbar::OnShowConflictsOnlyCheckStateChanged(const ECheckBoxState NewCheckState)
+{
+	bShowConfictsOnly = NewCheckState == ECheckBoxState::Checked;
+	OnSearchChanged.ExecuteIfBound();
+}
+
+#undef LOCTEXT_NAMESPACE
diff --git a/Engine/Plugins/VirtualProduction/DMX/DMXEngine/Source/DMXEditor/Private/Widgets/FixturePatch/SDMXFixturePatchListToolbar.h b/Engine/Plugins/VirtualProduction/DMX/DMXEngine/Source/DMXEditor/Private/Widgets/FixturePatch/SDMXFixturePatchListToolbar.h
new file mode 100644
--- /dev/null
+++ b/Engine/Plugins/VirtualProduction/DMX/DMXEngine/Source/DMXEditor/Private/Widgets/FixturePatch/SDMXFixturePatchListToolbar.h
@@ -0,0 +1,60 @@
+// Copyright Epic Games, Inc. All Rights Reserved.
+
+#pragma once
+
+#include "CoreMinimal.h"
+#include "Widgets/SCompoundWidget.h"
+
+enum class ECheckBoxState : uint8;
+class FDMXEditor;
+class FDMXFixturePatchListItem;
+template<typename TEntityType> class SDMXEntityDropdownMenu;
+class UDMXEntity;
+class UDMXEntityFixtureType;
+namespace UE::DMXEditor::FixturePatchEditor { class SAddFixturePatchMenu; }
+
+
+/** Search bar for the MVR Fixture List */
+class SDMXFixturePatchListToolbar
+	: public SCompoundWidget
+{
+public:
+	SLATE_BEGIN_ARGS(SDMXFixturePatchListToolbar)
+	{}
+
+		/** Executed when the search changed */
+		SLATE_EVENT(FSimpleDelegate, OnSearchChanged)
+
+	SLATE_END_ARGS()
+
+	/** Constructs this widget */
+	void Construct(const FArguments& InArgs, TWeakPtr<FDMXEditor> InDMXEditor);
+
+	/** Filters items according to the state of the search bar */
+	[[nodiscard]] TArray<TSharedPtr<FDMXFixturePatchListItem>> FilterItems(const TArray<TSharedPtr<FDMXFixturePatchListItem>>& Items);
+
+private:
+	/** Generates the 'Add MVR Fixture' Dropdown Menu */
+	TSharedRef<SWidget> GenerateFixtureTypeDropdownMenu();
+
+	/** Called when the Search Text changed */
+	void OnSearchTextChanged(const FText& SearchText);
+
+	/** Called when the Check State of the 'Show Conflicts Only' Check Box changed */
+	void OnShowConflictsOnlyCheckStateChanged(const ECheckBoxState NewCheckState);
+
+	/** Current Search String */
+	FString SearchString;
+
+	/** True if only conflicts are searched for */
+	bool bShowConfictsOnly = false;
+
+	/** Menu displayed to add fixture patches */
+	TSharedPtr<UE::DMXEditor::FixturePatchEditor::SAddFixturePatchMenu> AddFixturePatchMenu;
+
+	/** The DMX Editor that owns this widget */
+	TWeakPtr<FDMXEditor> WeakDMXEditor;
+
+	// Slate Args
+	FSimpleDelegate OnSearchChanged;
+};
diff --git a/Engine/Plugins/VirtualProduction/DMX/DMXEngine/Source/DMXEditor/Private/Widgets/FixturePatch/SDMXMVRFixtureList.cpp b/Engine/Plugins/VirtualProduction/DMX/DMXEngine/Source/DMXEditor/Private/Widgets/FixturePatch/SDMXMVRFixtureList.cpp
deleted file mode 100644
diff --git a/Engine/Plugins/VirtualProduction/DMX/DMXEngine/Source/DMXEditor/Private/Widgets/FixturePatch/SDMXMVRFixtureList.h b/Engine/Plugins/VirtualProduction/DMX/DMXEngine/Source/DMXEditor/Private/Widgets/FixturePatch/SDMXMVRFixtureList.h
deleted file mode 100644
diff --git a/Engine/Plugins/VirtualProduction/DMX/DMXEngine/Source/DMXEditor/Private/Widgets/FixturePatch/SDMXMVRFixtureListRow.cpp b/Engine/Plugins/VirtualProduction/DMX/DMXEngine/Source/DMXEditor/Private/Widgets/FixturePatch/SDMXMVRFixtureListRow.cpp
deleted file mode 100644
diff --git a/Engine/Plugins/VirtualProduction/DMX/DMXEngine/Source/DMXEditor/Private/Widgets/FixturePatch/SDMXMVRFixtureListRow.h b/Engine/Plugins/VirtualProduction/DMX/DMXEngine/Source/DMXEditor/Private/Widgets/FixturePatch/SDMXMVRFixtureListRow.h
deleted file mode 100644
diff --git a/Engine/Plugins/VirtualProduction/DMX/DMXEngine/Source/DMXEditor/Private/Widgets/FixturePatch/SDMXMVRFixtureListToolbar.cpp b/Engine/Plugins/VirtualProduction/DMX/DMXEngine/Source/DMXEditor/Private/Widgets/FixturePatch/SDMXMVRFixtureListToolbar.cpp
deleted file mode 100644
diff --git a/Engine/Plugins/VirtualProduction/DMX/DMXEngine/Source/DMXEditor/Private/Widgets/FixturePatch/SDMXMVRFixtureListToolbar.h b/Engine/Plugins/VirtualProduction/DMX/DMXEngine/Source/DMXEditor/Private/Widgets/FixturePatch/SDMXMVRFixtureListToolbar.h
deleted file mode 100644
diff --git a/Engine/Plugins/VirtualProduction/DMX/DMXEngine/Source/DMXEditor/Public/DMXEditorSettings.h b/Engine/Plugins/VirtualProduction/DMX/DMXEngine/Source/DMXEditor/Public/DMXEditorSettings.h
--- a/Engine/Plugins/VirtualProduction/DMX/DMXEngine/Source/DMXEditor/Public/DMXEditorSettings.h
+++ b/Engine/Plugins/VirtualProduction/DMX/DMXEngine/Source/DMXEditor/Public/DMXEditorSettings.h
@@ -15,27 +15,30 @@ struct FDMXMVRFixtureListSettings
 {
 	GENERATED_BODY()
 
-	/** Width of the entier list. The right side should take the rest of the space */
+	UPROPERTY()
 	float ListWidth = .5f;
 
 	/** Width of the Fixture ID column */
 	UPROPERTY()
-	float FixtureIDColumnWidth = 0.f;
+	float FixtureIDColumnWidth = .1f;
 
 	/** With of the Fixture Type column */
 	UPROPERTY()
-	float FixtureTypeColumnWidth = 0.f;
+	float FixtureTypeColumnWidth = .2f;
 
 	/** With of the Mode column */
 	UPROPERTY()
-	float ModeColumnWidth = 0.f;
+	float ModeColumnWidth = .2f;
 
 	/** With of the Patch column */
 	UPROPERTY()
-	float PatchColumnWidth = 0.f;
+	float PatchColumnWidth = .2f;
 
 	UPROPERTY()
 	FName SortByCollumnID = NAME_None;
+
+	UPROPERTY()
+	int32 SortPriorityEnumIndex = 1;
 };
 
 UENUM()
diff --git a/Engine/Plugins/VirtualProduction/DMX/DMXEngine/Source/DMXRuntime/Private/Library/DMXEntityFixturePatch.cpp b/Engine/Plugins/VirtualProduction/DMX/DMXEngine/Source/DMXRuntime/Private/Library/DMXEntityFixturePatch.cpp
--- a/Engine/Plugins/VirtualProduction/DMX/DMXEngine/Source/DMXRuntime/Private/Library/DMXEntityFixturePatch.cpp
+++ b/Engine/Plugins/VirtualProduction/DMX/DMXEngine/Source/DMXRuntime/Private/Library/DMXEntityFixturePatch.cpp
@@ -654,31 +654,9 @@ void UDMXEntityFixturePatch::GenerateFixtureID(int32 DesiredFixtureID)
 
 bool UDMXEntityFixturePatch::FindFixtureID(int32& OutFixtureID) const
 {
-	// DEPRECATED 5.5. The patch now holds its MVR fixture ID
-	if (!MVRFixtureUUID.IsValid())
-	{
-		return false;
-	}
-
-	const UDMXLibrary* DMXLibrary = GetParentLibrary();
-	if (!DMXLibrary)
-	{
-		return false;
-	}
-
-	UDMXMVRGeneralSceneDescription* GeneralSceneDescription = DMXLibrary->GetLazyGeneralSceneDescription();
-	if (!GeneralSceneDescription)
-	{
-		return false;
-	}
-
-	UDMXMVRFixtureNode* FixtureNode = GeneralSceneDescription->FindFixtureNode(MVRFixtureUUID);
-	if (!FixtureNode)
-	{
-		return false;
-	}
-
-	return LexTryParseString(OutFixtureID, *FixtureNode->FixtureID);
+	// DEPRECATED 5.5
+	OutFixtureID = FixtureID;
+	return true;
 }
 
 #if WITH_EDITOR
