DMX: Fix CIS Static Analysis Win64 (MSVC) Dereferencing NULL pointer 'DMXLibrary' in SDMXFixturePatchList

diff --git a/Engine/Plugins/VirtualProduction/DMX/DMXEngine/Source/DMXEditor/Private/Widgets/FixturePatch/SDMXFixturePatchList.cpp b/Engine/Plugins/VirtualProduction/DMX/DMXEngine/Source/DMXEditor/Private/Widgets/FixturePatch/SDMXFixturePatchList.cpp
--- a/Engine/Plugins/VirtualProduction/DMX/DMXEngine/Source/DMXEditor/Private/Widgets/FixturePatch/SDMXFixturePatchList.cpp
+++ b/Engine/Plugins/VirtualProduction/DMX/DMXEngine/Source/DMXEditor/Private/Widgets/FixturePatch/SDMXFixturePatchList.cpp
@@ -679,7 +679,7 @@ void SDMXFixturePatchList::RefreshList()
 
 	const TSharedPtr<FDMXEditor> DMXEditor = WeakDMXEditor.Pin();
 	UDMXLibrary* DMXLibrary = DMXEditor.IsValid() ? DMXEditor->GetDMXLibrary() : nullptr;
-	if (!IsValid(DMXLibrary))
+	if (!DMXLibrary)
 	{
 		return;
 	}
