DMX: Restore functionality to export GDTFs as they were imported, fix GDTF related bugs

* Fixture Types generated from GDTF now export as imported, unless set to generate a GDTF in advanced properties
* Fix issues with the geometrical layout of matrix fixtures
* Fix the DMXValue implementation, default values are now loaded correctly
* Allow users to mend  GDTFs with an invalid asset name due to a bug in UE5.4.0

diff --git a/Engine/Plugins/VirtualProduction/DMX/DMXEngine/Source/DMXEditor/Private/Customizations/DMXEntityFixtureTypeDetails.cpp b/Engine/Plugins/VirtualProduction/DMX/DMXEngine/Source/DMXEditor/Private/Customizations/DMXEntityFixtureTypeDetails.cpp
--- a/Engine/Plugins/VirtualProduction/DMX/DMXEngine/Source/DMXEditor/Private/Customizations/DMXEntityFixtureTypeDetails.cpp
+++ b/Engine/Plugins/VirtualProduction/DMX/DMXEngine/Source/DMXEditor/Private/Customizations/DMXEntityFixtureTypeDetails.cpp
@@ -2,15 +2,14 @@
 
 #include "Customizations/DMXEntityFixtureTypeDetails.h"
 
+#include "DetailCategoryBuilder.h"
 #include "DetailLayoutBuilder.h"
-#include "DMXEditorLog.h"
 #include "Factories/DMXGDTFToFixtureTypeConverter.h"
 #include "IPropertyUtilities.h"
 #include "Library/DMXEntityFixtureType.h"
 #include "Library/DMXImportGDTF.h"
 #include "PropertyHandle.h"
 
-
 #define LOCTEXT_NAMESPACE "DMXEntityFixtureTypeDetails"
 
 TSharedRef<IDetailCustomization> FDMXEntityFixtureTypeDetails::MakeInstance()
@@ -26,34 +25,73 @@ void FDMXEntityFixtureTypeDetails::CustomizeDetails(IDetailLayoutBuilder& Detail
 
 	GDTFSourceHandle = DetailBuilder.GetProperty(GET_MEMBER_NAME_CHECKED(UDMXEntityFixtureType, GDTFSource));
 	GDTFSourceHandle->SetOnPropertyValueChanged(FSimpleDelegate::CreateSP(this, &FDMXEntityFixtureTypeDetails::OnGDTFSourceChanged));
+
+	// Customize the bExportGeneratedGDTF property
+	const TSharedRef<IPropertyHandle> ExportGeneratedGDTFHandle = DetailBuilder.GetProperty(GET_MEMBER_NAME_CHECKED(UDMXEntityFixtureType, bExportGeneratedGDTF));
+	const FName GDTFCategory = ExportGeneratedGDTFHandle->GetDefaultCategoryName();
+	DetailBuilder.EditCategory(GDTFCategory)
+		.AddProperty(ExportGeneratedGDTFHandle, EPropertyLocation::Advanced)
+		.Visibility(TAttribute<EVisibility>::CreateSP(this, &FDMXEntityFixtureTypeDetails::GetExportGeneratedGDTFPropertyVisibility));
 }
 
 void FDMXEntityFixtureTypeDetails::OnGDTFSourceChanged()
 {
 	using namespace UE::DMX::GDTF;
 
-	const TArray<TWeakObjectPtr<UObject>>& SelectedObjects = PropertyUtilities->GetSelectedObjects();
+	TArray<void*> RawDataArray;
+	GDTFSourceHandle->AccessRawData(RawDataArray);
 
-	for (TWeakObjectPtr<UObject> WeakFixtureTypeObject : SelectedObjects)
+	TSoftObjectPtr<UDMXImportGDTF>* GDTFAssetPtr = nullptr;
+	if (!RawDataArray.IsEmpty())
 	{
-		if (UDMXEntityFixtureType* FixtureType = Cast<UDMXEntityFixtureType>(WeakFixtureTypeObject.Get()))
-		{
-			FixtureType->PreEditChange(nullptr);
-			FixtureType->Modes.Reset();
-			FixtureType->PostEditChange();
+		GDTFAssetPtr = reinterpret_cast<TSoftObjectPtr<UDMXImportGDTF>*>(RawDataArray[0]);
+	}
 
-			if (FixtureType->GDTFSource.IsNull())
+	const TArray<TWeakObjectPtr<UObject>>& SelectedObjects = PropertyUtilities->GetSelectedObjects();
+	if (GDTFAssetPtr)
+	{
+		if (UDMXImportGDTF* GDTF = GDTFAssetPtr->LoadSynchronous())
+		{
+			for (const TWeakObjectPtr<UObject>& WeakFixtureTypeObject : SelectedObjects)
 			{
-				continue;
+				if (UDMXEntityFixtureType* FixtureType = Cast<UDMXEntityFixtureType>(WeakFixtureTypeObject.Get()))
+				{
+					FixtureType->PreEditChange(nullptr);
+					constexpr bool bUpdateFixtureTypeName = true;
+					FDMXGDTFToFixtureTypeConverter::ConvertGDTF(*FixtureType, *GDTF, bUpdateFixtureTypeName);
+					FixtureType->PostEditChange();
+				}
 			}
-			UDMXImportGDTF* GDTF = FixtureType->GDTFSource.LoadSynchronous();
-
-			FixtureType->PreEditChange(nullptr);
-			constexpr bool bUpdateFixtureTypeName = true;
-			FDMXGDTFToFixtureTypeConverter::ConvertGDTF(*FixtureType, *GDTF, bUpdateFixtureTypeName);
-			FixtureType->PostEditChange();
 		}
+		else
+		{
+			for (const TWeakObjectPtr<UObject>& WeakFixtureTypeObject : SelectedObjects)
+			{
+				if (UDMXEntityFixtureType* FixtureType = Cast<UDMXEntityFixtureType>(WeakFixtureTypeObject.Get()))
+				{
+					FixtureType->PreEditChange(nullptr);
+					FixtureType->Modes.Reset();
+					FixtureType->PostEditChange();
+				}
+			}
+		}
+	}
+}
+
+EVisibility FDMXEntityFixtureTypeDetails::GetExportGeneratedGDTFPropertyVisibility() const
+{
+	TArray<const void*> RawDataArray;
+	GDTFSourceHandle->AccessRawData(RawDataArray);
+
+	if (!RawDataArray.IsEmpty())
+	{
+		const TSoftObjectPtr<UDMXImportGDTF>* GDTFAssetPtr = reinterpret_cast<const TSoftObjectPtr<UDMXImportGDTF>*>(RawDataArray[0]);
+
+		const bool bHasGDTF = GDTFAssetPtr && GDTFAssetPtr->IsValid();
+		return bHasGDTF ? EVisibility::Visible : EVisibility::Collapsed;
 	}
+
+	return EVisibility::Collapsed;
 }
 
 #undef LOCTEXT_NAMESPACE
diff --git a/Engine/Plugins/VirtualProduction/DMX/DMXEngine/Source/DMXEditor/Private/Customizations/DMXEntityFixtureTypeDetails.h b/Engine/Plugins/VirtualProduction/DMX/DMXEngine/Source/DMXEditor/Private/Customizations/DMXEntityFixtureTypeDetails.h
--- a/Engine/Plugins/VirtualProduction/DMX/DMXEngine/Source/DMXEditor/Private/Customizations/DMXEntityFixtureTypeDetails.h
+++ b/Engine/Plugins/VirtualProduction/DMX/DMXEngine/Source/DMXEditor/Private/Customizations/DMXEntityFixtureTypeDetails.h
@@ -2,8 +2,8 @@
 
 #pragma once
 
-#include "CoreMinimal.h"
 #include "IDetailCustomization.h"
+#include "Layout/Visibility.h"
 
 class IPropertyHandle;
 class IPropertyUtilities;
@@ -26,6 +26,9 @@ private:
 	/** Called when the GDTFSource property changed */
 	void OnGDTFSourceChanged();
 
+	/** Returns the visibility of the bExportGeneratedGDTF property */
+	EVisibility GetExportGeneratedGDTFPropertyVisibility() const;
+
 	/** Handle to the GDTFSource property */
 	TSharedPtr<IPropertyHandle> GDTFSourceHandle;
 
diff --git a/Engine/Plugins/VirtualProduction/DMX/DMXEngine/Source/DMXEditor/Private/DMXEditor.cpp b/Engine/Plugins/VirtualProduction/DMX/DMXEngine/Source/DMXEditor/Private/DMXEditor.cpp
--- a/Engine/Plugins/VirtualProduction/DMX/DMXEngine/Source/DMXEditor/Private/DMXEditor.cpp
+++ b/Engine/Plugins/VirtualProduction/DMX/DMXEngine/Source/DMXEditor/Private/DMXEditor.cpp
@@ -13,7 +13,6 @@
 #include "DMXRuntimeLog.h"
 #include "Exporters/DMXMVRExporter.h"
 #include "Framework/Application/SlateApplication.h"
-#include "Framework/Notifications/NotificationManager.h"
 #include "IDesktopPlatform.h"
 #include "Library/DMXEntityFixturePatch.h"
 #include "Library/DMXEntityFixtureType.h"
@@ -21,15 +20,12 @@
 #include "Library/DMXLibrary.h"
 #include "Misc/MessageDialog.h"
 #include "Modes/DMXEditorApplicationMode.h"
-#include "ScopedTransaction.h"
-#include "Styling/AppStyle.h"
 #include "Toolbars/DMXEditorToolbar.h"
 #include "Utils.h"
 #include "Widgets/Docking/SDockTab.h"
 #include "Widgets/FixturePatch/SDMXFixturePatchEditor.h"
 #include "Widgets/FixtureType/SDMXFixtureTypeEditor.h"
 #include "Widgets/LibrarySettings/SDMXLibraryEditorTab.h"
-#include "Widgets/Notifications/SNotificationList.h"
 #include "Widgets/SDMXEntityEditor.h"
 
 #define LOCTEXT_NAMESPACE "FDMXEditor"
@@ -176,51 +172,7 @@ void FDMXEditor::ExportDMXLibrary() const
 		return;
 	}
 
-	IDesktopPlatform* DesktopPlatform = FDesktopPlatformModule::Get();
-	if (DesktopPlatform)
-	{
-		UDMXEditorSettings* DMXEditorSettings = GetMutableDefault<UDMXEditorSettings>();
-		check(DMXEditorSettings);
-
-		const FString LastMVRExportPath = DMXEditorSettings->LastMVRExportPath;
-		const FString DefaultPath = FPaths::DirectoryExists(LastMVRExportPath) ? LastMVRExportPath : FPaths::ProjectSavedDir();
-
-		TArray<FString> SaveFilenames;
-		const bool bSaveFile = DesktopPlatform->SaveFileDialog(
-			FSlateApplication::Get().FindBestParentWindowHandleForDialogs(nullptr),
-			LOCTEXT("ExportMVR", "Export MVR").ToString(),
-			DefaultPath,
-			DMXLibrary->GetName() + TEXT(".mvr"),
-			TEXT("My Virtual Rig (*.mvr)|*.mvr"),
-			EFileDialogFlags::None,
-			SaveFilenames);
-
-		if (!bSaveFile || SaveFilenames.IsEmpty())
-		{
-			return;
-		}
-
-		FText ErrorReason;
-		UE::DMX::FDMXMVRExporter::Export(DMXLibrary, SaveFilenames[0], ErrorReason);
-		if (ErrorReason.IsEmpty())
-		{
-			DMXEditorSettings->LastMVRExportPath = FPaths::GetPath(SaveFilenames[0]);
-			DMXEditorSettings->SaveConfig();
-
-			FNotificationInfo NotificationInfo(FText::Format(LOCTEXT("ExportDMXLibraryAsMVRSuccessNotification", "Successfully exported MVR to {0}."), FText::FromString(SaveFilenames[0])));
-			NotificationInfo.ExpireDuration = 5.f;
-
-			FSlateNotificationManager::Get().AddNotification(NotificationInfo);
-		}
-		else
-		{
-			FNotificationInfo NotificationInfo(ErrorReason);
-			NotificationInfo.ExpireDuration = 10.f;
-			NotificationInfo.Image = FAppStyle::GetBrush("Icons.Warning");
-
-			FSlateNotificationManager::Get().AddNotification(NotificationInfo);
-		}
-	}
+	UE::DMX::FDMXMVRExporter::Export(DMXLibrary);
 }
 
 void FDMXEditor::RegisterToolbarTab(const TSharedRef<class FTabManager>& InTabManager)
diff --git a/Engine/Plugins/VirtualProduction/DMX/DMXEngine/Source/DMXEditor/Private/DMXEditorModule.cpp b/Engine/Plugins/VirtualProduction/DMX/DMXEngine/Source/DMXEditor/Private/DMXEditorModule.cpp
--- a/Engine/Plugins/VirtualProduction/DMX/DMXEngine/Source/DMXEditor/Private/DMXEditorModule.cpp
+++ b/Engine/Plugins/VirtualProduction/DMX/DMXEngine/Source/DMXEditor/Private/DMXEditorModule.cpp
@@ -25,25 +25,27 @@
 #include "Framework/Commands/UICommandList.h"
 #include "Framework/Docking/TabManager.h"
 #include "Framework/MultiBox/MultiBoxBuilder.h"
+#include "IAssetTools.h"
 #include "ISequencerModule.h"
 #include "LevelEditor.h"
 #include "Library/DMXEntity.h"
 #include "Library/DMXEntityFixtureType.h"
 #include "Library/DMXEntityReference.h"
+#include "Library/DMXImportGDTF.h"
 #include "Library/DMXLibrary.h"
+#include "Misc/MessageDialog.h"
 #include "MVR/DMXMVRSceneActor.h"
+#include "ObjectTools.h"
 #include "PropertyEditorModule.h"
 #include "Sequencer/DMXLibraryTrackEditor.h"
 #include "Sequencer/TakeRecorderDMXLibrarySource.h"
 #include "ToolMenus.h"
 #include "Widgets/Docking/SDockTab.h"
-#include "Widgets/Images/SImage.h"
 #include "Widgets/Monitors/SDMXActivityMonitor.h"
 #include "Widgets/Monitors/SDMXChannelsMonitor.h"
 #include "Widgets/Monitors/SDMXConflictMonitor.h"
 #include "Widgets/PatchTool/SDMXPatchTool.h"
 
-
 #define LOCTEXT_NAMESPACE "DMXEditorModule"
 
 const FName FDMXEditorModule::DMXEditorAppIdentifier(TEXT("DMXEditorApp"));
@@ -73,6 +75,9 @@ void FDMXEditorModule::StartupModule()
 	CreateLevelEditorToolbarDMXMenuExtender();
 	
 	StartupPIEManager();
+
+	IAssetRegistry& AssetRegistry = FModuleManager::LoadModuleChecked<FAssetRegistryModule>(TEXT("AssetRegistry")).Get();
+	AssetRegistry.OnFilesLoaded().AddStatic(&FDMXEditorModule::FixGDTFAssetNames);
 }
 
 void FDMXEditorModule::ShutdownModule()
@@ -108,6 +113,72 @@ TSharedRef<FDMXEditor> FDMXEditorModule::CreateEditor(const EToolkitMode::Type M
 	return NewDMXEditor;
 }
 
+void FDMXEditorModule::FixGDTFAssetNames()
+{
+	if (!UObjectInitialized())
+	{
+		return;
+	}
+
+	IAssetRegistry& AssetRegistry = FModuleManager::LoadModuleChecked<FAssetRegistryModule>(FName("AssetRegistry")).Get();
+
+	constexpr bool bSearchSubclasses = true;
+	TArray<FAssetData> GDTFAssets;
+	AssetRegistry.GetAssetsByClass(UDMXImportGDTF::StaticClass()->GetClassPathName(), GDTFAssets, bSearchSubclasses);
+
+	TArray<UDMXImportGDTF*> InvalidGDTFAssets;
+	Algo::TransformIf(GDTFAssets, InvalidGDTFAssets,
+		[](const FAssetData& Asset)
+		{
+			return 
+				Asset.IsValid() &&
+				ObjectTools::SanitizeObjectName(Asset.AssetName.ToString()) != Asset.AssetName;
+		},
+		[](const FAssetData& Asset)
+		{
+			return Cast<UDMXImportGDTF>(Asset.GetAsset());
+		});
+
+	if (InvalidGDTFAssets.IsEmpty())
+	{
+		return;
+	}
+
+	FText MessageText = LOCTEXT("FixGDTFAssetNames", "Detected GDTF assets with invalid asset names.\nWould you like to rename them now?\n\n");
+	for (UDMXImportGDTF* GDTF : InvalidGDTFAssets)
+	{
+		MessageText = FText::Format(FText::FromString("{0}\n{1} -> {2}"), MessageText, FText::FromString(GDTF->GetName()), FText::FromString(ObjectTools::SanitizeObjectName(GDTF->GetName())));
+	}
+	
+	if (FMessageDialog::Open(EAppMsgType::YesNo, MessageText) == EAppReturnType::No)
+	{
+		return;
+	}
+
+	for (UDMXImportGDTF* GDTF : InvalidGDTFAssets)
+	{
+		if (GDTF->GetOutermostObject() == GDTF->GetPackage() && GDTF->HasAnyFlags(RF_ClassDefaultObject | RF_Standalone | RF_Public))
+		{
+			FString SaneName = ObjectTools::SanitizeObjectName(GDTF->GetName());
+			if (!IsUniqueObjectName(*SaneName, GDTF))
+			{
+				SaneName = MakeUniqueObjectName(GDTF->GetOutermost(), GDTF->GetClass(), *SaneName).ToString();
+			}
+
+			const FSoftObjectPath OldObjectPath = FSoftObjectPath(GDTF);
+
+			const FTopLevelAssetPath NewAssetPath(*GDTF->GetPathName(), *SaneName);
+			const FSoftObjectPath NewObjectPath(NewAssetPath);
+
+			TArray<FAssetRenameData> AssetToRename;
+			AssetToRename.Add(FAssetRenameData(OldObjectPath, NewObjectPath));
+
+			FAssetToolsModule& Module = FModuleManager::GetModuleChecked<FAssetToolsModule>("AssetTools");
+			Module.Get().RenameAssets(AssetToRename);
+		}
+	}
+}
+
 void FDMXEditorModule::BindDMXEditorCommands()
 {
 	check(!DMXLevelEditorMenuCommands.IsValid());
@@ -264,6 +335,8 @@ void FDMXEditorModule::RegisterAssetTypeActions()
 
 void FDMXEditorModule::RegisterClassCustomizations()
 {
+	using namespace UE::DMX;
+
 	// Details customization for the UDMXEntityFixtureType class
 	RegisterCustomClassLayout(UDMXEntityFixtureType::StaticClass()->GetFName(),
 		FOnGetDetailCustomizationInstance::CreateStatic(&FDMXEntityFixtureTypeDetails::MakeInstance)
diff --git a/Engine/Plugins/VirtualProduction/DMX/DMXEngine/Source/DMXEditor/Private/Exporters/DMXFixtureTypeToGDTFConverter.cpp b/Engine/Plugins/VirtualProduction/DMX/DMXEngine/Source/DMXEditor/Private/Exporters/DMXFixtureTypeToGDTFConverter.cpp
--- a/Engine/Plugins/VirtualProduction/DMX/DMXEngine/Source/DMXEditor/Private/Exporters/DMXFixtureTypeToGDTFConverter.cpp
+++ b/Engine/Plugins/VirtualProduction/DMX/DMXEngine/Source/DMXEditor/Private/Exporters/DMXFixtureTypeToGDTFConverter.cpp
@@ -4,6 +4,7 @@
 
 #include "Algo/Find.h"
 #include "Algo/Sort.h"
+#include "DMXFixtureTypeToGDTFGeometryFactory.h"
 #include "DMXGDTF.h"
 #include "DMXUnrealToGDTFAttributeConversion.h"
 #include "GDTF/AttributeDefinitions/DMXGDTFAttribute.h"
@@ -16,24 +17,26 @@
 #include "GDTF/DMXModes/DMXGDTFDMXChannel.h"
 #include "GDTF/DMXModes/DMXGDTFDMXMode.h"
 #include "GDTF/DMXModes/DMXGDTFLogicalChannel.h"
-#include "GDTF/Geometries/DMXGDTFBeamGeometry.h"
-#include "GDTF/Geometries/DMXGDTFGeometryBreak.h"
+#include "GDTF/Geometries/DMXGDTFGeometry.h"
 #include "GDTF/Geometries/DMXGDTFGeometryCollect.h"
-#include "GDTF/Geometries/DMXGDTFGeometryReference.h"
 #include "GDTF/Models/DMXGDTFModel.h"
-#include "Interfaces/IPluginManager.h"
 #include "Library/DMXEntityFixtureType.h"
 #include "Library/DMXLibrary.h"
 #include "XmlFile.h"
 
 namespace UE::DMX::GDTF
 {
-	const FName FDMXFixtureTypeToGTDFConverter::MatrixCellGeometryReferenceName = TEXT("Instance");
-
-	TSharedPtr<FXmlFile> FDMXFixtureTypeToGTDFConverter::Convert(const UDMXEntityFixtureType* UnrealFixtureType)
+	TSharedPtr<FXmlFile> FDMXFixtureTypeToGDTFConverter::Convert(const UDMXEntityFixtureType* UnrealFixtureType)
 	{
-		FDMXFixtureTypeToGTDFConverter Converter;
-		const TSharedRef<FDMXGDTFFixtureType> GDTFFixtureType = Converter.CreateFixtureType(UnrealFixtureType);
+		if (!UnrealFixtureType)
+		{
+			return nullptr;
+		}
+
+		FDMXFixtureTypeToGDTFConverter Converter;
+
+		// Convert to GDTF
+		const TSharedRef<FDMXGDTFFixtureType> GDTFFixtureType = Converter.CreateFixtureType(*UnrealFixtureType);
 
 		// Create the XML file
 		UDMXGDTF* GDTF = NewObject<UDMXGDTF>();
@@ -43,15 +46,15 @@ namespace UE::DMX::GDTF
 		return DescriptionXml;
 	}
 
-	TSharedRef<FDMXGDTFFixtureType> FDMXFixtureTypeToGTDFConverter::CreateFixtureType(const UDMXEntityFixtureType* UnrealFixtureType)
+	TSharedRef<FDMXGDTFFixtureType> FDMXFixtureTypeToGDTFConverter::CreateFixtureType(const UDMXEntityFixtureType& UnrealFixtureType)
 	{
 		const TSharedRef<FDMXGDTFFixtureType> GDTFFixtureType = MakeShared<FDMXGDTFFixtureType>();
 
-		GDTFFixtureType->Name = *UnrealFixtureType->Name;
-		GDTFFixtureType->ShortName = *UnrealFixtureType->Name;
-		GDTFFixtureType->LongName = UnrealFixtureType->GetParentLibrary()->GetName() + TEXT(" ") + UnrealFixtureType->Name;
+		GDTFFixtureType->Name = *UnrealFixtureType.Name;
+		GDTFFixtureType->ShortName = *UnrealFixtureType.Name;
+		GDTFFixtureType->LongName = UnrealFixtureType.GetParentLibrary()->GetName() + TEXT(" ") + UnrealFixtureType.Name;
 		GDTFFixtureType->Manufacturer = TEXT("Epic Games");
-		GDTFFixtureType->Description = FString::Printf(TEXT("Unreal Engine generated fixture type"));
+		GDTFFixtureType->Description = FString::Printf(TEXT("Unreal Engine generated Fixture Type"));
 		GDTFFixtureType->FixtureTypeID = FGuid::NewGuid(); // Avoid any ambiguity with previously exported GDTFs, even if they're identical.
 		GDTFFixtureType->bCanHaveChildren = false;
 
@@ -63,10 +66,10 @@ namespace UE::DMX::GDTF
 		return GDTFFixtureType;
 	}
 
-	void FDMXFixtureTypeToGTDFConverter::CreateAttributeDefinitions(const UDMXEntityFixtureType* UnrealFixtureType, const TSharedRef<FDMXGDTFFixtureType>& GDTFFixtureType)
+	void FDMXFixtureTypeToGDTFConverter::CreateAttributeDefinitions(const UDMXEntityFixtureType& UnrealFixtureType, const TSharedRef<FDMXGDTFFixtureType>& GDTFFixtureType)
 	{
 		TArray<FName> AttributeNames;
-		for (const FDMXFixtureMode& Mode : UnrealFixtureType->Modes)
+		for (const FDMXFixtureMode& Mode : UnrealFixtureType.Modes)
 		{
 			for (const FDMXFixtureFunction& Function : Mode.Functions)
 			{
@@ -141,10 +144,10 @@ namespace UE::DMX::GDTF
 		}
 	}
 
-	void FDMXFixtureTypeToGTDFConverter::CreateModels(const UDMXEntityFixtureType* UnrealFixtureType, const TSharedRef<FDMXGDTFFixtureType>& GDTFFixtureType)
+	void FDMXFixtureTypeToGDTFConverter::CreateModels(const UDMXEntityFixtureType& UnrealFixtureType, const TSharedRef<FDMXGDTFFixtureType>& GDTFFixtureType)
 	{
 		// Create a model for the matrix if this is a matrix
-		const bool bIsMatrix = Algo::FindBy(UnrealFixtureType->Modes, true, &FDMXFixtureMode::bFixtureMatrixEnabled) != nullptr;
+		const bool bIsMatrix = Algo::FindBy(UnrealFixtureType.Modes, true, &FDMXFixtureMode::bFixtureMatrixEnabled) != nullptr;
 		if (bIsMatrix)
 		{
 			const TSharedRef<FDMXGDTFModel> InstanceModel = MakeShared<FDMXGDTFModel>(GDTFFixtureType);
@@ -158,92 +161,21 @@ namespace UE::DMX::GDTF
 		}
 	}
 
-	void FDMXFixtureTypeToGTDFConverter::CreateGeometryCollect(const UDMXEntityFixtureType* UnrealFixtureType, const TSharedRef<FDMXGDTFFixtureType>& GDTFFixtureType)
+	void FDMXFixtureTypeToGDTFConverter::CreateGeometryCollect(const UDMXEntityFixtureType& UnrealFixtureType, const TSharedRef<FDMXGDTFFixtureType>& GDTFFixtureType)
 	{
 		GDTFFixtureType->GeometryCollect = MakeShared<FDMXGDTFGeometryCollect>(GDTFFixtureType);
 
-		// Always add a base geometry
-		const TSharedRef<FDMXGDTFGeometry> BaseGeometry = MakeShared<FDMXGDTFGeometry>(GDTFFixtureType->GeometryCollect.ToSharedRef());
-		GDTFFixtureType->GeometryCollect->GeometryArray.Add(BaseGeometry);
-
-		BaseGeometry->Name = TEXT("Base");
-
-		const TSharedRef<FDMXGDTFBeamGeometry> BeamGeometry = MakeShared<FDMXGDTFBeamGeometry>(GDTFFixtureType->GeometryCollect.ToSharedRef());
-		BaseGeometry->BeamArray.Add(BeamGeometry);
-
-		BeamGeometry->Name = TEXT("Beam");
-
-		for (const FDMXFixtureMode& UnrealMode : UnrealFixtureType->Modes)
-		{				
-			// Remember the geometry for this mode so it later can be referenced when building DMX Modes
-			UnrealModeToRootGeometryMap.Add(&UnrealMode, BaseGeometry);
-
-			CreateChildGeometries(UnrealMode, BeamGeometry);
-		}
-	}
-
-	void FDMXFixtureTypeToGTDFConverter::CreateChildGeometries(const FDMXFixtureMode& UnrealMode, const TSharedRef<FDMXGDTFGeometry>& RootGeometry)
-	{
-		if (UnrealMode.bFixtureMatrixEnabled && UnrealMode.FixtureMatrixConfig.GetNumChannels() > 0)
-		{
-			// Get the byte size of an Unreal Matrix Cell
-			const int32 CellSize = [UnrealMode]()
-				{
-					int32 OutCellSize = 0;
-					for (const FDMXFixtureCellAttribute& CellAttribute : UnrealMode.FixtureMatrixConfig.CellAttributes)
-					{
-						OutCellSize += CellAttribute.GetNumChannels();
-					}
+		FDMXFixtureTypeToGDTFGeometryFactory GeometryFactory(UnrealFixtureType, GDTFFixtureType->GeometryCollect.ToSharedRef());
 
-					return OutCellSize;
-				}();
-
-			// Create Geometry References for each Unreal Matrix Cell
-			const int32 NumCells = UnrealMode.FixtureMatrixConfig.XCells * UnrealMode.FixtureMatrixConfig.YCells;
-			int32 Offset = 1;
-			for (int32 CellID = 0; CellID < NumCells; CellID++)
-			{
-				const int32 DMXOffset = CellID * CellSize + 1;
-				const FName GeometryName = *FString::Printf(TEXT("Layer_%i"), CellID + 1);
-
-				const TSharedRef<FDMXGDTFGeometryReference> GeometryReference = MakeShared<FDMXGDTFGeometryReference>(RootGeometry);
-				RootGeometry->GeometryReferenceArray.Add(GeometryReference);
-
-				GeometryReference->Name = GeometryName;
-				GeometryReference->Geometry = MatrixCellGeometryReferenceName;
-				GeometryReference->Model = TEXT("Layers");
-
-				const TSharedRef<FDMXGDTFGeometryBreak> Break = MakeShared<FDMXGDTFGeometryBreak>(GeometryReference);
-				Break->DMXBreak = 1;
-				Break->DMXOffset = Offset;
-				GeometryReference->BreakArray.Add(Break);
-
-				Break->DMXBreak = 1; // Unreal does not support multi universe patches, DMXBreak is always 1
-				Break->DMXOffset = DMXOffset;
-
-				for (const FDMXFixtureCellAttribute& UnrealMatrixAttribute : UnrealMode.FixtureMatrixConfig.CellAttributes)
-				{
-					UnrealCellAttributeToGeometryReferenceMap.Add(&UnrealMatrixAttribute, GeometryReference);
-				}
-
-				Offset += CellSize;
-			}
-		}
-
-		// Always add common functions
-		for (const FDMXFixtureFunction& UnrealFunction : UnrealMode.Functions)
-		{
-			// Common Unreal Functions are assigned to the root geometry for now so they do not need specific handling
-			UnrealFunctionToGeometryMap.Add(&UnrealFunction, RootGeometry);
-		}
+		FunctionsWithControlledGeometry = GeometryFactory.GetFunctionsWithControlledGeometry();
 	}
 
-	void FDMXFixtureTypeToGTDFConverter::CreateDMXModes(const UDMXEntityFixtureType* UnrealFixtureType, const TSharedRef<FDMXGDTFFixtureType>& GDTFFixtureType)
+	void FDMXFixtureTypeToGDTFConverter::CreateDMXModes(const UDMXEntityFixtureType& UnrealFixtureType, const TSharedRef<FDMXGDTFFixtureType>& GDTFFixtureType)
 	{
-		for (const FDMXFixtureMode& UnrealMode : UnrealFixtureType->Modes)
+		for (const FDMXFixtureMode& UnrealMode : UnrealFixtureType.Modes)
 		{
-			const TSharedRef<FDMXGDTFGeometry>* RootGeometryPtr = UnrealModeToRootGeometryMap.Find(&UnrealMode);
-			if (!ensureMsgf(RootGeometryPtr, TEXT("%hs: Unexpected cannot find root geometry for DMX Mode '%s'. Failed to convert mode to GDTF."), __FUNCTION__, *UnrealMode.ModeName))
+			const FDMXFixtureFunctionWithControlledGeometry* ControlledGeometryPtr = Algo::FindBy(FunctionsWithControlledGeometry, &UnrealMode, &FDMXFixtureFunctionWithControlledGeometry::ModePtr);
+			if (!ensureMsgf(ControlledGeometryPtr, TEXT("%hs: Unexpected cannot controlled root geometry for DMX Mode '%s'. Failed to convert mode to GDTF."), __FUNCTION__, *UnrealMode.ModeName))
 			{
 				continue;
 			}
@@ -254,19 +186,19 @@ namespace UE::DMX::GDTF
 
 			DMXMode->Name = *UnrealMode.ModeName;
 			DMXMode->Description = TEXT("Unreal Engine generated DMX Mode");
-			DMXMode->Geometry = (*RootGeometryPtr)->Name;
+			DMXMode->Geometry = ControlledGeometryPtr->ControlledGeometry->Name;
 
 			CreateDMXChannels(UnrealMode, DMXMode);
 		}
 	}
 
-	void FDMXFixtureTypeToGTDFConverter::CreateDMXChannels(const FDMXFixtureMode& UnrealMode, const TSharedRef<FDMXGDTFDMXMode>& GDTFDMXMode)
+	void FDMXFixtureTypeToGDTFConverter::CreateDMXChannels(const FDMXFixtureMode& UnrealMode, const TSharedRef<FDMXGDTFDMXMode>& GDTFDMXMode)
 	{			
 		// Create DMX Channels for non-matrix Unreal Functions
 		for (const FDMXFixtureFunction& UnrealFunction : UnrealMode.Functions)
 		{
-			const TSharedRef<FDMXGDTFGeometry>* DMXChannelGeometryPtr = UnrealFunctionToGeometryMap.Find(&UnrealFunction);
-			if (!ensureMsgf(DMXChannelGeometryPtr, TEXT("%hs: Unexpected cannot find geometry for DMX Function '%s'. Failed to convert mode to GDTF."), __FUNCTION__, *UnrealFunction.FunctionName))
+			const FDMXFixtureFunctionWithControlledGeometry* ControlledGeometryPtr = Algo::FindBy(FunctionsWithControlledGeometry, &UnrealFunction, &FDMXFixtureFunctionWithControlledGeometry::FunctionPtr);
+			if (!ensureMsgf(ControlledGeometryPtr, TEXT("%hs: Unexpected cannot find geometry for DMX Function '%s'. Failed to convert mode to GDTF."), __FUNCTION__, *UnrealFunction.FunctionName))
 			{
 				continue;
 			}
@@ -274,12 +206,13 @@ namespace UE::DMX::GDTF
 			const TSharedRef<FDMXGDTFDMXChannel> DMXChannel = MakeShared<FDMXGDTFDMXChannel>(GDTFDMXMode);
 			GDTFDMXMode->DMXChannels.Add(DMXChannel);
 
+			const FString ControlledGeometryName = ControlledGeometryPtr->ControlledGeometry->Name.ToString();
 			const FString GDTFAttribute = FDMXUnrealToGDTFAttributeConversion::ConvertUnrealToGDTFAttribute(UnrealFunction.Attribute.Name).ToString();
 			const FString ChannelFunctionName = UnrealFunction.FunctionName;
 
 			// The initial function has to be written in following format "GeometryName_LogicalChannelAttribute.ChannelFunctionAttribute.ChannelFunctionName"
-			DMXChannel->InitialFunction = FString::Printf(TEXT("%s_%s.%s.%s"), *(*DMXChannelGeometryPtr)->Name.ToString(), *GDTFAttribute, *GDTFAttribute, *ChannelFunctionName);
-			DMXChannel->Geometry = (*DMXChannelGeometryPtr)->Name;
+			DMXChannel->InitialFunction = FString::Printf(TEXT("%s_%s.%s.%s"), *ControlledGeometryName, *GDTFAttribute, *GDTFAttribute, *ChannelFunctionName);
+			DMXChannel->Geometry = *ControlledGeometryName;
 			DMXChannel->Offset = [UnrealFunction]()
 				{			
 					const int32 Offset = UnrealFunction.Channel;
@@ -313,11 +246,12 @@ namespace UE::DMX::GDTF
 				const TSharedRef<FDMXGDTFDMXChannel> DMXChannel = MakeShared<FDMXGDTFDMXChannel>(GDTFDMXMode);
 				GDTFDMXMode->DMXChannels.Add(DMXChannel);
 
+				const FString MatrixBeamGeometryName = *FDMXFixtureTypeToGDTFGeometryFactory::MatrixBeamGeometryName.ToString();
 				const FString GDTFAttribute = FDMXUnrealToGDTFAttributeConversion::ConvertUnrealToGDTFAttribute(UnrealCellAttribute.Attribute.Name).ToString();
 
 				// For a matrix with geometry references, the initial function has to be written in following format "GeometryName_LogicalChannelAttribute.ChannelFunctionAttribute.ChannelFunctionName"
-				DMXChannel->InitialFunction = FString::Printf(TEXT("%s_%s.%s.%s"), *MatrixCellGeometryReferenceName.ToString(), *GDTFAttribute, *GDTFAttribute, *UnrealCellAttribute.Attribute.Name.ToString());
-				DMXChannel->Geometry = MatrixCellGeometryReferenceName;
+				DMXChannel->InitialFunction = FString::Printf(TEXT("%s_%s.%s.%s"), *MatrixBeamGeometryName, *GDTFAttribute, *GDTFAttribute, *UnrealCellAttribute.Attribute.Name.ToString());
+				DMXChannel->Geometry = *MatrixBeamGeometryName;
 				DMXChannel->Offset = [&Offset, UnrealCellAttribute]()
 					{					
 						const uint8 Size = UnrealCellAttribute.GetNumChannels();
@@ -357,7 +291,7 @@ namespace UE::DMX::GDTF
 			});
 	}
 
-	void FDMXFixtureTypeToGTDFConverter::CreateLogicalChannel(const FDMXFixtureFunction& UnrealFunction, const TSharedRef<FDMXGDTFDMXChannel>& GDTFDMXChannel, const FString& GDTFAttribute)
+	void FDMXFixtureTypeToGDTFConverter::CreateLogicalChannel(const FDMXFixtureFunction& UnrealFunction, const TSharedRef<FDMXGDTFDMXChannel>& GDTFDMXChannel, const FString& GDTFAttribute)
 	{
 		const TSharedRef<FDMXGDTFLogicalChannel> LogicalChannel = MakeShared<FDMXGDTFLogicalChannel>(GDTFDMXChannel);
 		GDTFDMXChannel->LogicalChannelArray.Add(LogicalChannel);
@@ -367,7 +301,7 @@ namespace UE::DMX::GDTF
 		CrateChannelFunction(UnrealFunction, LogicalChannel, GDTFAttribute);
 	}
 
-	void FDMXFixtureTypeToGTDFConverter::CreateLogicalChannel(const FDMXFixtureCellAttribute& UnrealCellAttribute, const TSharedRef<FDMXGDTFDMXChannel>& GDTFDMXChannel, const FString& GDTFAttribute)
+	void FDMXFixtureTypeToGDTFConverter::CreateLogicalChannel(const FDMXFixtureCellAttribute& UnrealCellAttribute, const TSharedRef<FDMXGDTFDMXChannel>& GDTFDMXChannel, const FString& GDTFAttribute)
 	{
 		const TSharedRef<FDMXGDTFLogicalChannel> LogicalChannel = MakeShared<FDMXGDTFLogicalChannel>(GDTFDMXChannel);
 		GDTFDMXChannel->LogicalChannelArray.Add(LogicalChannel);
@@ -377,7 +311,7 @@ namespace UE::DMX::GDTF
 		CrateChannelFunction(UnrealCellAttribute, LogicalChannel, GDTFAttribute);
 	}
 
-	void FDMXFixtureTypeToGTDFConverter::CrateChannelFunction(const FDMXFixtureFunction& UnrealFunction, const TSharedRef<FDMXGDTFLogicalChannel>& GDTFLogicalChannel, const FString& GDTFAttribute)
+	void FDMXFixtureTypeToGDTFConverter::CrateChannelFunction(const FDMXFixtureFunction& UnrealFunction, const TSharedRef<FDMXGDTFLogicalChannel>& GDTFLogicalChannel, const FString& GDTFAttribute)
 	{
 		const TSharedRef<FDMXGDTFChannelFunction> ChannelFunction = MakeShared<FDMXGDTFChannelFunction>(GDTFLogicalChannel);
 		GDTFLogicalChannel->ChannelFunctionArray.Add(ChannelFunction);
@@ -385,12 +319,12 @@ namespace UE::DMX::GDTF
 		ChannelFunction->Name = *UnrealFunction.FunctionName;
 		ChannelFunction->Attribute = GDTFAttribute;
 
-		const FDMXGDTFDMXValue Default = UnrealFunction.DefaultValue;
+		const FDMXGDTFDMXValue Default = FDMXGDTFDMXValue(UnrealFunction.DefaultValue, UnrealFunction.GetNumChannels());
 		ChannelFunction->Default = Default;
 		ChannelFunction->DMXFrom = 0;
 	}
 
-	void FDMXFixtureTypeToGTDFConverter::CrateChannelFunction(const FDMXFixtureCellAttribute& UnrealCellAttribute, const TSharedRef<FDMXGDTFLogicalChannel>& GDTFLogicalChannel, const FString& GDTFAttribute)
+	void FDMXFixtureTypeToGDTFConverter::CrateChannelFunction(const FDMXFixtureCellAttribute& UnrealCellAttribute, const TSharedRef<FDMXGDTFLogicalChannel>& GDTFLogicalChannel, const FString& GDTFAttribute)
 	{
 		const TSharedRef<FDMXGDTFChannelFunction> ChannelFunction = MakeShared<FDMXGDTFChannelFunction>(GDTFLogicalChannel);
 		GDTFLogicalChannel->ChannelFunctionArray.Add(ChannelFunction);
diff --git a/Engine/Plugins/VirtualProduction/DMX/DMXEngine/Source/DMXEditor/Private/Exporters/DMXFixtureTypeToGDTFConverter.h b/Engine/Plugins/VirtualProduction/DMX/DMXEngine/Source/DMXEditor/Private/Exporters/DMXFixtureTypeToGDTFConverter.h
--- a/Engine/Plugins/VirtualProduction/DMX/DMXEngine/Source/DMXEditor/Private/Exporters/DMXFixtureTypeToGDTFConverter.h
+++ b/Engine/Plugins/VirtualProduction/DMX/DMXEngine/Source/DMXEditor/Private/Exporters/DMXFixtureTypeToGDTFConverter.h
@@ -3,7 +3,6 @@
 #pragma once
 
 #include "Templates/SharedPointer.h"
-#include "UObject/NameTypes.h"
 
 struct FDMXFixtureFunction;
 struct FDMXFixtureCellAttribute;
@@ -13,6 +12,7 @@ class UDMXEntityFixtureType;
 
 namespace UE::DMX::GDTF
 {
+	struct FDMXFixtureFunctionWithControlledGeometry;
 	class FDMXGDTFAttributeDefinitions;
 	class FDMXGDTFDMXChannel;
 	class FDMXGDTFDMXMode;
@@ -23,30 +23,27 @@ namespace UE::DMX::GDTF
 	class FDMXGDTFLogicalChannel;
 
 	/** Converts a Fixture Type to a GDTF. Internally caches of each collect, and finally assembles the GDTF. */
-	class FDMXFixtureTypeToGTDFConverter
+	class FDMXFixtureTypeToGDTFConverter
 	{
 	public:
 		/** Converts the Fixture Type to a GDTF description */
 		static TSharedPtr<FXmlFile> Convert(const UDMXEntityFixtureType* UnrealFixtureType);
 
 	private:
-		/** Creates the Fixture Type */
-		TSharedRef<FDMXGDTFFixtureType> CreateFixtureType(const UDMXEntityFixtureType* UnrealFixtureType);
+		/** Creates the GDTF Fixture Type */
+		TSharedRef<FDMXGDTFFixtureType> CreateFixtureType(const UDMXEntityFixtureType& UnrealFixtureType);
 
-		/** Creates attribute definitions */
-		void CreateAttributeDefinitions(const UDMXEntityFixtureType* UnrealFixtureType, const TSharedRef<FDMXGDTFFixtureType>& GDTFFixtureType);
+		/** Creates Attribute Aefinitions */
+		void CreateAttributeDefinitions(const UDMXEntityFixtureType& UnrealFixtureType, const TSharedRef<FDMXGDTFFixtureType>& GDTFFixtureType);
 
-		/** Creates model */
-		void CreateModels(const UDMXEntityFixtureType* UnrealFixtureType, const TSharedRef<FDMXGDTFFixtureType>& GDTFFixtureType);
+		/** Creates Models */
+		void CreateModels(const UDMXEntityFixtureType& UnrealFixtureType, const TSharedRef<FDMXGDTFFixtureType>& GDTFFixtureType);
 
-		/** Creates the geometry collect */
-		void CreateGeometryCollect(const UDMXEntityFixtureType* UnrealFixtureType, const TSharedRef<FDMXGDTFFixtureType>& GDTFFixtureType);
-
-		/** Creates Child geometries inside a root geometry */
-		void CreateChildGeometries(const FDMXFixtureMode& UnrealMode, const TSharedRef<FDMXGDTFGeometry>& RootGeometry);
+		/** Creates the Geometry Collect */
+		void CreateGeometryCollect(const UDMXEntityFixtureType& UnrealFixtureType, const TSharedRef<FDMXGDTFFixtureType>& GDTFFixtureType);
 
 		/** Creates DMX Modes */
-		void CreateDMXModes(const UDMXEntityFixtureType* UnrealFixtureType, const TSharedRef<FDMXGDTFFixtureType>& GDTFFixtureType);
+		void CreateDMXModes(const UDMXEntityFixtureType& UnrealFixtureType, const TSharedRef<FDMXGDTFFixtureType>& GDTFFixtureType);
 
 		/** Creates DMX Channels for specified mode */
 		void CreateDMXChannels(const FDMXFixtureMode& UnrealMode, const TSharedRef<FDMXGDTFDMXMode>& GDTFDMXMode);
@@ -63,16 +60,7 @@ namespace UE::DMX::GDTF
 		/** Creates a Channel Function for the Logical Channel, using an Unreal Cell Attribute as the input */
 		void CrateChannelFunction(const FDMXFixtureCellAttribute& UnrealCellAttribute, const TSharedRef<FDMXGDTFLogicalChannel>& GDTFLogicalChannel, const FString& GDTFAttribute);
 
-		/** Map of Unreal Modes to the root geometry they control */
-		TMap<const FDMXFixtureMode*, TSharedRef<FDMXGDTFGeometry>> UnrealModeToRootGeometryMap;
-
-		/** Map of Unreal Fixture Functions to the geometry they control */
-		TMap<const FDMXFixtureFunction*, TSharedRef<FDMXGDTFGeometry>> UnrealFunctionToGeometryMap;
-
-		/** Map of Unreal Matrix Attribute to the geometry they control */
-		TMap<const FDMXFixtureCellAttribute*, TSharedRef<FDMXGDTFGeometryReference>> UnrealCellAttributeToGeometryReferenceMap;
-
-		/** The name of a matrix cell geometry reference */
-		static const FName MatrixCellGeometryReferenceName;
+		/** Map of Functions with the geometry they control */
+		TArray<FDMXFixtureFunctionWithControlledGeometry> FunctionsWithControlledGeometry;
 	};
 }
diff --git a/Engine/Plugins/VirtualProduction/DMX/DMXEngine/Source/DMXEditor/Private/Exporters/DMXFixtureTypeToGDTFGeometryFactory.cpp b/Engine/Plugins/VirtualProduction/DMX/DMXEngine/Source/DMXEditor/Private/Exporters/DMXFixtureTypeToGDTFGeometryFactory.cpp
new file mode 100644
--- /dev/null
+++ b/Engine/Plugins/VirtualProduction/DMX/DMXEngine/Source/DMXEditor/Private/Exporters/DMXFixtureTypeToGDTFGeometryFactory.cpp
@@ -0,0 +1,211 @@
+// Copyright Epic Games, Inc. All Rights Reserved.
+
+#include "Exporters/DMXFixtureTypeToGDTFGeometryFactory.h"
+
+#include "Algo/AnyOf.h"
+#include "GDTF/Geometries/DMXGDTFAxisGeometry.h"
+#include "GDTF/Geometries/DMXGDTFGeometry.h"
+#include "GDTF/Geometries/DMXGDTFGeometryBreak.h"
+#include "GDTF/Geometries/DMXGDTFGeometryCollect.h"
+#include "GDTF/Geometries/DMXGDTFGeometryReference.h"
+#include "GDTF/Geometries/DMXGDTFBeamGeometry.h"
+#include "Library/DMXEntityFixtureType.h"
+
+namespace UE::DMX::GDTF
+{
+	const FName FDMXFixtureTypeToGDTFGeometryFactory::BaseGeometryName = "Base";
+	const FName FDMXFixtureTypeToGDTFGeometryFactory::PanGeometryName = "Yoke";
+	const FName FDMXFixtureTypeToGDTFGeometryFactory::HeadGeometryName = "Head";
+	const FName FDMXFixtureTypeToGDTFGeometryFactory::BeamGeometryName = "Beam";
+	const FName FDMXFixtureTypeToGDTFGeometryFactory::MatrixBeamGeometryName = "Instance";
+
+	namespace Internal
+	{
+		FDMXUniqueModeGeometry::FDMXUniqueModeGeometry(const FDMXFixtureMode& Mode)
+		{
+			bIsMatrix = Mode.bFixtureMatrixEnabled;
+
+			TotalNumChannels = bIsMatrix ? Mode.FixtureMatrixConfig.GetNumChannels() : 0;
+			for (const FDMXFixtureFunction& Function : Mode.Functions)
+			{
+				TotalNumChannels += Function.GetNumChannels();
+			}
+
+			bWithPan = Algo::AnyOf(Mode.Functions, [](const FDMXFixtureFunction& Function) { return Function.Attribute.Name == TEXT("Pan"); });
+			bWithTilt = Algo::AnyOf(Mode.Functions, [](const FDMXFixtureFunction& Function) { return Function.Attribute.Name == TEXT("Tilt"); });
+		}
+	}
+
+	FDMXFixtureFunctionWithControlledGeometry::FDMXFixtureFunctionWithControlledGeometry(const FDMXFixtureMode* InModePtr, const FDMXFixtureFunction* InFunctionPtr, TSharedRef<const FDMXGDTFGeometry> InControlledGeometry)
+		: ModePtr(InModePtr)
+		, FunctionPtr(InFunctionPtr)
+		, ControlledGeometry(InControlledGeometry)
+	{}
+
+	FDMXFixtureTypeToGDTFGeometryFactory::FDMXFixtureTypeToGDTFGeometryFactory(const UDMXEntityFixtureType& InFixtureType, const TSharedRef<FDMXGDTFGeometryCollect>& InBaseGeometryCollect)
+		: FixtureType(InFixtureType)
+		, BaseGeometryCollect(InBaseGeometryCollect)
+	{
+		BuildGeometries();
+	}
+
+	TArray<FDMXFixtureFunctionWithControlledGeometry> FDMXFixtureTypeToGDTFGeometryFactory::GetFunctionsWithControlledGeometry() const
+	{
+		TArray<FDMXFixtureFunctionWithControlledGeometry> Result;
+
+		for (const FDMXFixtureMode& Mode : FixtureType.Modes)
+		{
+			FDMXUniqueModeGeometry UniqueModeGeometry(Mode);
+			const TSharedRef<FDMXGDTFGeometry>* GeometryPtr = UniqueModeGeometryToGeometryMap.Find(UniqueModeGeometry);
+
+			if (!ensureMsgf(GeometryPtr, TEXT("Unexpected cannot find Geometry for Mode %s."), *Mode.ModeName))
+			{
+				return Result;
+			}
+
+			Algo::Transform(Mode.Functions, Result,
+				[&Mode, GeometryPtr](const FDMXFixtureFunction& Function)
+				{
+					return FDMXFixtureFunctionWithControlledGeometry(&Mode, &Function, *GeometryPtr);
+				});
+		}
+
+		return Result;
+	}
+
+	void FDMXFixtureTypeToGDTFGeometryFactory::BuildGeometries()
+	{
+		for (const FDMXFixtureMode& Mode : FixtureType.Modes)
+		{
+			BuildGeometry(Mode);
+		}
+	}
+
+	void FDMXFixtureTypeToGDTFGeometryFactory::BuildGeometry(const FDMXFixtureMode& Mode)
+	{
+		const FDMXUniqueModeGeometry UniqueModeGeometry(Mode);
+		const TSharedRef<FDMXGDTFGeometry>* BasePtr = UniqueModeGeometryToGeometryMap.Find(UniqueModeGeometry);
+
+		if (!BasePtr)
+		{
+			const TSharedRef<FDMXGDTFGeometry> Base = MakeShared<FDMXGDTFGeometry>(BaseGeometryCollect);
+			BasePtr = &Base;
+			Base->Name = BaseGeometryName;
+			BaseGeometryCollect->GeometryArray.Add(Base);
+
+			TSharedPtr<FDMXGDTFGeometry> Head;
+			if (UniqueModeGeometry.bWithPan && UniqueModeGeometry.bWithTilt)
+			{
+				const TSharedRef<FDMXGDTFGeometry> Yoke = AddYoke(Base);
+				Head = AddHead(Yoke);
+			}
+			else if (UniqueModeGeometry.bWithPan || UniqueModeGeometry.bWithTilt)
+			{
+				Head = AddHead(Base);
+			}
+			else
+			{
+				Head = Base;
+			}
+			check(Head.IsValid());
+
+			if (UniqueModeGeometry.bIsMatrix)
+			{
+				BuildMatrixGeometryReferences(Mode, Head.ToSharedRef());
+			}
+			else
+			{
+				AddBeam(Head.ToSharedRef());
+			}
+
+			UniqueModeGeometryToGeometryMap.Add(UniqueModeGeometry, Base);
+		}
+	}
+
+	TSharedRef<FDMXGDTFGeometry> FDMXFixtureTypeToGDTFGeometryFactory::GetOrCreateMatrixBeamGeometryInstance() const
+	{
+		if (BaseGeometryCollect->BeamArray.IsEmpty())
+		{
+			const TSharedRef<FDMXGDTFBeamGeometry> BeamGeometry = MakeShared<FDMXGDTFBeamGeometry>(BaseGeometryCollect);
+			BaseGeometryCollect->BeamArray.Add(BeamGeometry);
+
+			BeamGeometry->Name = MatrixBeamGeometryName;
+		}
+
+		return BaseGeometryCollect->BeamArray.Last().ToSharedRef();
+	}
+
+	TSharedRef<FDMXGDTFAxisGeometry> FDMXFixtureTypeToGDTFGeometryFactory::AddYoke(const TSharedRef<FDMXGDTFGeometry>& OuterGeometry)
+	{
+		const TSharedRef<FDMXGDTFAxisGeometry> Yoke = MakeShared<FDMXGDTFAxisGeometry>(OuterGeometry);
+		OuterGeometry->AxisArray.Add(Yoke);
+
+		Yoke->Name = PanGeometryName;
+
+		return Yoke;
+	}
+
+	TSharedRef<FDMXGDTFAxisGeometry> FDMXFixtureTypeToGDTFGeometryFactory::AddHead(const TSharedRef<FDMXGDTFGeometry>& OuterGeometry)
+	{
+		const TSharedRef<FDMXGDTFAxisGeometry> Head = MakeShared<FDMXGDTFAxisGeometry>(OuterGeometry);
+		OuterGeometry->AxisArray.Add(Head);
+
+		Head->Name = HeadGeometryName;
+
+		return Head;
+	}
+
+	TSharedRef<FDMXGDTFBeamGeometry> FDMXFixtureTypeToGDTFGeometryFactory::AddBeam(const TSharedRef<FDMXGDTFGeometry>& OuterGeometry)
+	{
+		const TSharedRef<FDMXGDTFBeamGeometry> Beam = MakeShared<FDMXGDTFBeamGeometry>(OuterGeometry);
+		OuterGeometry->BeamArray.Add(Beam);
+		
+		Beam->Name = BeamGeometryName;
+
+		return Beam;
+	}
+
+	void FDMXFixtureTypeToGDTFGeometryFactory::BuildMatrixGeometryReferences(const FDMXFixtureMode& Mode, const TSharedRef<FDMXGDTFGeometry>& OuterGeometry)
+	{
+		if (Mode.bFixtureMatrixEnabled && Mode.FixtureMatrixConfig.GetNumChannels() > 0)
+		{
+			// Get the byte size of an Unreal Matrix Cell
+			const int32 CellSize = [Mode]()
+			{
+				int32 OutCellSize = 0;
+				for (const FDMXFixtureCellAttribute& CellAttribute : Mode.FixtureMatrixConfig.CellAttributes)
+				{
+					OutCellSize += CellAttribute.GetNumChannels();
+				}
+
+				return OutCellSize;
+			}();
+
+			// Create Geometry References for each Unreal Matrix Cell
+			const int32 NumCells = Mode.FixtureMatrixConfig.XCells * Mode.FixtureMatrixConfig.YCells;
+			int32 Offset = 1;
+			for (int32 CellID = 0; CellID < NumCells; CellID++)
+			{
+				const int32 DMXOffset = CellID * CellSize + 1;
+				const FName GeometryName = *FString::Printf(TEXT("Layer_%i"), CellID + 1);
+
+				const TSharedRef<FDMXGDTFGeometryReference> GeometryReference = MakeShared<FDMXGDTFGeometryReference>(OuterGeometry);
+				OuterGeometry->GeometryReferenceArray.Add(GeometryReference);
+
+				GeometryReference->Name = GeometryName;
+				GeometryReference->Geometry = GetOrCreateMatrixBeamGeometryInstance()->Name;
+				GeometryReference->Model = TEXT("Layers");
+
+				const TSharedRef<FDMXGDTFGeometryBreak> Break = MakeShared<FDMXGDTFGeometryBreak>(GeometryReference);
+				Break->DMXBreak = 1;
+				Break->DMXOffset = Offset;
+				GeometryReference->BreakArray.Add(Break);
+
+				Break->DMXBreak = 1; // Unreal does not support multi universe patches, DMXBreak is always 1
+				Break->DMXOffset = DMXOffset;
+
+				Offset += CellSize;
+			}
+		}
+	}
+}
diff --git a/Engine/Plugins/VirtualProduction/DMX/DMXEngine/Source/DMXEditor/Private/Exporters/DMXFixtureTypeToGDTFGeometryFactory.h b/Engine/Plugins/VirtualProduction/DMX/DMXEngine/Source/DMXEditor/Private/Exporters/DMXFixtureTypeToGDTFGeometryFactory.h
new file mode 100644
--- /dev/null
+++ b/Engine/Plugins/VirtualProduction/DMX/DMXEngine/Source/DMXEditor/Private/Exporters/DMXFixtureTypeToGDTFGeometryFactory.h
@@ -0,0 +1,105 @@
+// Copyright Epic Games, Inc. All Rights Reserved.
+
+#pragma once
+
+#include "Templates/SharedPointer.h"
+
+struct FDMXFixtureFunction;
+struct FDMXFixtureMode;
+class UDMXEntityFixtureType;
+
+namespace UE::DMX::GDTF
+{
+	class FDMXGDTFAxisGeometry;
+	class FDMXGDTFBeamGeometry;
+	class FDMXGDTFGeometry;
+	class FDMXGDTFGeometryCollect;
+	class FDMXGDTFGeometryReference;
+
+	namespace Internal
+	{
+		struct FDMXUniqueModeGeometry
+		{
+			FDMXUniqueModeGeometry(const FDMXFixtureMode& Mode);
+
+			uint32 TotalNumChannels = 0; 
+
+			bool bIsMatrix = false;
+			bool bWithPan = false;
+			bool bWithTilt = false;
+
+			bool operator==(const FDMXUniqueModeGeometry& Other) const 
+			{ 
+				return
+					TotalNumChannels == Other.TotalNumChannels &&
+					bIsMatrix == Other.bIsMatrix &&
+					bWithPan == Other.bWithPan &&
+					bWithTilt == Other.bWithTilt;
+			}
+
+			friend FORCEINLINE uint32 GetTypeHash(const FDMXUniqueModeGeometry& UniqueMatrixDefinition)
+			{
+				return UniqueMatrixDefinition.TotalNumChannels |
+					static_cast<uint32>(UniqueMatrixDefinition.bWithPan) << 31 |
+					static_cast<uint32>(UniqueMatrixDefinition.bWithTilt << 30) |
+					static_cast<uint32>(UniqueMatrixDefinition.bWithTilt << 29);
+			}
+		};
+	}
+
+	struct FDMXFixtureFunctionWithControlledGeometry
+	{
+		FDMXFixtureFunctionWithControlledGeometry(const FDMXFixtureMode* InModePtr, const FDMXFixtureFunction* InFunctionPtr, TSharedRef<const FDMXGDTFGeometry> InControlledGeometry);
+
+		const FDMXFixtureMode* ModePtr = nullptr;
+		const FDMXFixtureFunction* FunctionPtr = nullptr;
+		const TSharedRef<const FDMXGDTFGeometry> ControlledGeometry;
+	};
+
+	/** Helper to build geometries for the Fixture Type */
+	class FDMXFixtureTypeToGDTFGeometryFactory
+		: public FNoncopyable
+	{
+		using FDMXUniqueModeGeometry = Internal::FDMXUniqueModeGeometry;
+
+	public:
+		FDMXFixtureTypeToGDTFGeometryFactory(const UDMXEntityFixtureType& InFixtureType, const TSharedRef<FDMXGDTFGeometryCollect>& InBaseGeometryCollect);
+		
+		/** Returns an array of functions with their linked geometry */
+		TArray<FDMXFixtureFunctionWithControlledGeometry> GetFunctionsWithControlledGeometry() const;
+
+		static const FName BaseGeometryName;
+		static const FName PanGeometryName;
+		static const FName HeadGeometryName;
+		static const FName BeamGeometryName;
+		static const FName MatrixBeamGeometryName;
+
+	private:
+		/** Builds all geometries */
+		void BuildGeometries();
+
+		/** Builds the geometry for a mode */
+		void BuildGeometry(const FDMXFixtureMode& Mode);
+
+		/** Gets the matrix beam instance, creates it if it doesn't exist yet */
+		TSharedRef<FDMXGDTFGeometry> GetOrCreateMatrixBeamGeometryInstance() const;
+
+		/** Adds a yoke geometry child to the outer geometry */
+		TSharedRef<FDMXGDTFAxisGeometry> AddYoke(const TSharedRef<FDMXGDTFGeometry>& OuterGeometry);
+
+		/** Adds a head geometry child to the outer geometry */
+		TSharedRef<FDMXGDTFAxisGeometry> AddHead(const TSharedRef<FDMXGDTFGeometry>& OuterGeometry);
+
+		/** Adds a beam geometry child to the outer geometry */
+		TSharedRef<FDMXGDTFBeamGeometry> AddBeam(const TSharedRef<FDMXGDTFGeometry>& OuterGeometry);
+
+		/** Builds geometry references for a matrix */
+		void BuildMatrixGeometryReferences(const FDMXFixtureMode& Mode, const TSharedRef<FDMXGDTFGeometry>& OuterGeometry);
+
+		/** A map unique matrix definitions with teir base geometry, to reuse. */
+		TMap<FDMXUniqueModeGeometry, TSharedRef<FDMXGDTFGeometry>> UniqueModeGeometryToGeometryMap;
+
+		const UDMXEntityFixtureType& FixtureType;
+		const TSharedRef<FDMXGDTFGeometryCollect>& BaseGeometryCollect;
+	};
+}
\ No newline at end of file
diff --git a/Engine/Plugins/VirtualProduction/DMX/DMXEngine/Source/DMXEditor/Private/Exporters/DMXMVRExportOptions.h b/Engine/Plugins/VirtualProduction/DMX/DMXEngine/Source/DMXEditor/Private/Exporters/DMXMVRExportOptions.h
new file mode 100644
--- /dev/null
+++ b/Engine/Plugins/VirtualProduction/DMX/DMXEngine/Source/DMXEditor/Private/Exporters/DMXMVRExportOptions.h
@@ -0,0 +1,31 @@
+// Copyright Epic Games, Inc. All Rights Reserved.
+
+#pragma once
+
+#include "UObject/Object.h"
+
+#include "DMXMVRExportOptions.generated.h"
+
+UCLASS(Config = DMXEditor)
+class UDMXMVRExportOptions
+	: public UObject
+{
+	GENERATED_BODY()
+
+public:
+	/** If checked, exports the fixtures with transforms as in the current level */
+	UPROPERTY(EditAnywhere, Category = "Level Options")
+	bool bConsiderLevel = true;
+
+	/** If checked, exports the fixtures with transforms as in the current level */
+	UPROPERTY(EditAnywhere, Meta = (EditCondition = "bConsiderLevel", HideEditConditionToggle, EditConditionHides), Category = "Level Options")
+	bool bUseTransformsFromLevel = true;
+
+	/** If checked, exports patches that are not in use in the level */
+	UPROPERTY(EditAnywhere, Meta = (EditCondition = "bConsiderLevel", HideEditConditionToggle, EditConditionHides), Category = "Level Options")
+	bool bExportPatchesNotPresentInLevel = true;
+
+	/** If set to true, the import was canceled */
+	UPROPERTY(Transient)
+	bool bCanceled = false;
+};
diff --git a/Engine/Plugins/VirtualProduction/DMX/DMXEngine/Source/DMXEditor/Private/Exporters/DMXMVRExporter.cpp b/Engine/Plugins/VirtualProduction/DMX/DMXEngine/Source/DMXEditor/Private/Exporters/DMXMVRExporter.cpp
--- a/Engine/Plugins/VirtualProduction/DMX/DMXEngine/Source/DMXEditor/Private/Exporters/DMXMVRExporter.cpp
+++ b/Engine/Plugins/VirtualProduction/DMX/DMXEngine/Source/DMXEditor/Private/Exporters/DMXMVRExporter.cpp
@@ -2,19 +2,23 @@
 
 #include "Exporters/DMXMVRExporter.h"
 
+#include "Algo/AnyOf.h"
 #include "DesktopPlatformModule.h"
 #include "DMXEditorLog.h"
 #include "DMXEditorSettings.h"
 #include "DMXFixtureTypeToGDTFConverter.h"
+#include "DMXMVRExportOptions.h"
 #include "DMXMVRXmlMergeUtility.h"
 #include "DMXZipper.h"
 #include "Editor.h"
 #include "EngineUtils.h"
 #include "Factories/DMXGDTFToFixtureTypeConverter.h"
+#include "Framework/Application/SlateApplication.h"
 #include "Framework/Notifications/NotificationManager.h"
 #include "Game/DMXComponent.h"
 #include "GenericPlatform/GenericPlatformFile.h"
 #include "IDesktopPlatform.h"
+#include "Interfaces/IMainFrameModule.h"
 #include "Library/DMXEntityFixturePatch.h"
 #include "Library/DMXEntityFixtureType.h"
 #include "Library/DMXGDTFAssetImportData.h"
@@ -28,37 +32,102 @@
 #include "MVR/Types/DMXMVRFixtureNode.h"
 #include "Subsystems/UnrealEditorSubsystem.h"
 #include "Widgets/Notifications/SNotificationList.h"
+#include "Widgets/SDMXMVRExportOptions.h"
 #include "XmlFile.h"
 
 #define LOCTEXT_NAMESPACE "DMXMVRExporter"
 
 namespace UE::DMX
 {
-	void FDMXMVRExporter::Export(UDMXLibrary* DMXLibrary, const FString& FilePathAndName, FText& OutErrorReason)
+	void FDMXMVRExporter::Export(UDMXLibrary* DMXLibrary)
 	{
 		FDMXMVRExporter Instance;
-		Instance.ExportInternal(DMXLibrary, FilePathAndName, OutErrorReason);
+		FText ErrorReason;
+		FString FilePathAndName;
+		Instance.ExportInternal(DMXLibrary, ErrorReason, FilePathAndName);
+
+		if (ErrorReason.IsEmpty())
+		{
+			FNotificationInfo NotificationInfo(FText::Format(LOCTEXT("ExportDMXLibraryAsMVRSuccessNotification", "Successfully exported MVR to {0}."), FText::FromString(FilePathAndName)));
+			NotificationInfo.ExpireDuration = 5.f;
+
+			FSlateNotificationManager::Get().AddNotification(NotificationInfo);
+		}
+		else
+		{
+			FNotificationInfo NotificationInfo(ErrorReason);
+			NotificationInfo.ExpireDuration = 10.f;
+			NotificationInfo.Image = FAppStyle::GetBrush("Icons.Warning");
+
+			FSlateNotificationManager::Get().AddNotification(NotificationInfo);
+		}
 	}
 
-	void FDMXMVRExporter::ExportInternal(UDMXLibrary* DMXLibrary, const FString& FilePathAndName, FText& OutErrorReason)
+	void FDMXMVRExporter::ExportInternal(UDMXLibrary* DMXLibrary, FText& OutErrorReason, FString& OutFilePathAndName)
 	{
-		OutErrorReason = FText::GetEmpty();
-
 		if (!ensureAlwaysMsgf(DMXLibrary, TEXT("Trying to export DMX Library '%s' as MVR file, but the DMX Library is invalid."), *DMXLibrary->GetName()))
 		{
-			OutErrorReason = FText::Format(LOCTEXT("MVRExportDMXLibraryInvalidReason", "DMX Library is invalid. Cannot export {0}."), FText::FromString(FilePathAndName));
+			OutErrorReason = LOCTEXT("MVRExportDMXLibraryInvalidReason", "DMX Library {0} is invalid. Cannot export MVR file.");
+			return;
+		}
+
+		UpdateExportOptions(*DMXLibrary);
+		if (GetDefault<UDMXMVRExportOptions>()->bCanceled)
+		{
+			OutErrorReason = LOCTEXT("MVRExportCanceledInvalidReason", "Canceled MVR export.");
 			return;
 		}
 
+		OutFilePathAndName = [DMXLibrary]()
+			{					
+				UDMXEditorSettings* DMXEditorSettings = GetMutableDefault<UDMXEditorSettings>();
+				const FString LastMVRExportPath = DMXEditorSettings->LastMVRExportPath;
+				const FString DefaultPath = FPaths::DirectoryExists(LastMVRExportPath) ? LastMVRExportPath : FPaths::ProjectSavedDir();
+
+				IDesktopPlatform* DesktopPlatform = FDesktopPlatformModule::Get();
+				if (!DesktopPlatform)
+				{
+					return FString();
+				}
+
+				TArray<FString> SaveFilenames;
+				const bool bSaveFile = DesktopPlatform->SaveFileDialog(
+					FSlateApplication::Get().FindBestParentWindowHandleForDialogs(nullptr),
+					LOCTEXT("ExportMVR", "Export MVR").ToString(),
+					DefaultPath,
+					DMXLibrary->GetName() + TEXT(".mvr"),
+					TEXT("My Virtual Rig (*.mvr)|*.mvr"),
+					EFileDialogFlags::None,
+					SaveFilenames);
+
+				if (!bSaveFile || SaveFilenames.IsEmpty())
+				{
+					return FString();
+				}
+
+				DMXEditorSettings->LastMVRExportPath = FPaths::GetPath(SaveFilenames[0]);
+				DMXEditorSettings->SaveConfig();
+
+				return SaveFilenames[0];
+			}();
+
 		// Create a copy of the library's general scene description, so transforms exported aren't written to the dmx library
-		UDMXMVRGeneralSceneDescription* TempGeneralSceneDescription = UDMXMVRGeneralSceneDescription::CreateFromDMXLibrary(*DMXLibrary, GetTransientPackage(), NAME_None);
+		DMXLibrary->UpdateGeneralSceneDescription();
+		UDMXMVRGeneralSceneDescription* TempGeneralSceneDescription = DuplicateObject<UDMXMVRGeneralSceneDescription>(DMXLibrary->GetLazyGeneralSceneDescription(), GetTransientPackage());
 		if (!ensureAlwaysMsgf(TempGeneralSceneDescription, TEXT("Trying to export DMX Library '%s' as MVR file, but its General Scene Description is invalid."), *DMXLibrary->GetName()))
 		{
-			OutErrorReason = FText::Format(LOCTEXT("MVRExportGeneralSceneDescriptionInvalidReason", "DMX Library is invalid. Cannot export {0}."), FText::FromString(FilePathAndName));
+			OutErrorReason = FText::Format(LOCTEXT("MVRExportGeneralSceneDescriptionInvalidReason", "DMX Library is invalid. Cannot export {0}."), FText::FromString(OutFilePathAndName));
 			return;
 		}
 
-		WriteMVRFixtureTransformsFromLevel(TempGeneralSceneDescription);
+		const UDMXMVRExportOptions* ExportOptions = GetDefault<UDMXMVRExportOptions>();
+		if (ExportOptions->bConsiderLevel)
+		{
+			const TMap<const UDMXComponent*, const AActor*> DMXComponentToActorMap = GetDMXComponentToActorMap();
+
+			TryRemoveMVRFixturesNotPresentInLevel(TempGeneralSceneDescription, DMXLibrary, DMXComponentToActorMap);
+			TryUpdateFixtureNodesFromLevel(TempGeneralSceneDescription, DMXComponentToActorMap, ExportOptions->bUseTransformsFromLevel);
+		}
 
 		const TSharedRef<FDMXZipper> Zip = MakeShared<FDMXZipper>();
 		if (!ZipGeneralSceneDescription(Zip, TempGeneralSceneDescription, OutErrorReason))
@@ -68,18 +137,106 @@ namespace UE::DMX
 
 		if (!ZipGDTFs(Zip, DMXLibrary))
 		{
-			OutErrorReason = FText::Format(LOCTEXT("MVRExportZipGDTFsFailedReason", "Some Fixture Types could not be converted to GDTF. Exported MVR to {0}."), FText::FromString(FilePathAndName));
+			OutErrorReason = FText::Format(LOCTEXT("MVRExportZipGDTFsFailedReason", "Some Fixture Types could not be converted to GDTF. Exported MVR to {0}."), FText::FromString(OutFilePathAndName));
 			// Allow continuation of export
 		}
 
 		ZipThirdPartyData(Zip, TempGeneralSceneDescription);
-		if (!Zip->SaveToFile(FilePathAndName))
+		if (!Zip->SaveToFile(OutFilePathAndName))
 		{
-			OutErrorReason = FText::Format(LOCTEXT("MVRExportWriteZipFailedReason", "File is not writable or locked by another process. Cannot export {0}."), FText::FromString(FilePathAndName));
+			OutErrorReason = FText::Format(LOCTEXT("MVRExportWriteZipFailedReason", "File is not writable or locked by another process. Cannot export {0}."), FText::FromString(OutFilePathAndName));
 			return;
 		}
 	}
 
+	void FDMXMVRExporter::UpdateExportOptions(const UDMXLibrary& DMXLibrary) const
+	{
+		UUnrealEditorSubsystem* UnrealEditorSubsystem = GEditor->GetEditorSubsystem<UUnrealEditorSubsystem>();
+		const bool bHasEditorWorld = UnrealEditorSubsystem && UnrealEditorSubsystem->GetEditorWorld();
+
+		// No need to show options when there is no world, as of now the options are related to the current level only
+		if (!bHasEditorWorld)
+		{
+			return;
+		}
+
+		TSharedPtr<SWindow> ParentWindow;
+		if (FModuleManager::Get().IsModuleLoaded("MainFrame"))
+		{
+			IMainFrameModule& MainFrame = FModuleManager::LoadModuleChecked<IMainFrameModule>("MainFrame");
+			ParentWindow = MainFrame.GetParentWindow();
+		}
+
+		const FVector2D WindowSize = FVector2D(512.f, 338.f);
+
+		const FSlateRect WorkAreaRect = FSlateApplicationBase::Get().GetPreferredWorkArea();
+		const FVector2D DisplayTopLeft(WorkAreaRect.Left, WorkAreaRect.Top);
+		const FVector2D DisplaySize(WorkAreaRect.Right - WorkAreaRect.Left, WorkAreaRect.Bottom - WorkAreaRect.Top);
+
+		const FVector2D WindowPosition = (DisplayTopLeft + (DisplaySize - WindowSize) / 2.0f);
+
+		const FText Caption = LOCTEXT("ExportWindowCaption", "MVR Export Options");
+
+		const TSharedRef<SWindow> Window = SNew(SWindow)
+			.Title(Caption)
+			.SizingRule(ESizingRule::FixedSize)
+			.AutoCenter(EAutoCenter::None)
+			.ClientSize(WindowSize)
+			.ScreenPosition(WindowPosition);
+
+		Window->SetContent
+		(
+			SNew(SDMXMVRExportOptions, Window)
+		);
+
+		constexpr bool bSlowTaskWindow = false;
+		FSlateApplication::Get().AddModalWindow(Window, ParentWindow, bSlowTaskWindow);
+	}
+
+	TMap<const UDMXComponent*, const AActor*> FDMXMVRExporter::GetDMXComponentToActorMap() const
+	{
+		TMap<const UDMXComponent*, const AActor*> DMXComponentToActorMap;
+
+		UUnrealEditorSubsystem* UnrealEditorSubsystem = GEditor->GetEditorSubsystem<UUnrealEditorSubsystem>();
+		const UWorld* EditorWorld = UnrealEditorSubsystem ? UnrealEditorSubsystem->GetEditorWorld() : nullptr;
+		if (!EditorWorld)
+		{
+			return DMXComponentToActorMap;
+		}
+
+		// Find actors with a patched DMX component
+		for (TActorIterator<AActor> It(EditorWorld, AActor::StaticClass()); It; ++It)
+		{
+			const AActor* Actor = *It;
+			if (!Actor)
+			{
+				continue;
+			}
+
+			const TSet<UActorComponent*> Components = Actor->GetComponents();
+
+			TArray<const UDMXComponent*> DMXComponents;
+			Algo::TransformIf(Components, DMXComponents,
+				[](const UActorComponent* Component)
+				{
+					return Component &&
+						Component->IsA(UDMXComponent::StaticClass()) &&
+						CastChecked<UDMXComponent>(Component)->GetFixturePatch();
+				},
+				[](const UActorComponent* Component)
+				{
+					return CastChecked<UDMXComponent>(Component);
+				});
+
+			for (const UDMXComponent* DMXComponent : DMXComponents)
+			{
+				DMXComponentToActorMap.Add(TTuple<const UDMXComponent*, const AActor*>(DMXComponent, Actor));
+			}
+		}
+
+		return DMXComponentToActorMap;
+	}
+
 	bool FDMXMVRExporter::ZipGeneralSceneDescription(const TSharedRef<FDMXZipper>& Zip, const UDMXMVRGeneralSceneDescription* GeneralSceneDescription, FText& OutErrorReason)
 	{
 		if (!GeneralSceneDescription->CanCreateXmlFile(OutErrorReason))
@@ -149,55 +306,86 @@ namespace UE::DMX
 		bool bAllZippedSuccessfully = true;
 		for (UDMXEntityFixtureType* FixtureType : FixtureTypesToExport)
 		{
-			constexpr bool bWithExtension = true;
-			const FString CleanGDTFFilename = FixtureType->GetCleanGDTFFileNameSynchronous(bWithExtension);
-
-			using namespace UE::DMX::GDTF;
-			const TSharedPtr<FXmlFile> XmlFile = FDMXFixtureTypeToGTDFConverter::Convert(FixtureType);
-			if (!XmlFile.IsValid())
+			if (FixtureType->GDTFSource.IsNull() || FixtureType->bExportGeneratedGDTF)
 			{
-				UE_LOG(LogDMXEditor, Warning, TEXT("Failed to create GDTF from Fixture Type '%s'. See previous errors."), *FixtureType->Name);
-				bAllZippedSuccessfully &= false;
-				continue;
-			}
+				constexpr bool bWithExtension = true;
+				const FString CleanGDTFFilename = FixtureType->GetCleanGDTFFileNameSynchronous(bWithExtension);
 
-			if (!FPaths::DirectoryExists(TempGDTFDir))
-			{
-				IPlatformFile::GetPlatformPhysical().CreateDirectory(*TempGDTFDir);
-			}
+				using namespace UE::DMX::GDTF;
+				const TSharedPtr<FXmlFile> XmlFile = FDMXFixtureTypeToGDTFConverter::Convert(FixtureType);
+				if (!XmlFile.IsValid())
+				{
+					UE_LOG(LogDMXEditor, Warning, TEXT("Failed to create GDTF from Fixture Type '%s'. See previous errors."), *FixtureType->Name);
+					bAllZippedSuccessfully = false;
+					continue;
+				}
 
-			if (!ensureMsgf(FPaths::DirectoryExists(TempGDTFDir), TEXT("Cannot create temporary directory for GDTFs.")))
-			{
-				continue;
-			}
+				if (!FPaths::DirectoryExists(TempGDTFDir))
+				{
+					IPlatformFile::GetPlatformPhysical().CreateDirectory(*TempGDTFDir);
+				}
 
-			const FString TempXmlFilename = TempGDTFDir / FixtureType->Name + TEXT(".description.xml");
-			XmlFile->Save(TempXmlFilename);
+				if (!ensureMsgf(FPaths::DirectoryExists(TempGDTFDir), TEXT("Cannot create temporary directory for GDTFs.")))
+				{
+					continue;
+				}
 
-			TArray64<uint8> DescriptionXmlData;
-			if (FFileHelper::LoadFileToArray(DescriptionXmlData, *TempXmlFilename))
-			{
-				FDMXZipper NewGDTFZip;
-				NewGDTFZip.AddFile(TEXT("description.xml"), DescriptionXmlData);
-				const FString TempGDTFilePathAndName = TempGDTFDir / CleanGDTFFilename;
+				const FString TempXmlFilename = TempGDTFDir / FixtureType->Name + TEXT(".description.xml");
+				XmlFile->Save(TempXmlFilename);
 
-				TArray64<uint8> GDTFData;
-				if (NewGDTFZip.GetData(GDTFData))
+				TArray64<uint8> DescriptionXmlData;
+				if (FFileHelper::LoadFileToArray(DescriptionXmlData, *TempXmlFilename))
 				{
-					Zip->AddFile(CleanGDTFFilename, GDTFData);
+					FDMXZipper NewGDTFZip;
+					NewGDTFZip.AddFile(TEXT("description.xml"), DescriptionXmlData);
+					const FString TempGDTFilePathAndName = TempGDTFDir / CleanGDTFFilename;
+
+					TArray64<uint8> GDTFData;
+					if (NewGDTFZip.GetData(GDTFData))
+					{
+						Zip->AddFile(CleanGDTFFilename, GDTFData);
+					}
+					else
+					{
+						UE_LOG(LogDMXEditor, Warning, TEXT("Failed to load temporary GDTF for Fixture Type '%s'."), *FixtureType->Name);
+						bAllZippedSuccessfully = false;
+						continue;
+					}
 				}
 				else
 				{
-					UE_LOG(LogDMXEditor, Warning, TEXT("Failed to load temporary GDTF for Fixture Type '%s'."), *FixtureType->Name);
-					bAllZippedSuccessfully &= false;
+					UE_LOG(LogDMXEditor, Warning, TEXT("Failed to load temporary description.xml for Fixture Type '%s'."), *FixtureType->Name);
+					bAllZippedSuccessfully = false;
 					continue;
 				}
 			}
 			else
 			{
-				UE_LOG(LogDMXEditor, Warning, TEXT("Failed to load temporary description.xml for Fixture Type '%s'."), *FixtureType->Name);
-				bAllZippedSuccessfully &= false;
-				continue;
+				UDMXImportGDTF* DMXImportGDTF = FixtureType->GDTFSource.LoadSynchronous();
+				if (!DMXImportGDTF)
+				{
+					UE_LOG(LogDMXEditor, Warning, TEXT("Cannot export Fixture Type '%s' to MVR, but the Fixture Type has a DMX Import Type which is not GDTF."), *FixtureType->Name);
+					continue;
+				}
+
+				UDMXGDTFAssetImportData* GDTFAssetImportData = DMXImportGDTF->GetGDTFAssetImportData();
+				if (!ensureAlwaysMsgf(GDTFAssetImportData, TEXT("Missing default GDTF Asset Import Data subobject in GDTF.")))
+				{
+					continue;
+				}
+
+				const TArray64<uint8>& RawSourceData = RefreshSourceDataAndFixtureType(*FixtureType, *GDTFAssetImportData);
+				if (!RawSourceData.IsEmpty())
+				{
+					const FString GDTFFilename = FPaths::GetCleanFilename(GDTFAssetImportData->GetFilePathAndName());
+					Zip->AddFile(GDTFFilename, GDTFAssetImportData->GetRawSourceData());
+				}
+				else
+				{
+					bAllZippedSuccessfully = false;
+					const UDMXImportGDTF* GDTF = Cast<UDMXImportGDTF>(GDTFAssetImportData->GetOuter());
+					UE_LOG(LogDMXEditor, Error, TEXT("Cannot export '%s' to MVR File. The asset is missing source data."), GDTF ? *GDTF->GetName() : TEXT("Invalid GDTF Asset"));
+				}
 			}
 		}
 
@@ -239,77 +427,75 @@ namespace UE::DMX
 		}
 	}
 
-	void FDMXMVRExporter::WriteMVRFixtureTransformsFromLevel(UDMXMVRGeneralSceneDescription* GeneralSceneDescription)
+	void FDMXMVRExporter::TryRemoveMVRFixturesNotPresentInLevel(UDMXMVRGeneralSceneDescription* GeneralSceneDescription, UDMXLibrary* DMXLibrary, const TMap<const UDMXComponent*, const AActor*>& DMXComponentToActorMap)
 	{
-		if (!ensureMsgf(GeneralSceneDescription, TEXT("Trying to write MVR Fixture Transforms from Level to MVR General Scene Description, but General Scene Description is invalid")))
+		if (!ensureMsgf(GeneralSceneDescription && DMXLibrary, TEXT("Trying to write MVR Fixture Transforms from Level to MVR General Scene Description, but General Scene Description or DMX Library are invalid")))
 		{
 			return;
 		}
 
-		UUnrealEditorSubsystem* UnrealEditorSubsystem = GEditor->GetEditorSubsystem<UUnrealEditorSubsystem>();
-		if (!UnrealEditorSubsystem)
+		const UDMXMVRExportOptions* ExportOptions = GetDefault<UDMXMVRExportOptions>();
+		if (ExportOptions->bExportPatchesNotPresentInLevel)
 		{
 			return;
 		}
 
-		const UWorld* EditorWorld = UnrealEditorSubsystem->GetEditorWorld();
-		if (!EditorWorld)
-		{
-			return;
-		}
+		const TArray<UDMXEntityFixturePatch*> FixturePatchesInLibrary = DMXLibrary->GetEntitiesTypeCast<UDMXEntityFixturePatch>();
 
-		// Find actors with a patched DMX component
-		TMap<const UDMXEntityFixturePatch*, FTransform> PatchToTransformMap;
-		for (TActorIterator<AActor> It(EditorWorld, AActor::StaticClass()); It; ++It)
-		{
-			const AActor* Actor = *It;
-			if (!Actor)
+		TArray<const UDMXEntityFixturePatch*> FixturePatchesToRemove;
+		Algo::TransformIf(FixturePatchesInLibrary, FixturePatchesToRemove,
+			[&DMXComponentToActorMap](const UDMXEntityFixturePatch* FixturePatch)
 			{
-				continue;
-			}
-
-			const TSet<UActorComponent*> Components = Actor->GetComponents();
-
-			TArray<const UDMXEntityFixturePatch*> FixturePatches;
-			Algo::TransformIf(Components, FixturePatches, 
-				[](const UActorComponent* Component)
-				{
-					return Component && 
-						Component->IsA(UDMXComponent::StaticClass()) &&
-						CastChecked<UDMXComponent>(Component)->GetFixturePatch();
-				},
-				[](const UActorComponent* Component)
-				{
-					return CastChecked<UDMXComponent>(Component)->GetFixturePatch();
-				});
-
-			for (const UDMXEntityFixturePatch* FixturePatch : FixturePatches)
+				const bool bPatchExistsInLevel = Algo::AnyOf(DMXComponentToActorMap, [FixturePatch](const TTuple<const UDMXComponent*, const AActor*>& DMXComponentToActorPair)
+					{
+						const UDMXComponent* DMXComponent = DMXComponentToActorPair.Key;
+						return
+							DMXComponent &&
+							FixturePatch &&
+							DMXComponent->GetFixturePatch() == FixturePatch;
+					});
+
+				return !bPatchExistsInLevel;
+			},
+			[](const UDMXEntityFixturePatch* FixturePatch)
 			{
-				PatchToTransformMap.Add(FixturePatch, Actor->GetTransform());
-			}
+				return FixturePatch;
+			});
+
+		for (const UDMXEntityFixturePatch* FixturePatch : FixturePatchesToRemove)
+		{
+			GeneralSceneDescription->RemoveFixtureNode(FixturePatch->GetMVRFixtureUUID());
 		}
+	}
+
+	void FDMXMVRExporter::TryUpdateFixtureNodesFromLevel(UDMXMVRGeneralSceneDescription* GeneralSceneDescription, const TMap<const UDMXComponent*, const AActor*>& DMXComponentToActorMap, bool bUseTransformsFromLevel)
+	{
+		if (!ensureMsgf(GeneralSceneDescription, TEXT("Trying to write MVR Fixture Transforms from Level to MVR General Scene Description, but General Scene Description are invalid")))
+		{
+			return;
+		}		
+		
+		const UDMXMVRExportOptions* ExportOptions = GetDefault<UDMXMVRExportOptions>();
 
 		TArray<UDMXMVRFixtureNode*> FixtureNodes;
 		GeneralSceneDescription->GetFixtureNodes(FixtureNodes);
 
-		TArray<UDMXMVRFixtureNode*> HandledMVRFixtureNodes;
-		for (const TTuple<const UDMXEntityFixturePatch*, FTransform>& PatchToTransformPair : PatchToTransformMap)
+		TSet<FGuid> UsedMVRUUIDs;
+		for (const TTuple<const UDMXComponent*, const AActor*>& DMXComponentToActorPair : DMXComponentToActorMap)
 		{
-			UDMXMVRFixtureNode** FixtureNodePtr = Algo::FindByPredicate(FixtureNodes, [PatchToTransformPair](UDMXMVRFixtureNode* MVRFixtureNode)
-				{
-					return MVRFixtureNode->UUID == PatchToTransformPair.Key->GetMVRFixtureUUID();
-				});
-
-			if (FixtureNodePtr)
+			const UDMXEntityFixturePatch* FixturePatch = DMXComponentToActorPair.Key->GetFixturePatch();
+			const AActor* Actor = DMXComponentToActorPair.Value;
+			if (!FixturePatch || !Actor || UsedMVRUUIDs.Contains(FixturePatch->GetMVRFixtureUUID()))
 			{
-				if (HandledMVRFixtureNodes.Contains(*FixtureNodePtr))
-				{
-					UE_LOG(LogDMXEditor, Warning, TEXT("Ambigous Fixture ID for Fixture Patch '%s' when exporting MVR. Each Fixture Patch can only be used once in an MVR scene. Using transform of arbitrary instance with ambigous ID."), *PatchToTransformPair.Key->Name);
-					continue;
-				}
+				continue;
+			}
 
-				(*FixtureNodePtr)->SetTransformAbsolute(PatchToTransformPair.Value);
-				HandledMVRFixtureNodes.Add(*FixtureNodePtr);
+			UDMXMVRFixtureNode* FixtureNode = GeneralSceneDescription->WriteFixturePatch(*FixturePatch);
+			UsedMVRUUIDs.Add(FixtureNode->UUID);
+
+			if (FixtureNode && ExportOptions->bUseTransformsFromLevel)
+			{
+				FixtureNode->SetTransformAbsolute(Actor->GetTransform());
 			}
 		}
 	}
diff --git a/Engine/Plugins/VirtualProduction/DMX/DMXEngine/Source/DMXEditor/Private/Exporters/DMXMVRExporter.h b/Engine/Plugins/VirtualProduction/DMX/DMXEngine/Source/DMXEditor/Private/Exporters/DMXMVRExporter.h
--- a/Engine/Plugins/VirtualProduction/DMX/DMXEngine/Source/DMXEditor/Private/Exporters/DMXMVRExporter.h
+++ b/Engine/Plugins/VirtualProduction/DMX/DMXEngine/Source/DMXEditor/Private/Exporters/DMXMVRExporter.h
@@ -2,19 +2,19 @@
 
 #pragma once
 
-#if UE_ENABLE_INCLUDE_ORDER_DEPRECATED_IN_5_2
-#include "CoreMinimal.h"
-#endif
-#include "Containers/Array.h"
 #include "Templates/SharedPointer.h"
 
+class AActor;
 class FDMXZipper;
 class FXmlFile;
 class UDMXEntityFixtureType;
 class UDMXGDTFAssetImportData;
+class UDMXComponent;
 class UDMXLibrary;
 class UDMXMVRAssetImportData;
+class UDMXMVRFixtureNode;
 class UDMXMVRGeneralSceneDescription;
+class UWorld;
 
 namespace UE::DMX
 {
@@ -23,11 +23,17 @@ namespace UE::DMX
 	{
 	public:
 		/** Exports the DMX Library as MVR File */
-		static void Export(UDMXLibrary* DMXLibrary, const FString& FilePathAndName, FText& OutErrorReason);
+		static void Export(UDMXLibrary* DMXLibrary);
 
 	private:
 		/** Exports the DMX Library as MVR File. If OutErrorReason is not empty there were issues with the export. */
-		void ExportInternal(UDMXLibrary* DMXLibrary, const FString& FilePathAndName, FText& OutErrorReason);
+		void ExportInternal(UDMXLibrary* DMXLibrary, FText& OutErrorReason, FString& OutFilePathAndName);
+
+		/** Updates the MVR export options. Returns false if the import was canceled */
+		void UpdateExportOptions(const UDMXLibrary& DMXLibrary) const;
+		
+		/** Builds the DMX component to actor map */
+		TMap<const UDMXComponent*, const AActor*> GetDMXComponentToActorMap() const;
 
 		/** Zips the GeneralSceneDescription.xml */
 		[[nodiscard]] bool ZipGeneralSceneDescription(const TSharedRef<FDMXZipper>& Zip, const UDMXMVRGeneralSceneDescription* GeneralSceneDescription, FText& OutErrorReason);
@@ -41,8 +47,11 @@ namespace UE::DMX
 		/** Creates an General Scene Description Xml File from the MVR Source, as it was imported */
 		const TSharedPtr<FXmlFile> CreateSourceGeneralSceneDescriptionXmlFile(const UDMXMVRGeneralSceneDescription* GeneralSceneDescription) const;
 
-		/** Writes transforms from level to the MVR Fixture Actors where applicable */
-		void WriteMVRFixtureTransformsFromLevel(UDMXMVRGeneralSceneDescription* GeneralSceneDescription);
+		/** Tries to remove MVR fixtures that are not present in the level, depending on export options. */
+		void TryRemoveMVRFixturesNotPresentInLevel(UDMXMVRGeneralSceneDescription* GeneralSceneDescription, UDMXLibrary* DMXLibrary, const TMap<const UDMXComponent*, const AActor*>& DMXComponentToActorMap);
+
+		/** Tries to create MVR fixtures from patches used more than once in the level and apply transforms depending on export options. */
+		void TryUpdateFixtureNodesFromLevel(UDMXMVRGeneralSceneDescription* GeneralSceneDescription, const TMap<const UDMXComponent*, const AActor*>& DMXComponentToActorMap, bool bUseTransformsFromLevel);
 
 		/**
 		 * Gets raw source data or creates (possibly empty) source data where the source data is not present.
diff --git a/Engine/Plugins/VirtualProduction/DMX/DMXEngine/Source/DMXEditor/Private/Factories/DMXGDTFFactory.cpp b/Engine/Plugins/VirtualProduction/DMX/DMXEngine/Source/DMXEditor/Private/Factories/DMXGDTFFactory.cpp
--- a/Engine/Plugins/VirtualProduction/DMX/DMXEngine/Source/DMXEditor/Private/Factories/DMXGDTFFactory.cpp
+++ b/Engine/Plugins/VirtualProduction/DMX/DMXEngine/Source/DMXEditor/Private/Factories/DMXGDTFFactory.cpp
@@ -3,6 +3,7 @@
 #include "Factories/DMXGDTFFactory.h"
 
 #include "Application/SlateApplicationBase.h"
+#include "AssetRegistry/AssetRegistryModule.h"
 #include "DMXEditorLog.h"
 #include "Editor.h"
 #include "Editor/EditorEngine.h"
diff --git a/Engine/Plugins/VirtualProduction/DMX/DMXEngine/Source/DMXEditor/Private/Factories/DMXGDTFImporter.cpp b/Engine/Plugins/VirtualProduction/DMX/DMXEngine/Source/DMXEditor/Private/Factories/DMXGDTFImporter.cpp
--- a/Engine/Plugins/VirtualProduction/DMX/DMXEngine/Source/DMXEditor/Private/Factories/DMXGDTFImporter.cpp
+++ b/Engine/Plugins/VirtualProduction/DMX/DMXEngine/Source/DMXEditor/Private/Factories/DMXGDTFImporter.cpp
@@ -2,15 +2,12 @@
 
 #include "Factories/DMXGDTFImporter.h"
 
-#include "DMXEditorLog.h"
 #include "DMXZipper.h"
-#include "EditorFramework/AssetImportData.h"
+#include "Editor.h"
 #include "Factories/DMXGDTFFactory.h"
-#include "GDTF/DMXGDTFDescription.h"
 #include "Library/DMXGDTFAssetImportData.h"
 #include "Library/DMXImportGDTF.h"
 #include "Misc/Paths.h"
-#include "XmlFile.h"
 
 #define LOCTEXT_NAMESPACE "DMXGDTFImporter"
 
@@ -61,7 +58,10 @@ namespace UE::DMX
 		}
 
 		// Create GDTF
-		const FName UniqueName = MakeUniqueObjectName(ImportArgs.Parent.Get(), UDMXImportGDTF::StaticClass(), ImportArgs.Name);
+		const FName UniqueName = IsUniqueObjectName(ImportArgs.Name, ImportArgs.Parent.Get()) ?
+			ImportArgs.Name :
+			MakeUniqueObjectName(ImportArgs.Parent.Get(), UDMXImportGDTF::StaticClass(), ImportArgs.Name);
+
 		UDMXImportGDTF* NewGDTF = NewObject<UDMXImportGDTF>(ImportArgs.Parent.Get(), UniqueName, ImportArgs.Flags | RF_Public);
 
 		// Set Asset Import Data
diff --git a/Engine/Plugins/VirtualProduction/DMX/DMXEngine/Source/DMXEditor/Private/Factories/DMXGDTFToFixtureTypeConverter.cpp b/Engine/Plugins/VirtualProduction/DMX/DMXEngine/Source/DMXEditor/Private/Factories/DMXGDTFToFixtureTypeConverter.cpp
--- a/Engine/Plugins/VirtualProduction/DMX/DMXEngine/Source/DMXEditor/Private/Factories/DMXGDTFToFixtureTypeConverter.cpp
+++ b/Engine/Plugins/VirtualProduction/DMX/DMXEngine/Source/DMXEditor/Private/Factories/DMXGDTFToFixtureTypeConverter.cpp
@@ -183,7 +183,7 @@ namespace UE::DMX::GDTF
 			if (ChannelFunction.IsValid() &&
 				ChannelFunction->Default.IsSet())
 			{
-				return static_cast<int64>(ChannelFunction->Default.AsIntChecked());
+				return static_cast<int64>(ChannelFunction->Default.GetChecked(DMXChannelNode));
 			}
 
 			return 0;
@@ -333,6 +333,7 @@ namespace UE::DMX::GDTF
 			InOutFixtureType.Name = GDTFFixtureType->Name.ToString();
 		}
 
+		InOutFixtureType.Modes.Reset();
 		for (const TSharedPtr<FDMXGDTFDMXMode>& GDTFDMXMode : GDTFFixtureType->DMXModes)
 		{
 			FDMXFixtureMode Mode;
@@ -364,6 +365,7 @@ namespace UE::DMX::GDTF
 		FDMXFixtureMode Mode;
 		Mode.ModeName = InDMXModeNode->Name.ToString();
 
+		TMap<FName, int32> AttributeNameToCountMap;
 		for (const TSharedPtr<FDMXGDTFDMXChannel>& DMXChannelNode : InDMXModeNode->DMXChannels)
 		{
 			FDMXGDTFChannelInterpreter::FChannelProperties ChannelProperties;
@@ -395,15 +397,22 @@ namespace UE::DMX::GDTF
 			}
 			else
 			{
+				const int32* AttributeCountPtr = AttributeNameToCountMap.Find(*ChannelProperties.AttributeName);
+				const FString Attribute = AttributeCountPtr ? 
+					ChannelProperties.AttributeName + TEXT("_") + FString::FromInt(*AttributeCountPtr) :
+					ChannelProperties.AttributeName;
+
 				FDMXFixtureFunction Function;
 				Function.FunctionName = ChannelProperties.AttributeName;
-				Function.Attribute = FDMXAttributeName(*ChannelProperties.AttributeName);
+				Function.Attribute = FDMXAttributeName(*Attribute);
 				Function.Channel = ChannelProperties.FirstChannel;
 				Function.bUseLSBMode = ChannelProperties.bLSBMode;
 				Function.DataType = ChannelProperties.SignalFormat;
 				Function.DefaultValue = ChannelProperties.DefaultValue;
 
 				Mode.Functions.Add(Function);
+
+				AttributeNameToCountMap.FindOrAdd(*ChannelProperties.AttributeName, 0)++;
 			}
 		}
 
diff --git a/Engine/Plugins/VirtualProduction/DMX/DMXEngine/Source/DMXEditor/Private/Factories/DMXLibraryFromMVRImporter.cpp b/Engine/Plugins/VirtualProduction/DMX/DMXEngine/Source/DMXEditor/Private/Factories/DMXLibraryFromMVRImporter.cpp
--- a/Engine/Plugins/VirtualProduction/DMX/DMXEngine/Source/DMXEditor/Private/Factories/DMXLibraryFromMVRImporter.cpp
+++ b/Engine/Plugins/VirtualProduction/DMX/DMXEngine/Source/DMXEditor/Private/Factories/DMXLibraryFromMVRImporter.cpp
@@ -123,7 +123,7 @@ bool FDMXLibraryFromMVRImporter::LoadMVRFile(const FString& InFilename)
 	}
 
 	// Create the General Scene Description
-	GeneralSceneDescription = CreateGeneralSsceneDescription(GetTransientPackage());
+	GeneralSceneDescription = CreateGeneralSceneDescription(GetTransientPackage());
 	if (!GeneralSceneDescription)
 	{
 		UE_LOG(LogDMXEditor, Error, TEXT("Cannot read General Scene Description from MVR '%s'. File is not a valid MVR."), *Filename);
@@ -352,7 +352,7 @@ void FDMXLibraryFromMVRImporter::AddReferencedObjects(FReferenceCollector& Colle
 	Collector.AddReferencedObject(GeneralSceneDescription);
 }
 
-UDMXMVRGeneralSceneDescription* FDMXLibraryFromMVRImporter::CreateGeneralSsceneDescription(UObject* Outer) const
+UDMXMVRGeneralSceneDescription* FDMXLibraryFromMVRImporter::CreateGeneralSceneDescription(UObject* Outer) const
 {
 	if (!Zip.IsValid())
 	{
diff --git a/Engine/Plugins/VirtualProduction/DMX/DMXEngine/Source/DMXEditor/Private/Factories/DMXLibraryFromMVRImporter.h b/Engine/Plugins/VirtualProduction/DMX/DMXEngine/Source/DMXEditor/Private/Factories/DMXLibraryFromMVRImporter.h
--- a/Engine/Plugins/VirtualProduction/DMX/DMXEngine/Source/DMXEditor/Private/Factories/DMXLibraryFromMVRImporter.h
+++ b/Engine/Plugins/VirtualProduction/DMX/DMXEngine/Source/DMXEditor/Private/Factories/DMXLibraryFromMVRImporter.h
@@ -73,7 +73,7 @@ protected:
 
 private:
 	/** Creates a General Scene Description from the MVR zip. Returns the General Scene Description or nullptr if not a valid MVR zip. */
-	UDMXMVRGeneralSceneDescription* CreateGeneralSsceneDescription(UObject* Outer) const;
+	UDMXMVRGeneralSceneDescription* CreateGeneralSceneDescription(UObject* Outer) const;
 
 	/** Creates GDTF assets from the MVR */
 	TArray<UDMXImportGDTF*> CreateGDTFs(UObject* InParent, EObjectFlags InFlags, bool bReimportExisting);
diff --git a/Engine/Plugins/VirtualProduction/DMX/DMXEngine/Source/DMXEditor/Private/Widgets/SDMXMVRExportOptions.cpp b/Engine/Plugins/VirtualProduction/DMX/DMXEngine/Source/DMXEditor/Private/Widgets/SDMXMVRExportOptions.cpp
new file mode 100644
--- /dev/null
+++ b/Engine/Plugins/VirtualProduction/DMX/DMXEngine/Source/DMXEditor/Private/Widgets/SDMXMVRExportOptions.cpp
@@ -0,0 +1,90 @@
+// Copyright Epic Games, Inc. All Rights Reserved.
+
+#include "SDMXMVRExportOptions.h"
+
+#include "Exporters/DMXMVRExportOptions.h"
+#include "Modules/ModuleManager.h"
+#include "PropertyEditorModule.h"
+#include "SPrimaryButton.h"
+
+#define LOCTEXT_NAMESPACE "SDMXMVRExportOptions"
+
+namespace UE::DMX
+{
+	SDMXMVRExportOptions::~SDMXMVRExportOptions()
+	{
+		UDMXMVRExportOptions* ExportOptions = GetMutableDefault<UDMXMVRExportOptions>();
+		ExportOptions->SaveConfig();
+	}
+
+	void SDMXMVRExportOptions::Construct(const FArguments& InArgs, const TSharedRef<SWindow>& ParentWindow)
+	{		
+		// Build an options details view
+		FPropertyEditorModule& PropertyEditorModule = FModuleManager::GetModuleChecked<FPropertyEditorModule>("PropertyEditor");
+		FDetailsViewArgs DetailsViewArgs;
+		DetailsViewArgs.bAllowSearch = false;
+		DetailsViewArgs.NameAreaSettings = FDetailsViewArgs::HideNameArea;
+		const TSharedRef<IDetailsView> DetailsView = PropertyEditorModule.CreateDetailView(DetailsViewArgs);
+
+		DetailsView->SetObject(GetMutableDefault<UDMXMVRExportOptions>());
+
+		ChildSlot
+		[
+			SNew(SVerticalBox)
+	
+			+ SVerticalBox::Slot()
+			.HAlign(HAlign_Fill)
+			.VAlign(VAlign_Fill)
+			.FillHeight(1.f)
+			[
+				DetailsView
+			]
+				
+			+ SVerticalBox::Slot()
+			.HAlign(HAlign_Right)
+			.VAlign(VAlign_Bottom)
+			.AutoHeight()
+			[
+				SNew(SHorizontalBox)
+					
+				+ SHorizontalBox::Slot()
+				.HAlign(HAlign_Right)
+				.VAlign(VAlign_Center)
+				.AutoWidth()
+				.Padding(4.f, 2.f)
+				[
+					SNew(SPrimaryButton)
+					.Text(LOCTEXT("OptionWindow_Export", "Export"))
+					.OnClicked_Lambda([ParentWindow]()
+						{
+							UDMXMVRExportOptions* ExportOptions = GetMutableDefault<UDMXMVRExportOptions>();
+							ExportOptions->bCanceled = false;
+
+							ParentWindow->RequestDestroyWindow();
+							return FReply::Handled();
+						})
+				]
+
+				+ SHorizontalBox::Slot()
+				.HAlign(HAlign_Right)
+				.VAlign(VAlign_Bottom)
+				.AutoWidth()
+				.Padding(4.f, 2.f)
+				[
+					SNew(SButton)
+					.Text(LOCTEXT("OptionWindow_Cancel", "Cancel"))
+					.OnClicked_Lambda([ParentWindow]()
+						{							
+							UDMXMVRExportOptions* ExportOptions = GetMutableDefault<UDMXMVRExportOptions>();
+							ExportOptions->bCanceled = true;
+
+							ParentWindow->RequestDestroyWindow();
+							return FReply::Handled();
+						})
+				]
+			]
+		];
+	}
+}
+
+#undef LOCTEXT_NAMESPACE
diff --git a/Engine/Plugins/VirtualProduction/DMX/DMXEngine/Source/DMXEditor/Private/Widgets/SDMXMVRExportOptions.h b/Engine/Plugins/VirtualProduction/DMX/DMXEngine/Source/DMXEditor/Private/Widgets/SDMXMVRExportOptions.h
new file mode 100644
--- /dev/null
+++ b/Engine/Plugins/VirtualProduction/DMX/DMXEngine/Source/DMXEditor/Private/Widgets/SDMXMVRExportOptions.h
@@ -0,0 +1,24 @@
+// Copyright Epic Games, Inc. All Rights Reserved.
+
+#pragma once
+
+#include "Widgets/SCompoundWidget.h"
+
+namespace UE::DMX
+{
+	/** The UI presented when exporting MVR files */
+	class SDMXMVRExportOptions final
+		: public SCompoundWidget 
+	{
+	public:
+		SLATE_BEGIN_ARGS(SDMXMVRExportOptions)
+		{}
+
+		SLATE_END_ARGS()
+
+		~SDMXMVRExportOptions();
+
+		/** Constructs the widget */
+		void Construct(const FArguments& InArgs, const TSharedRef<SWindow>& ParentWindow);
+	};
+}
diff --git a/Engine/Plugins/VirtualProduction/DMX/DMXEngine/Source/DMXEditor/Public/DMXEditorModule.h b/Engine/Plugins/VirtualProduction/DMX/DMXEngine/Source/DMXEditor/Public/DMXEditorModule.h
--- a/Engine/Plugins/VirtualProduction/DMX/DMXEngine/Source/DMXEditor/Public/DMXEditorModule.h
+++ b/Engine/Plugins/VirtualProduction/DMX/DMXEngine/Source/DMXEditor/Public/DMXEditorModule.h
@@ -74,6 +74,9 @@ public:
 	TSharedPtr<FExtensibilityManager> ToolBarExtensibilityManager;
 
 private:
+	/** Before 5.4, GDTF assets could be stored containing invalid characters such as @ common to most GDTFs. This function mends all asset names.*/
+	static void FixGDTFAssetNames();
+
 	/** Binds commands for the DMX editor */
 	void BindDMXEditorCommands();
 
diff --git a/Engine/Plugins/VirtualProduction/DMX/DMXEngine/Source/DMXRuntime/Private/DMXRuntimeMainStreamObjectVersion.h b/Engine/Plugins/VirtualProduction/DMX/DMXEngine/Source/DMXRuntime/Private/DMXRuntimeMainStreamObjectVersion.h
--- a/Engine/Plugins/VirtualProduction/DMX/DMXEngine/Source/DMXRuntime/Private/DMXRuntimeMainStreamObjectVersion.h
+++ b/Engine/Plugins/VirtualProduction/DMX/DMXEngine/Source/DMXRuntime/Private/DMXRuntimeMainStreamObjectVersion.h
@@ -36,6 +36,9 @@ struct FDMXRuntimeMainStreamObjectVersion
 		// Update Fixture Patch to hold the MVR Fixture ID (since UE 5.5)
 		DMXFixturePatchHasFixtureID,
 
+		// Rename GDTF import assets with invalid names
+		DMXFixGDTFImportAssetsWithInvalidNames,
+
 		VersionPlusOne,
 		LatestVersion = VersionPlusOne - 1
 	};
diff --git a/Engine/Plugins/VirtualProduction/DMX/DMXEngine/Source/DMXRuntime/Private/Library/DMXEntityFixtureType.cpp b/Engine/Plugins/VirtualProduction/DMX/DMXEngine/Source/DMXRuntime/Private/Library/DMXEntityFixtureType.cpp
--- a/Engine/Plugins/VirtualProduction/DMX/DMXEngine/Source/DMXRuntime/Private/Library/DMXEntityFixtureType.cpp
+++ b/Engine/Plugins/VirtualProduction/DMX/DMXEngine/Source/DMXRuntime/Private/Library/DMXEntityFixtureType.cpp
@@ -315,23 +315,26 @@ void UDMXEntityFixtureType::PostEditUndo()
 #if WITH_EDITOR
 FString UDMXEntityFixtureType::GetCleanGDTFFileNameSynchronous(bool bWithExtension) const
 {
-	if (!GDTFSource.IsNull())
-	{
-		if (UDMXImportGDTF* GDTF = GDTFSource.LoadSynchronous())
-		{
-			const FString GDTFFilename = FPaths::GetBaseFilename(GDTF->GetGDTFAssetImportData()->GetFilePathAndName());
-			const FString Extension = bWithExtension ? TEXT(".gdtf") : TEXT("");
-
-			return GDTFFilename + Extension;
-		}
-	}
+	if (GDTFSource.IsValid())
+	{				
+		// Return the GDTF source filename if a GDTF is set
+		UDMXImportGDTF* GDTF = GDTFSource.LoadSynchronous();
 
-	const FString EngineVersion = FString::Printf(TEXT("%u_%u"), FEngineVersion::Current().GetMajor(), FEngineVersion::Current().GetMinor());
-	const FString DateTime = FDateTime::Now().ToString(TEXT("%d_%m_%y"));
-	const FString GDTFFilename = FString::Printf(TEXT("EpicGames@UE%s_Generated_%s@%s"), *EngineVersion, *Name, *DateTime);
-	const FString Extension = bWithExtension ? TEXT(".gdtf") : TEXT("");
+		const FString GDTFFilename = FPaths::GetBaseFilename(GDTF->GetGDTFAssetImportData()->GetFilePathAndName());
+		const FString Extension = bWithExtension ? TEXT(".gdtf") : TEXT("");
 
-	return GDTFFilename + Extension;
+		return GDTFFilename + Extension;
+	}
+	else
+	{		
+		// Generate a filename if no GDTF is set
+		const FString EngineVersion = FString::Printf(TEXT("%u_%u"), FEngineVersion::Current().GetMajor(), FEngineVersion::Current().GetMinor());
+		const FString DateTime = FDateTime::Now().ToString(TEXT("%d_%m_%y"));
+		const FString GeneratedGDTFFilename = FString::Printf(TEXT("EpicGames@UE%s_Generated_%s@%s"), *EngineVersion, *Name, *DateTime);
+		const FString Extension = bWithExtension ? TEXT(".gdtf") : TEXT("");
+
+		return GeneratedGDTFFilename + Extension;
+	}
 }
 #endif // WITH_EDITOR
 
diff --git a/Engine/Plugins/VirtualProduction/DMX/DMXEngine/Source/DMXRuntime/Private/Library/DMXLibrary.cpp b/Engine/Plugins/VirtualProduction/DMX/DMXEngine/Source/DMXRuntime/Private/Library/DMXLibrary.cpp
--- a/Engine/Plugins/VirtualProduction/DMX/DMXEngine/Source/DMXRuntime/Private/Library/DMXLibrary.cpp
+++ b/Engine/Plugins/VirtualProduction/DMX/DMXEngine/Source/DMXRuntime/Private/Library/DMXLibrary.cpp
@@ -578,7 +578,7 @@ UDMXMVRGeneralSceneDescription* UDMXLibrary::UpdateGeneralSceneDescription()
 {
 	if (ensureAlwaysMsgf(GeneralSceneDescription, TEXT("Trying to update General Scene Description of %s, but the General Scene Description is not valid."), *GetName()))
 	{
-		GeneralSceneDescription->WriteDMXLibraryToGeneralSceneDescription(*this);
+		GeneralSceneDescription->WriteDMXLibrary(*this);
 		return GeneralSceneDescription;
 	}
 
diff --git a/Engine/Plugins/VirtualProduction/DMX/DMXEngine/Source/DMXRuntime/Private/MVR/DMXMVRGeneralSceneDescription.cpp b/Engine/Plugins/VirtualProduction/DMX/DMXEngine/Source/DMXRuntime/Private/MVR/DMXMVRGeneralSceneDescription.cpp
--- a/Engine/Plugins/VirtualProduction/DMX/DMXEngine/Source/DMXRuntime/Private/MVR/DMXMVRGeneralSceneDescription.cpp
+++ b/Engine/Plugins/VirtualProduction/DMX/DMXEngine/Source/DMXRuntime/Private/MVR/DMXMVRGeneralSceneDescription.cpp
@@ -2,14 +2,19 @@
 
 #include "MVR/DMXMVRGeneralSceneDescription.h"
 
+#include "Algo/AnyOf.h"
+#include "Algo/MaxElement.h"
 #include "Library/DMXEntityFixturePatch.h"
 #include "Library/DMXEntityFixtureType.h"
+#include "Library/DMXImportGDTF.h"
+#include "Library/DMXGDTFAssetImportData.h"
 #include "Library/DMXLibrary.h"
 #include "MVR/DMXMVRAssetImportData.h"
 #include "MVR/Types/DMXMVRChildListNode.h"
 #include "MVR/Types/DMXMVRFixtureNode.h"
 #include "MVR/Types/DMXMVRParametricObjectNodeBase.h"
 #include "MVR/Types/DMXMVRRootNode.h"
+#include "MVR/Types/DMXMVRSceneNode.h"
 #include "XmlFile.h"
 
 #define LOCTEXT_NAMESPACE "DMXMVRGeneralSceneDescription"
@@ -59,24 +64,28 @@ UDMXMVRGeneralSceneDescription* UDMXMVRGeneralSceneDescription::CreateFromXmlFil
 UDMXMVRGeneralSceneDescription* UDMXMVRGeneralSceneDescription::CreateFromDMXLibrary(const UDMXLibrary& DMXLibrary, UObject* Outer, FName Name, EObjectFlags Flags)
 {
 	UDMXMVRGeneralSceneDescription* GeneralSceneDescription = NewObject<UDMXMVRGeneralSceneDescription>(Outer, Name, Flags);
-	GeneralSceneDescription->WriteDMXLibraryToGeneralSceneDescription(DMXLibrary);
+	GeneralSceneDescription->WriteDMXLibrary(DMXLibrary);
 
 	return GeneralSceneDescription;
 }
 #endif // WITH_EDITOR
 
-
 #if WITH_EDITOR
 void UDMXMVRGeneralSceneDescription::WriteDMXLibraryToGeneralSceneDescription(const UDMXLibrary& DMXLibrary)
+{
+	// Deprecated 5.5, renamed to WriteDMXLibrary
+	WriteDMXLibrary(DMXLibrary);
+}
+
+void UDMXMVRGeneralSceneDescription::WriteDMXLibrary(const UDMXLibrary& DMXLibrary)
 {
 	TArray<UDMXEntityFixturePatch*> FixturePatches = DMXLibrary.GetEntitiesTypeCast<UDMXEntityFixturePatch>();
-	FixturePatches.Sort([](const UDMXEntityFixturePatch& FixturePatchA, const UDMXEntityFixturePatch& FixturePatchB)
+	Algo::Sort(FixturePatches, [](const UDMXEntityFixturePatch* FixturePatchA, const UDMXEntityFixturePatch* FixturePatchB)
 		{
-			const bool bUniverseIsSmaller = FixturePatchA.GetUniverseID() < FixturePatchB.GetUniverseID();
-			const bool bUniverseIsEqual = FixturePatchA.GetUniverseID() == FixturePatchB.GetUniverseID();
-			const bool bAddressIsSmaller = FixturePatchA.GetStartingChannel() <= FixturePatchB.GetStartingChannel();
+			const uint64 AbsoluteAddressA = static_cast<uint64>(FixturePatchA->GetUniverseID()) * DMX_UNIVERSE_SIZE + FixturePatchA->GetStartingChannel();
+			const uint64 AbsoluteAddressB = static_cast<uint64>(FixturePatchB->GetUniverseID()) * DMX_UNIVERSE_SIZE + FixturePatchB->GetStartingChannel();
 
-			return bUniverseIsSmaller || (bUniverseIsEqual && bAddressIsSmaller);
+			return AbsoluteAddressA <= AbsoluteAddressB;
 		});
 
 	// Remove Fixture Nodes no longer defined in the DMX Library
@@ -100,55 +109,24 @@ void UDMXMVRGeneralSceneDescription::WriteDMXLibraryToGeneralSceneDescription(co
 	{
 		if (FixturePatch)
 		{
-			WriteFixturePatchToGeneralSceneDescription(*FixturePatch);
+			WriteFixturePatch(*FixturePatch);
 		}
 	}
 }
 #endif // WITH_EDITOR
 
 #if WITH_EDITOR
-bool UDMXMVRGeneralSceneDescription::CanCreateXmlFile(FText& OutReason) const
+UDMXMVRFixtureNode* UDMXMVRGeneralSceneDescription::WriteFixturePatch(const UDMXEntityFixturePatch& FixturePatch)
 {
-	TArray<UDMXMVRFixtureNode*> FixtureNodes;
-	GetFixtureNodes(FixtureNodes);
-	if (FixtureNodes.IsEmpty())
-	{
-		OutReason = LOCTEXT("CannotCreateXmlFileBecauseNoFixtures", "DMX Library does not define valid MVR fixtures.");
-		return false;
-	}
-	return true;
-}
-#endif // WITH_EDITOR
-
-#if WITH_EDITOR
-TSharedPtr<FXmlFile> UDMXMVRGeneralSceneDescription::CreateXmlFile() const
-{
-	return RootNode ? RootNode->CreateXmlFile() : nullptr;
-}
-#endif // WITH_EDITOR
-
-#if WITH_EDITOR
-void UDMXMVRGeneralSceneDescription::WriteFixturePatchToGeneralSceneDescription(const UDMXEntityFixturePatch& FixturePatch)
-{
-	if (!ensureMsgf(RootNode, TEXT("Unexpected: MVR General Scene Description Root Node is invalid.")))
-	{
-		return;
-	}
-
 	const UDMXLibrary* DMXLibrary = FixturePatch.GetParentLibrary();
 	if (!DMXLibrary)
 	{
-		return;
+		return nullptr;
 	}
 
 	const FGuid& MVRFixtureUUID = FixturePatch.GetMVRFixtureUUID();
-	TObjectPtr<UDMXMVRParametricObjectNodeBase>* ParametricObjectNodePtr = RootNode->FindParametricObjectNodeByUUID(MVRFixtureUUID);
-
-	UDMXMVRFixtureNode* MVRFixtureNode = nullptr;
-	if (ParametricObjectNodePtr)
-	{
-		MVRFixtureNode = Cast<UDMXMVRFixtureNode>(*ParametricObjectNodePtr);
-	}
+	const TObjectPtr<UDMXMVRParametricObjectNodeBase>* ParametricObjectNodePtr = RootNode->FindParametricObjectNodeByUUID(MVRFixtureUUID);
+	UDMXMVRFixtureNode* MVRFixtureNode = ParametricObjectNodePtr ? Cast<UDMXMVRFixtureNode>(*ParametricObjectNodePtr) : nullptr;
 
 	if (!MVRFixtureNode)
 	{
@@ -167,9 +145,10 @@ void UDMXMVRGeneralSceneDescription::WriteFixturePatchToGeneralSceneDescription(
 	UDMXEntityFixtureType* FixtureType = FixturePatch.GetFixtureType();
 	const int32 ModeIndex = FixturePatch.GetActiveModeIndex();
 	bool bSetGDTFSpec = false;
-	if (FixtureType &&
-		FixtureType->Modes.IsValidIndex(ModeIndex))
+
+	if (FixtureType && FixtureType->Modes.IsValidIndex(ModeIndex))
 	{
+		// Instead refer to the generated file name
 		MVRFixtureNode->GDTFMode = FixtureType->Modes[ModeIndex].ModeName;
 
 		constexpr bool bWithExtension = false;
@@ -184,6 +163,83 @@ void UDMXMVRGeneralSceneDescription::WriteFixturePatchToGeneralSceneDescription(
 		MVRFixtureNode->GDTFMode = TEXT("");
 		MVRFixtureNode->GDTFSpec = TEXT("");
 	}
+
+	SanetizeFixtureNode(*MVRFixtureNode);
+
+	return MVRFixtureNode;
+}
+#endif // WITH_EDITOR
+
+#if WITH_EDITOR
+void UDMXMVRGeneralSceneDescription::RemoveFixtureNode(const FGuid& FixtureUUID)
+{
+	if (!ensureMsgf(RootNode, TEXT("Unexpected: MVR General Scene Description Root Node is invalid.")))
+	{
+		return;
+	}
+
+	const TObjectPtr<UDMXMVRParametricObjectNodeBase>* ParametricObjectNodePtr = RootNode->FindParametricObjectNodeByUUID(FixtureUUID);
+	if (ParametricObjectNodePtr && (*ParametricObjectNodePtr)->GetClass() == UDMXMVRFixtureNode::StaticClass())
+	{
+		RootNode->RemoveParametricObjectNode(*ParametricObjectNodePtr);
+	}
+}
+#endif // WITH_EDITOR
+
+#if WITH_EDITOR
+bool UDMXMVRGeneralSceneDescription::CanCreateXmlFile(FText& OutReason) const
+{
+	TArray<UDMXMVRFixtureNode*> FixtureNodes;
+	GetFixtureNodes(FixtureNodes);
+	if (FixtureNodes.IsEmpty())
+	{
+		OutReason = LOCTEXT("CannotCreateXmlFileBecauseNoFixtures", "DMX Library does not define valid MVR fixtures.");
+		return false;
+	}
+	return true;
+}
+#endif // WITH_EDITOR
+
+#if WITH_EDITOR
+TSharedPtr<FXmlFile> UDMXMVRGeneralSceneDescription::CreateXmlFile() const
+{
+	return RootNode ? RootNode->CreateXmlFile() : nullptr;
+}
+#endif // WITH_EDITOR
+
+#if WITH_EDITOR
+void UDMXMVRGeneralSceneDescription::SanetizeFixtureNode(UDMXMVRFixtureNode& FixtureNode)
+{
+	TArray<UDMXMVRFixtureNode*> FixtureNodes;
+	GetFixtureNodes(FixtureNodes);
+
+	const bool bInvalidFixtureID = Algo::AnyOf(FixtureNodes, [&FixtureNode](const UDMXMVRFixtureNode* Other)
+		{
+			return
+				Other != &FixtureNode &&
+				Other->FixtureID == FixtureNode.FixtureID;
+		});
+	if (bInvalidFixtureID)
+	{
+		const UDMXMVRFixtureNode* const* MaxFixtureNodePtr = Algo::MaxElementBy(FixtureNodes, [](const UDMXMVRFixtureNode* FixtureNode)
+			{
+				int32 IntegralFixtureID;
+				return LexTryParseString(IntegralFixtureID, *FixtureNode->FixtureID) ? IntegralFixtureID : 1;
+			});
+
+		FixtureNode.FixtureID = MaxFixtureNodePtr ? (*MaxFixtureNodePtr)->FixtureID : TEXT("1");
+	}
+
+	const bool bInvalidMVRUUID = Algo::AnyOf(FixtureNodes, [&FixtureNode](const UDMXMVRFixtureNode* Other)
+		{
+			return
+				Other != &FixtureNode &&
+				Other->UUID == FixtureNode.UUID;
+		});
+	if (bInvalidMVRUUID)
+	{
+		FixtureNode.UUID = FGuid::NewGuid();
+	}
 }
 #endif // WITH_EDITOR
 
diff --git a/Engine/Plugins/VirtualProduction/DMX/DMXEngine/Source/DMXRuntime/Public/Library/DMXEntityFixtureType.h b/Engine/Plugins/VirtualProduction/DMX/DMXEngine/Source/DMXRuntime/Public/Library/DMXEntityFixtureType.h
--- a/Engine/Plugins/VirtualProduction/DMX/DMXEngine/Source/DMXRuntime/Public/Library/DMXEntityFixtureType.h
+++ b/Engine/Plugins/VirtualProduction/DMX/DMXEngine/Source/DMXRuntime/Public/Library/DMXEntityFixtureType.h
@@ -524,6 +524,11 @@ public:
 	UPROPERTY(EditAnywhere, BlueprintReadOnly, Category = "Fixture Settings", Meta = (DisplayName = "GDTF Source"))
 	TSoftObjectPtr<UDMXImportGDTF> GDTFSource;
 
+#if WITH_EDITORONLY_DATA
+	/** If checked, generates a new GDTF instead of exporting the imported GDTF. This adopts changes in editor but in most cases will result in data loss and is not recommended. */
+	UPROPERTY(EditAnywhere, AdvancedDisplay, Category = "Fixture Settings", Meta = (DisplayName = "Generate GDTF (not recommended)"))
+	bool bExportGeneratedGDTF = false;
+#endif
 
 	//////////////////////////////////////////////////
 	// Deprecated Members
diff --git a/Engine/Plugins/VirtualProduction/DMX/DMXEngine/Source/DMXRuntime/Public/MVR/DMXMVRGeneralSceneDescription.h b/Engine/Plugins/VirtualProduction/DMX/DMXEngine/Source/DMXRuntime/Public/MVR/DMXMVRGeneralSceneDescription.h
--- a/Engine/Plugins/VirtualProduction/DMX/DMXEngine/Source/DMXRuntime/Public/MVR/DMXMVRGeneralSceneDescription.h
+++ b/Engine/Plugins/VirtualProduction/DMX/DMXEngine/Source/DMXRuntime/Public/MVR/DMXMVRGeneralSceneDescription.h
@@ -13,6 +13,7 @@ class FXmlFile;
 class UDMXEntityFixturePatch;
 class UDMXLibrary;
 class UDMXMVRAssetImportData;
+class UDMXMVRClassNode;
 class UDMXMVRRootNode;
 class UDMXMVRFixtureNode;
 
@@ -36,16 +37,30 @@ public:
 
 #if WITH_EDITOR
 	/** Creates an MVR General Scene Description from an Xml File */
-	static UDMXMVRGeneralSceneDescription* CreateFromXmlFile(TSharedRef<FXmlFile> GeneralSceneDescriptionXml, UObject* Outer, FName Name, EObjectFlags Flags = RF_NoFlags);
+	static UDMXMVRGeneralSceneDescription* CreateFromXmlFile(TSharedRef<FXmlFile> GeneralSceneDescriptionXml, UObject* Outer, FName Name = NAME_None, EObjectFlags Flags = RF_NoFlags);
 
 	/** Creates an MVR General Scene Description from a DMX Library */
-	static UDMXMVRGeneralSceneDescription* CreateFromDMXLibrary(const UDMXLibrary& DMXLibrary, UObject* Outer, FName Name, EObjectFlags Flags = RF_NoFlags);
+	static UDMXMVRGeneralSceneDescription* CreateFromDMXLibrary(const UDMXLibrary& DMXLibrary, UObject* Outer, FName Name = NAME_None, EObjectFlags Flags = RF_NoFlags);
 
+	/** DEPRECATED 5.5 */
+	UE_DEPRECATED(5.5, "Changed to WriteDMXLibrary for better readability and consitency with new members.")
+	void WriteDMXLibraryToGeneralSceneDescription(const UDMXLibrary& DMXLibrary);
+	
 	/**
 	 * Writes the Library to the General Scene Description, effectively removing inexisting and adding
 	 * new MVR Fixtures, according to what MVR Fixture UUIDs the Fixture Patches of the Library contain.
 	 */
-	void WriteDMXLibraryToGeneralSceneDescription(const UDMXLibrary& DMXLibrary);
+	void WriteDMXLibrary(const UDMXLibrary& DMXLibrary);
+
+	/**
+	 * Writes the Fixture Patch to the General Scene Description.
+	 *
+	 * Returns the fixture node or nullptr if the node could not be created.
+	 */
+	UDMXMVRFixtureNode* WriteFixturePatch(const UDMXEntityFixturePatch& FixturePatch);
+
+	/** Removes a fixture node from the General Scene Description. */
+	void RemoveFixtureNode(const FGuid& FixtureUUID);
 
 	/** Returns true if an Xml File can be created. Gives a reason if no MVR can be exported */
 	bool CanCreateXmlFile(FText& OutReason) const;
@@ -55,12 +70,10 @@ public:
 
 	/** Returns MVR Asset Import Data for this asset */
 	FORCEINLINE UDMXMVRAssetImportData* GetMVRAssetImportData() const { return MVRAssetImportData; }
-#endif
 
 private:
-#if WITH_EDITOR
-	/** Writes the Fixture Patch to the General Scene Description, adding a new MVR Fixture if required */
-	void WriteFixturePatchToGeneralSceneDescription(const UDMXEntityFixturePatch& FixturePatch);
+	/** Makes sure the node has a unique MVR UUID and Fixture ID */
+	void  SanetizeFixtureNode(UDMXMVRFixtureNode& FixtureNode);
 
 	/** Parses a General Scene Description Xml File. Only ment to be used for initialization (ensured) */
 	[[nodiscard]] bool ParseGeneralSceneDescriptionXml(const TSharedRef<FXmlFile>& GeneralSceneDescriptionXml);
diff --git a/Engine/Plugins/VirtualProduction/DMX/DMXEngine/Source/DMXRuntime/Public/MVR/Types/DMXMVRFixtureNode.h b/Engine/Plugins/VirtualProduction/DMX/DMXEngine/Source/DMXRuntime/Public/MVR/Types/DMXMVRFixtureNode.h
--- a/Engine/Plugins/VirtualProduction/DMX/DMXEngine/Source/DMXRuntime/Public/MVR/Types/DMXMVRFixtureNode.h
+++ b/Engine/Plugins/VirtualProduction/DMX/DMXEngine/Source/DMXRuntime/Public/MVR/Types/DMXMVRFixtureNode.h
@@ -163,6 +163,10 @@ public:
 	UPROPERTY()
 	FString Name;
 
+	/** The Class the object belongs to. */
+	UPROPERTY()
+	FDMXOptionalGuid Classing;
+
 	/** The name of the file containing the GDTF information for this light fixture. */
 	UPROPERTY()
 	FString GDTFSpec;
diff --git a/Engine/Plugins/VirtualProduction/DMX/DMXGDTF/Source/DMXGDTF/Internal/GDTF/DMXModes/DMXGDTFChannelFunction.h b/Engine/Plugins/VirtualProduction/DMX/DMXGDTF/Source/DMXGDTF/Internal/GDTF/DMXModes/DMXGDTFChannelFunction.h
--- a/Engine/Plugins/VirtualProduction/DMX/DMXGDTF/Source/DMXGDTF/Internal/GDTF/DMXModes/DMXGDTFChannelFunction.h
+++ b/Engine/Plugins/VirtualProduction/DMX/DMXGDTF/Source/DMXGDTF/Internal/GDTF/DMXModes/DMXGDTFChannelFunction.h
@@ -84,10 +84,10 @@ namespace UE::DMX::GDTF
 		FString ModeMaster;
 
 		/** Only used together with ModeMaster; DMX start value; Default value: 0/1 */
-		FDMXGDTFDMXValue ModeFrom = 0;
+		FDMXGDTFDMXValue ModeFrom = TEXT("0/1");
 
 		/** Only used together with ModeMaster; DMX start value; Default value: 0/1 */
-		FDMXGDTFDMXValue ModeTo = 0;
+		FDMXGDTFDMXValue ModeTo = TEXT("0/1");
 
 		/** (Optional) Link to DMX Profile; Starting point: DMX Profile Collect */
 		FString DMXProfile;
diff --git a/Engine/Plugins/VirtualProduction/DMX/DMXGDTF/Source/DMXGDTF/Internal/GDTF/DMXModes/DMXGDTFDMXValue.h b/Engine/Plugins/VirtualProduction/DMX/DMXGDTF/Source/DMXGDTF/Internal/GDTF/DMXModes/DMXGDTFDMXValue.h
--- a/Engine/Plugins/VirtualProduction/DMX/DMXGDTF/Source/DMXGDTF/Internal/GDTF/DMXModes/DMXGDTFDMXValue.h
+++ b/Engine/Plugins/VirtualProduction/DMX/DMXGDTF/Source/DMXGDTF/Internal/GDTF/DMXModes/DMXGDTFDMXValue.h
@@ -2,12 +2,13 @@
 
 #pragma once
 
-#include "Containers/UnrealString.h"
 #include "GenericPlatform/GenericPlatform.h"
-#include "Misc/Optional.h"
+#include "Templates/SharedPointer.h"
 
 namespace UE::DMX::GDTF
 {
+	class FDMXGDTFDMXChannel;
+
 	/**
 	 * Special type to define DMX value where n is the byte count. The byte count can be individually specified without depending on the resolution of the DMX Channel.
 	 * By default byte mirroring is used for the conversion. So 255/1 in a 16 bit channel will result in 65535.
@@ -16,36 +17,39 @@ namespace UE::DMX::GDTF
 	struct DMXGDTF_API FDMXGDTFDMXValue
 	{
 		FDMXGDTFDMXValue() = default;
-		FDMXGDTFDMXValue(const FString& InValue);
-		FDMXGDTFDMXValue(const uint32 InValue);
-		FDMXGDTFDMXValue(const TOptional<uint32>& InValue);
+		FDMXGDTFDMXValue(const TCHAR* InValue);
+		FDMXGDTFDMXValue(const uint32 InValue, const int32 InNumBytes, const bool bInByteMirroring = true);
+
+		/** Gets the DMX value in the resolution of the specified DMX channel. Returns false if the DMX Value has special value "None" */
+		bool Get(const TSharedRef<FDMXGDTFDMXChannel>& InDMXChannel, uint32& OutValue) const;
+	
+		/** Gets the DMX value in the resolution of the specified DMX channel, only valid when not IsSet (checked). */
+		uint32 GetChecked(const TSharedRef<FDMXGDTFDMXChannel>& InDMXChannel) const;
 
-		bool operator==(const FDMXGDTFDMXValue& Other) const { return IntegerValue == Other.IntegerValue; }
-		bool operator!=(const FDMXGDTFDMXValue& Other) const { return IntegerValue != Other.IntegerValue; }
+		/** Returns the DMX value as string. If the DMX value is not set, returns 0/1. */
+		FString AsString() const;
 
-		/** Returns the value as an optional integer. If the optional is not set, it equals to "None". */
-		TOptional<uint32> AsInt() const { return IntegerValue; }
+		/** Resets the DMX value */
+		void Reset();
 
-		/** Returns the value as an integer. Assumes the int is set (checked). */
-		uint32 AsIntChecked() const;
+		/** Returns true if the DMX Value is set. */
+		bool IsSet() const; 
 
-		/** Returns the value as string. */
-		FString AsString() const { return StringValue; }
+		bool operator==(const FDMXGDTFDMXValue& Other) const { return Value == Other.Value; }
+		bool operator!=(const FDMXGDTFDMXValue& Other) const { return Value != Other.Value; }
 
-		/** Returns true if the value is set. Otherwise it uses special value "None". */
-		bool IsSet() const;
 
-		/** Sets the DMX Value by int. */
-		void Set(uint32 InValue);
+	private:
+		/** Returns the Max value of WordSize */
+		uint32 GetMax(uint8 WordSize) const;
 
-		/** Sets the value by string. Special value "None" resets the optional int. */
-		void Set(const FString& InValue);
+		/** The DMX Value */
+		uint32 Value = 0;
 
-		/** Resets to None. */
-		void Reset();
+		/** The number of bytes of this DMX value */
+		uint8 NumBytes = 0;
 
-	private:
-		FString StringValue = TEXT("None");
-		TOptional<uint32> IntegerValue;
+		/** Whether to use byte mirroring or byte shifting when accessing the DMX value */
+		bool bByteMirroring = true;
 	};
 }
diff --git a/Engine/Plugins/VirtualProduction/DMX/DMXGDTF/Source/DMXGDTF/Private/GDTF/DMXModes/DMXGDTFChannelFunction.cpp b/Engine/Plugins/VirtualProduction/DMX/DMXGDTF/Source/DMXGDTF/Private/GDTF/DMXModes/DMXGDTFChannelFunction.cpp
--- a/Engine/Plugins/VirtualProduction/DMX/DMXGDTF/Source/DMXGDTF/Private/GDTF/DMXModes/DMXGDTFChannelFunction.cpp
+++ b/Engine/Plugins/VirtualProduction/DMX/DMXGDTF/Source/DMXGDTF/Private/GDTF/DMXModes/DMXGDTFChannelFunction.cpp
@@ -80,7 +80,7 @@ namespace UE::DMX::GDTF
 			.SetAttribute(TEXT("Attribute"), Attribute)
 			.SetAttribute(TEXT("OriginalAttribute"), OriginalAttribute)
 			.SetAttribute(TEXT("DMXFrom"), DMXFrom)
-			.SetAttribute(TEXT("Default"), Default, FDMXGDTFDMXValue())
+			.SetAttribute(TEXT("Default"), Default)
 			.SetAttribute(TEXT("PhysicalFrom"), PhysicalFrom)
 			.SetAttribute(TEXT("PhysicalTo"), PhysicalTo)
 			.SetAttribute(TEXT("RealFade"), RealFade)
diff --git a/Engine/Plugins/VirtualProduction/DMX/DMXGDTF/Source/DMXGDTF/Private/GDTF/DMXModes/DMXGDTFDMXChannel.cpp b/Engine/Plugins/VirtualProduction/DMX/DMXGDTF/Source/DMXGDTF/Private/GDTF/DMXModes/DMXGDTFDMXChannel.cpp
--- a/Engine/Plugins/VirtualProduction/DMX/DMXGDTF/Source/DMXGDTF/Private/GDTF/DMXModes/DMXGDTFDMXChannel.cpp
+++ b/Engine/Plugins/VirtualProduction/DMX/DMXGDTF/Source/DMXGDTF/Private/GDTF/DMXModes/DMXGDTFDMXChannel.cpp
@@ -57,10 +57,21 @@ PRAGMA_ENABLE_DEPRECATION_WARNINGS
 		ChildBuilder
 			.SetAttribute(TEXT("Offset"), Offset)
 			.SetAttribute(TEXT("InitialFunction"), InitialFunction, DefaultInitialFunction)
-			.SetAttribute(TEXT("Highlight"), Highlight)
 			.SetAttribute(TEXT("Geometry"), Geometry)
 			.AppendChildren(TEXT("LogicalChannel"), LogicalChannelArray);
 
+		// Set special value "None" for highlight if Highlight is not set
+		if (Highlight.IsSet())
+		{
+			ChildBuilder
+				.SetAttribute(TEXT("Highlight"), Highlight);
+		}
+		else
+		{
+			ChildBuilder
+				.SetAttribute(TEXT("Highlight"), TEXT("None"));
+		}
+
 		return ChildBuilder.GetIntermediateXmlNode();
 	}
 
diff --git a/Engine/Plugins/VirtualProduction/DMX/DMXGDTF/Source/DMXGDTF/Private/GDTF/DMXModes/DMXGDTFDMXValue.cpp b/Engine/Plugins/VirtualProduction/DMX/DMXGDTF/Source/DMXGDTF/Private/GDTF/DMXModes/DMXGDTFDMXValue.cpp
--- a/Engine/Plugins/VirtualProduction/DMX/DMXGDTF/Source/DMXGDTF/Private/GDTF/DMXModes/DMXGDTFDMXValue.cpp
+++ b/Engine/Plugins/VirtualProduction/DMX/DMXGDTF/Source/DMXGDTF/Private/GDTF/DMXModes/DMXGDTFDMXValue.cpp
@@ -2,126 +2,120 @@
 
 #include "GDTF/DMXModes/DMXGDTFDMXValue.h"
 
+#include "GDTF/DMXModes/DMXGDTFDMXChannel.h"
+
 namespace UE::DMX::GDTF
 {
-	FDMXGDTFDMXValue::FDMXGDTFDMXValue(const FString& InValue)
-	{
-		Set(InValue);
-	}
-
-	FDMXGDTFDMXValue::FDMXGDTFDMXValue(const uint32 InValue)
+	FDMXGDTFDMXValue::FDMXGDTFDMXValue(const TCHAR* InStringValue)
 	{
-		Set(InValue);
-	}
-
-	FDMXGDTFDMXValue::FDMXGDTFDMXValue(const TOptional<uint32>& InOptionalValue)
-	{
-		if (InOptionalValue.IsSet())
-		{
-			Set(InOptionalValue.GetValue());
-		}
-	}
-
-	uint32 FDMXGDTFDMXValue::AsIntChecked() const
-	{
-		check(IntegerValue.IsSet());
-
-		return IntegerValue.GetValue();
-	}
-
-	bool FDMXGDTFDMXValue::IsSet() const
-	{
-		return IntegerValue.IsSet();
-	}
-
-	void FDMXGDTFDMXValue::Set(uint32 InValue)
-	{
-		IntegerValue = InValue;
-		StringValue = *FString::Printf(TEXT("%u/1"), InValue);
-	}
-
-	void FDMXGDTFDMXValue::Set(const FString& InValue)
-	{
-		StringValue = InValue;
+		const FString StringValue = InStringValue;
 		if (StringValue.IsEmpty() || StringValue == TEXT("None"))
 		{
-			Reset();
 			return;
 		}
 
 		TArray<FString> Substrings;
 		StringValue.ParseIntoArray(Substrings, TEXT("/"));
 
-		// As per GDTF specs
 		if (Substrings.Num() != 2)
 		{
-			Reset();
 			return;
 		}
 
-		uint32 DMXValue;
-		if (!LexTryParseString(DMXValue, *Substrings[0]))
+		if (!LexTryParseString(Value, *Substrings[0]))
 		{
-			Reset();
 			return;
 		}
 
-		const bool bUseByteShifting = Substrings[1].EndsWith(TEXT("s"));
-		if (bUseByteShifting)
+		if (Substrings[1].EndsWith(TEXT("s")))
 		{
+			bByteMirroring = false;
 			Substrings[1].RemoveFromEnd(TEXT("s"));
+		}
 
-			uint8 ByteShiftValue;
-			if (!LexTryParseString(ByteShiftValue, *Substrings[1]) ||
-				ByteShiftValue > 4)
-			{
-				Reset();
-				return;
-			}
-		
-			IntegerValue = DMXValue << ByteShiftValue * 8;
+		if (!LexTryParseString(NumBytes, *Substrings[1]))
+		{
 			return;
 		}
+	}
+
+	FDMXGDTFDMXValue::FDMXGDTFDMXValue(const uint32 InValue, const int32 InNumBytes, const bool bInByteMirroring)
+		: Value(InValue)
+		, NumBytes(InNumBytes)
+		, bByteMirroring(bInByteMirroring)
+	{}
+
+	bool FDMXGDTFDMXValue::Get(const TSharedRef<FDMXGDTFDMXChannel>& InDMXChannel, uint32& OutValue) const
+	{
+		if (IsSet())
+		{
+			OutValue = 0;
+			return false;
+		}
+
+		OutValue = GetChecked(InDMXChannel);
+		return true;
+	}
+
+	uint32 FDMXGDTFDMXValue::GetChecked(const TSharedRef<FDMXGDTFDMXChannel>& InDMXChannel) const
+	{
+		check(IsSet());
+
+		const uint8 WordSize = InDMXChannel->Offset.Num();
+
+		if (WordSize == NumBytes)
+		{
+			return Value;
+		}
+		else if (bByteMirroring)
+		{
+			return static_cast<uint64>(Value) * GetMax(WordSize) / GetMax(NumBytes);
+		}
+		else
+		{
+			// Byte shift
+			return Value << NumBytes * 8;
+		}
+	}
+
+	FString FDMXGDTFDMXValue::AsString() const
+	{
+		if (IsSet())
+		{
+			return bByteMirroring ?
+				FString::Printf(TEXT("%u/%u"), Value, NumBytes) :
+				FString::Printf(TEXT("%u/%us"), Value, NumBytes);
+		}
 		else
 		{
-			uint8 ByteMirroringValue;
-			if (!LexTryParseString(ByteMirroringValue, *Substrings[1]) ||
-				ByteMirroringValue > 4)
-			{
-				Reset();
-				return;
-			}
-
-			if (ByteMirroringValue == 1)
-			{
-				// No mirroring needed
-				IntegerValue = DMXValue;
-				return;
-			}
-
-			constexpr uint32 Masks[] = { 0x000000FF, 0x0000FF00, 0x00FF0000, 0xFF000000 };
-			if (ByteMirroringValue == 2)
-			{
-				IntegerValue = ((DMXValue & 0x000000FF) << 8) | ((DMXValue & 0x0000FF00) >> 8);
-			}
-			else if (ByteMirroringValue == 3)
-			{
-#if PLATFORM_LITTLE_ENDIAN
-				IntegerValue = ((DMXValue & 0x0000FF00) << 16) | ((DMXValue & 0xFF000000) >> 16);
-#else
-				IntegerValue = ((DMXValue & 0x000000FF) << 16) | ((DMXValue & 0x00FF0000) >> 16);
-#endif
-			}
-			else
-			{
-				IntegerValue = ((DMXValue & 0x000000FF) << 24) | ((DMXValue & 0x0000FF00) << 8) | ((DMXValue & 0x00FF0000) >> 8) | ((DMXValue & 0xFF000000) >> 24);
-			}
+			return TEXT("0/1");
 		}
 	}
 
 	void FDMXGDTFDMXValue::Reset()
 	{
-		IntegerValue.Reset();
-		StringValue = TEXT("None");
+		NumBytes = 0;
+	}
+
+	bool FDMXGDTFDMXValue::IsSet() const
+	{
+		return NumBytes != 0;
+	}
+
+	uint32 FDMXGDTFDMXValue::GetMax(uint8 WordSize) const
+	{
+		switch (WordSize)
+		{
+		case 1:
+			return 0xFF;
+		case 2:
+			return 0xFFFF;
+		case 3:
+			return 0xFFFFFF;
+		case 4:
+			return 0xFFFFFFFF;
+		};
+
+		return 0;
 	}
 }
diff --git a/Engine/Plugins/VirtualProduction/DMX/DMXPixelMapping/Source/DMXPixelMappingRuntime/Private/LayoutScripts/DMXPixelMappingLayoutScript_LayoutByMVR.cpp b/Engine/Plugins/VirtualProduction/DMX/DMXPixelMapping/Source/DMXPixelMappingRuntime/Private/LayoutScripts/DMXPixelMappingLayoutScript_LayoutByMVR.cpp
--- a/Engine/Plugins/VirtualProduction/DMX/DMXPixelMapping/Source/DMXPixelMappingRuntime/Private/LayoutScripts/DMXPixelMappingLayoutScript_LayoutByMVR.cpp
+++ b/Engine/Plugins/VirtualProduction/DMX/DMXPixelMapping/Source/DMXPixelMappingRuntime/Private/LayoutScripts/DMXPixelMappingLayoutScript_LayoutByMVR.cpp
@@ -78,16 +78,13 @@ TMap<const UDMXMVRFixtureNode*, FDMXPixelMappingLayoutToken> UDMXPixelMappingLay
 		return Result;
 	}
 
+	DMXLibrary->UpdateGeneralSceneDescription();
 	UDMXMVRGeneralSceneDescription* GeneralSceneDescription = DMXLibrary->GetLazyGeneralSceneDescription();
 	if (!GeneralSceneDescription)
 	{
 		return Result;
 	}
 
-#if WITH_EDITOR
-	GeneralSceneDescription->WriteDMXLibraryToGeneralSceneDescription(*DMXLibrary);
-#endif
-
 	TArray<UDMXMVRFixtureNode*> FixtureNodes;
 	GeneralSceneDescription->GetFixtureNodes(FixtureNodes);
 
