DMX: Conflict monitor now shows a list of objects actively sending DMX even when no conflicts are detected

diff --git a/Engine/Plugins/VirtualProduction/DMX/DMXEngine/Source/DMXEditor/Private/Widgets/Monitors/DMXConflictMonitorActiveObjectItem.cpp b/Engine/Plugins/VirtualProduction/DMX/DMXEngine/Source/DMXEditor/Private/Widgets/Monitors/DMXConflictMonitorActiveObjectItem.cpp
new file mode 100644
--- /dev/null
+++ b/Engine/Plugins/VirtualProduction/DMX/DMXEngine/Source/DMXEditor/Private/Widgets/Monitors/DMXConflictMonitorActiveObjectItem.cpp
@@ -0,0 +1,11 @@
+// Copyright Epic Games, Inc. All Rights Reserved.
+
+#include "DMXConflictMonitorActiveObjectItem.h"
+
+namespace UE::DMX
+{
+	FDMXConflictMonitorActiveObjectItem::FDMXConflictMonitorActiveObjectItem(const FName& InObjectName, const FSoftObjectPath& InObjectPath)
+		: ObjectName(InObjectName)
+		, ObjectPath(InObjectPath)
+	{}
+}
diff --git a/Engine/Plugins/VirtualProduction/DMX/DMXEngine/Source/DMXEditor/Private/Widgets/Monitors/DMXConflictMonitorActiveObjectItem.h b/Engine/Plugins/VirtualProduction/DMX/DMXEngine/Source/DMXEditor/Private/Widgets/Monitors/DMXConflictMonitorActiveObjectItem.h
new file mode 100644
--- /dev/null
+++ b/Engine/Plugins/VirtualProduction/DMX/DMXEngine/Source/DMXEditor/Private/Widgets/Monitors/DMXConflictMonitorActiveObjectItem.h
@@ -0,0 +1,25 @@
+// Copyright Epic Games, Inc. All Rights Reserved.
+
+#pragma once
+
+#include "Containers/UnrealString.h"
+#include "UObject/NameTypes.h"
+#include "UObject/SoftObjectPath.h"
+
+namespace UE::DMX
+{
+	/** A row in the Conflict Monitor's Active Object list */
+	class FDMXConflictMonitorActiveObjectItem
+		: public TSharedFromThis<FDMXConflictMonitorActiveObjectItem>
+	{
+	public:
+		/** Constructor. Object Path is allowed to be invalid */
+		FDMXConflictMonitorActiveObjectItem(const FName& ObjectName, const FSoftObjectPath& ObjectPath);
+
+		/** The name of the object */
+		const FName ObjectName;
+
+		/** The object path. Can be invalid. */
+		const FSoftObjectPath ObjectPath;
+	};
+}
diff --git a/Engine/Plugins/VirtualProduction/DMX/DMXEngine/Source/DMXEditor/Private/Widgets/Monitors/SDMXConflictMonitor.cpp b/Engine/Plugins/VirtualProduction/DMX/DMXEngine/Source/DMXEditor/Private/Widgets/Monitors/SDMXConflictMonitor.cpp
--- a/Engine/Plugins/VirtualProduction/DMX/DMXEngine/Source/DMXEditor/Private/Widgets/Monitors/SDMXConflictMonitor.cpp
+++ b/Engine/Plugins/VirtualProduction/DMX/DMXEngine/Source/DMXEditor/Private/Widgets/Monitors/SDMXConflictMonitor.cpp
@@ -2,31 +2,39 @@
 
 #include "SDMXConflictMonitor.h"
 
+#include "Algo/AnyOf.h"
+#include "Algo/Find.h"
 #include "Algo/Transform.h"
 #include "Analytics/DMXEditorToolAnalyticsProvider.h"
+#include "Commands/DMXConflictMonitorCommands.h"
+#include "DMXConflictMonitorActiveObjectItem.h"
 #include "DMXConflictMonitorConflictModel.h"
 #include "DMXEditorLog.h"
-#include "Commands/DMXConflictMonitorCommands.h"
 #include "DMXEditorSettings.h"
 #include "DMXEditorStyle.h"
-#include "Framework/Commands/UICommandList.h"
+#include "DMXStats.h"
+#include "Fonts/FontMeasure.h"
 #include "Framework/Application/SlateApplication.h"
+#include "Framework/Commands/UICommandList.h"
+#include "Internationalization/Regex.h"
 #include "IO/DMXConflictMonitor.h"
 #include "IO/DMXPortManager.h"
+#include "SDMXConflictMonitorActiveObjectRow.h"
 #include "SDMXConflictMonitorToolbar.h"
 #include "Widgets/Layout/SScrollBox.h"
+#include "Widgets/Layout/SSplitter.h"
 #include "Widgets/Layout/SWrapBox.h"
 #include "Widgets/Text/SRichTextBlock.h"
 
-
 #define LOCTEXT_NAMESPACE "SDMXConflictMonitor"
 
+DECLARE_CYCLE_STAT(TEXT("DMX Conflict Monitor User Interface"), STAT_DMXConflictMonitorUI, STATGROUP_DMX);
+
 namespace UE::DMX
 {
-	const FName SDMXConflictMonitor::FColumnIds::Ports = "Ports";
-	const FName SDMXConflictMonitor::FColumnIds::Universe = "Universe";
-	const FName SDMXConflictMonitor::FColumnIds::Conflicts = "Conflicts";
-	const FName SDMXConflictMonitor::FColumnIds::Channels = "Channels";
+	const FName SDMXConflictMonitor::FActiveObjectCollumnID::ObjectName = TEXT("ObjectName");
+	const FName SDMXConflictMonitor::FActiveObjectCollumnID::OpenAsset = TEXT("OpenAsset");
+	const FName SDMXConflictMonitor::FActiveObjectCollumnID::ShowInContentBrowser = TEXT("ShowInContentBrowser");
 
 	SDMXConflictMonitor::SDMXConflictMonitor()
 		: StatusInfo(EDMXConflictMonitorStatusInfo::Idle)
@@ -52,33 +60,69 @@ namespace UE::DMX
 					{
 						return StatusInfo;
 					})
+				.TimeGameThread_Lambda([this]()
+					{
+						return TimeGameThread;
+					})
 				.OnDepthChanged_Lambda([this]()
 					{
 						Refresh();
 					})
 			]
 
-			// Log
-			+SVerticalBox::Slot()
+			+ SVerticalBox::Slot()
 			.FillHeight(1.f)
 			.Padding(16.f)
 			[
-				SNew(SScrollBox)
-				.Orientation(EOrientation::Orient_Vertical)
+				SNew(SSplitter)
+				.Orientation(EOrientation::Orient_Horizontal)
+
+				// Log
+				+ SSplitter::Slot()
+				.Value(0.62f)
+				.MinSize(10.f)
+				[
+					SNew(SScrollBox)
+					.Orientation(EOrientation::Orient_Vertical)
 					
-				+ SScrollBox::Slot()
-				.AutoSize()
+					+ SScrollBox::Slot()
+					.AutoSize()
+					[
+						SNew(SBorder)
+						.HAlign(HAlign_Fill)
+						.VAlign(VAlign_Fill)
+						.BorderImage(FAppStyle::GetBrush("NoBorder"))
+						[
+							SAssignNew(LogTextBlock, SRichTextBlock)
+							.Visibility(EVisibility::HitTestInvisible)
+							.AutoWrapText(true)
+							.TextStyle(FAppStyle::Get(), "MessageLog")
+							.DecoratorStyleSet(&FDMXEditorStyle::Get())
+						]
+					]
+				]
+
+				// Active objects
+				+ SSplitter::Slot()
+				.Value(0.38f)
+				.MinSize(10.f)
 				[
-					SNew(SBorder)
-					.HAlign(HAlign_Fill)
-					.VAlign(VAlign_Fill)
-					.BorderImage(FAppStyle::GetBrush("NoBorder"))
+					SNew(SScrollBox)
+					.Orientation(EOrientation::Orient_Vertical)
+					
+					+ SScrollBox::Slot()
+					.FillSize(1.f)
 					[
-						SAssignNew(TextBlock, SRichTextBlock)
-						.Visibility(EVisibility::HitTestInvisible)
-						.AutoWrapText(true)
-						.TextStyle(FAppStyle::Get(), "MessageLog")
-						.DecoratorStyleSet(&FDMXEditorStyle::Get())
+						SNew(SBorder)
+						.HAlign(HAlign_Fill)
+						.VAlign(VAlign_Fill)
+						.BorderImage(FAppStyle::GetBrush("NoBorder"))
+						[
+							SAssignNew(ActiveObjectList, SListView<TSharedPtr<FDMXConflictMonitorActiveObjectItem>>)
+							.HeaderRow(GenerateActiveObjectHeaderRow())
+							.ListItemsSource(&ActiveObjectListSource)
+							.OnGenerateRow(this, &SDMXConflictMonitor::OnGenerateActiveObjectRow)
+						]
 					]
 				]
 			]
@@ -95,6 +139,10 @@ namespace UE::DMX
 
 	void SDMXConflictMonitor::Tick(const FGeometry& AllottedGeometry, const double InCurrentTime, const float InDeltaTime)
 	{
+		SCOPE_CYCLE_COUNTER(STAT_DMXConflictMonitorUI);
+
+		const double StartTime = FPlatformTime::Seconds();
+
 		if (FSlateApplication::Get().AnyMenusVisible())
 		{
 			return;
@@ -106,50 +154,150 @@ namespace UE::DMX
 			return;
 		}
 
+		// Only refresh when data or conflicts changed. This is more performant and leaves the widgets interactable.
+		const  TArray<TSharedRef<FDMXMonitoredOutboundDMXData>> NewOutboundData = ConflictMonitor->GetMonitoredOutboundData();
+		const bool bDataChanged =
+			NewOutboundData.Num() != CachedOutboundData.Num() ||
+			Algo::AnyOf(NewOutboundData, [this](const TSharedRef<FDMXMonitoredOutboundDMXData>& Data)
+				{
+					return Algo::NoneOf(CachedOutboundData, [&Data](const TSharedRef<FDMXMonitoredOutboundDMXData>& Other)
+						{
+							return Other->Trace == Data->Trace;
+						});
+				});
+
 		const TMap<FName, TArray<TSharedRef<FDMXMonitoredOutboundDMXData>>> NewOutboundConflicts = ConflictMonitor->GetOutboundConflictsSynchronous();
+		const bool bConflictsChanged = !CachedOutboundConflicts.OrderIndependentCompareEqual(NewOutboundConflicts);
 
-		if (!CachedOutboundConflicts.OrderIndependentCompareEqual(NewOutboundConflicts) &&
-			!FSlateApplication::Get().GetPressedMouseButtons().Contains(EKeys::LeftMouseButton))
+		const bool bLeftMouseButtonDown = FSlateApplication::Get().GetPressedMouseButtons().Contains(EKeys::LeftMouseButton);
+
+		if ((bDataChanged || bConflictsChanged) && !bLeftMouseButtonDown)
 		{
+			CachedOutboundData = NewOutboundData;
 			CachedOutboundConflicts = NewOutboundConflicts;
 			Refresh();
 		}
 
 		UpdateStatusInfo();
+
+		const double ConflictMonitorTimeGameThread = ConflictMonitor->GetTimeGameThread();
+
+		const double EndTime = FPlatformTime::Seconds();
+		TimeGameThread = (EndTime - StartTime) * 1000.0 + ConflictMonitorTimeGameThread;
+	}
+
+	FReply SDMXConflictMonitor::OnMouseButtonDown(const FGeometry& MyGeometry, const FPointerEvent& MouseEvent)
+	{
+		return FReply::Handled();
+	}
+
+	TSharedRef<SHeaderRow> SDMXConflictMonitor::GenerateActiveObjectHeaderRow()
+	{
+		const TSharedRef<SHeaderRow> HeaderRow =
+			SNew(SHeaderRow)
+			.Visibility_Lambda([this]()
+				{
+					const bool bIsActive =
+						StatusInfo == EDMXConflictMonitorStatusInfo::OK ||
+						StatusInfo == EDMXConflictMonitorStatusInfo::Conflict;
+
+					return bIsActive ? EVisibility::Visible : EVisibility::Collapsed;
+				});
+
+		HeaderRow->AddColumn(
+			SHeaderRow::FColumn::FArguments()
+			.ColumnId(FActiveObjectCollumnID::ObjectName)
+			.DefaultLabel(LOCTEXT("ActiveObjectLabel", "Objects sending DMX"))
+			.FillWidth(1.f)
+		);
+
+		const FText AssetActionText = LOCTEXT("AssetActionLabel", "Asset Action");
+		const TSharedRef<FSlateFontMeasure> FontMeasureService = FSlateApplication::Get().GetRenderer()->GetFontMeasureService();
+		const float AssetActionLabelSize = FontMeasureService->Measure(AssetActionText, FAppStyle::GetFontStyle("NormalText")).X + 8.f;
+
+		HeaderRow->AddColumn(
+			SHeaderRow::FColumn::FArguments()
+			.ColumnId(FActiveObjectCollumnID::OpenAsset)
+			.DefaultLabel(AssetActionText)
+			.ManualWidth(AssetActionLabelSize)
+		);
+
+		const FText BrowseToText = LOCTEXT("BrowseToAssetLabel", "Browse To");
+		const float BrowseToTextLabelSize = FontMeasureService->Measure(BrowseToText, FAppStyle::GetFontStyle("NormalText")).X + 16.f;
+
+		HeaderRow->AddColumn(
+			SHeaderRow::FColumn::FArguments()
+			.ColumnId(FActiveObjectCollumnID::ShowInContentBrowser)
+			.DefaultLabel(BrowseToText)
+			.ManualWidth(BrowseToTextLabelSize)
+		);
+
+		return HeaderRow;
+	}
+
+	TSharedRef<ITableRow> SDMXConflictMonitor::OnGenerateActiveObjectRow(TSharedPtr<FDMXConflictMonitorActiveObjectItem> InItem, const TSharedRef<STableViewBase>& OwnerTable)
+	{
+		return 
+			SNew(SDMXConflictMonitorActiveObjectRow, OwnerTable, InItem.ToSharedRef())
+			.Visibility(EVisibility::SelfHitTestInvisible);
 	}
 
 	void SDMXConflictMonitor::Refresh()
 	{
-		// Test new items for changes
+		const bool bStopped = !UserSession.IsValid();
+		if (bStopped)
+		{
+			return;
+		}
+
+		// Fetch conflicts text
 		TArray<TSharedPtr<FDMXConflictMonitorConflictModel>> NewModels;
 		Algo::Transform(CachedOutboundConflicts, NewModels, [](const TPair<FName, TArray<TSharedRef<FDMXMonitoredOutboundDMXData>>>& Conflicts)
 			{
 				return MakeShared<FDMXConflictMonitorConflictModel>(Conflicts.Value);
 			});
 
-		FString NewText;
+		FString NewLogText;
 		for (const TSharedPtr<FDMXConflictMonitorConflictModel>& Model : NewModels)
 		{
 			constexpr bool bWithMarkup = true;
-			NewText.Append(Model->GetConflictAsString(bWithMarkup));
-			NewText.Append(TEXT("\n"));
+			NewLogText.Append(Model->GetConflictAsString(bWithMarkup));
+			NewLogText.Append(TEXT("\n"));
 		}
 
-		// Auto-pause even if the data hasn't changed
-		const UDMXEditorSettings* EditorSettings = GetDefault<UDMXEditorSettings>();
-		if (!NewModels.IsEmpty() && IsScanning() && EditorSettings->ConflictMonitorSettings.bAutoPause)
+		// Update active DMX Objects
+		TMap<FName, TSharedPtr<FDMXConflictMonitorActiveObjectItem>> ObjectNameToItemMap;
+		for (const TSharedRef<FDMXMonitoredOutboundDMXData>& Data : CachedOutboundData)
 		{
-			Pause();
-		}
+			const FString Trace = Data->Trace.ToString();
+			
+			const FRegexPattern ObjectPattern(TEXT("\\/([^\\/,]+)(?=(?:,|$))"));
+			FRegexMatcher ObjectMatcher(ObjectPattern, Trace);
 
-		// Skip if text did not change
-		if (TextBlock->GetText().ToString() == NewText)
-		{
-			return;
+			const FRegexPattern ObjectPathPattern(TEXT("^([^,]+)"));
+			FRegexMatcher ObjectPathMatcher(ObjectPathPattern, Trace);
+
+			const FName ObjectName = ObjectMatcher.FindNext() ? *ObjectMatcher.GetCaptureGroup(1) : *Trace;
+			const FName ObjectPath = ObjectPathMatcher.FindNext() ? *ObjectPathMatcher.GetCaptureGroup(1) : FName();
+			const FSoftObjectPath SoftObjectPath(ObjectPath.ToString());
+
+			ObjectNameToItemMap.FindOrAdd(ObjectName, MakeShared<FDMXConflictMonitorActiveObjectItem>(ObjectName, SoftObjectPath));
 		}
 
+		// Update texts
 		Models = NewModels;
-		TextBlock->SetText(FText::FromString(NewText));
+		LogTextBlock->SetText(FText::FromString(NewLogText));
+
+		ActiveObjectListSource.Reset();
+		ObjectNameToItemMap.GenerateValueArray(ActiveObjectListSource);
+		ActiveObjectList->RequestListRefresh();
+
+		// Auto-pause even if when data hasn't changed
+		const UDMXEditorSettings* EditorSettings = GetDefault<UDMXEditorSettings>();
+		if (!NewModels.IsEmpty() && IsScanning() && EditorSettings->ConflictMonitorSettings.bAutoPause)
+		{
+			Pause();
+		}
 
 		// Log conflicts (without markup)
 		if (IsPrintingToLog())
@@ -261,11 +409,16 @@ namespace UE::DMX
 		bIsPaused = false;
 		SetCanTick(false);
 
+		CachedOutboundData.Reset();
 		CachedOutboundConflicts.Reset();
 		Models.Reset();
 		Refresh();
 
 		UpdateStatusInfo();
+
+		// Empty the active object list when stopped, so it's clear that it is no longer updated
+		ActiveObjectListSource.Reset();
+		ActiveObjectList->RequestListRefresh();
 	}
 
 	void SDMXConflictMonitor::SetAutoPause(bool bEnabled)
diff --git a/Engine/Plugins/VirtualProduction/DMX/DMXEngine/Source/DMXEditor/Private/Widgets/Monitors/SDMXConflictMonitor.h b/Engine/Plugins/VirtualProduction/DMX/DMXEngine/Source/DMXEditor/Private/Widgets/Monitors/SDMXConflictMonitor.h
--- a/Engine/Plugins/VirtualProduction/DMX/DMXEngine/Source/DMXEditor/Private/Widgets/Monitors/SDMXConflictMonitor.h
+++ b/Engine/Plugins/VirtualProduction/DMX/DMXEngine/Source/DMXEditor/Private/Widgets/Monitors/SDMXConflictMonitor.h
@@ -3,14 +3,20 @@
 #pragma once
 
 #include "Analytics/DMXEditorToolAnalyticsProvider.h"
+#include "UObject/NameTypes.h"
 #include "Widgets/SCompoundWidget.h"
+#include "Widgets/Views/SListView.h"
 
 class FUICommandList;
+class ITableRow;
+class SHeaderRow;
+template <typename ItemType> class SListView;
 class SRichTextBlock;
-
+class STableViewBase;
 
 namespace UE::DMX
 {
+	class FDMXConflictMonitorActiveObjectItem;
 	class FDMXConflictMonitorConflictModel;
 	class FDMXConflictMonitorUserSession;
 	struct FDMXMonitoredOutboundDMXData;
@@ -26,26 +32,32 @@ namespace UE::DMX
 
 		SLATE_END_ARGS()
 
-		/** ColumnIds for the list view */
-		struct FColumnIds
-		{
-			static const FName Conflicts;
-			static const FName Ports;
-			static const FName Universe;
-			static const FName Channels;
-		};
-
 		SDMXConflictMonitor();
 
 		/** Constructs the widget */
 		void Construct(const FArguments& InArgs);
 
+
+		struct FActiveObjectCollumnID
+		{
+			static const FName ObjectName;
+			static const FName OpenAsset;
+			static const FName ShowInContentBrowser;
+		};
+
 	protected:
 		//~ Begin SWidget interface
 		virtual void Tick(const FGeometry& AllottedGeometry, const double InCurrentTime, const float InDeltaTime) override;
+		virtual FReply OnMouseButtonDown(const FGeometry& MyGeometry, const FPointerEvent& MouseEvent) override;
 		//~ End SWidget interface
 
 	private:
+		/** Generates the header row of the active objects list */
+		TSharedRef<SHeaderRow> GenerateActiveObjectHeaderRow();
+
+		/** Generates a row for an object that is actively sending DMX */
+		TSharedRef<ITableRow> OnGenerateActiveObjectRow(TSharedPtr<FDMXConflictMonitorActiveObjectItem> InItem, const TSharedRef<STableViewBase>& OwnerTable);
+
 		/** Refreshes the widget */
 		void Refresh();
 
@@ -74,6 +86,9 @@ namespace UE::DMX
 		void ToggleRunWhenOpened();
 		bool IsRunWhenOpened() const;
 
+		/** Cached outbound data */
+		TArray<TSharedRef<FDMXMonitoredOutboundDMXData>> CachedOutboundData;
+
 		/** Cached outbound conflicts */
 		TMap<FName, TArray<TSharedRef<FDMXMonitoredOutboundDMXData>>> CachedOutboundConflicts;
 
@@ -81,7 +96,7 @@ namespace UE::DMX
 		TArray<TSharedPtr<FDMXConflictMonitorConflictModel>> Models;
 
 		/** Text block displaying outbound conflicts, one conflict per row */
-		TSharedPtr<SRichTextBlock> TextBlock;
+		TSharedPtr<SRichTextBlock> LogTextBlock;
 
 		/** Timer to refresh at refresh period */
 		double Timer = 0.0;
@@ -89,6 +104,12 @@ namespace UE::DMX
 		/** True if paused */
 		bool bIsPaused = false;
 
+		/** Source for the Active Object List */
+		TArray<TSharedPtr<FDMXConflictMonitorActiveObjectItem>> ActiveObjectListSource;
+
+		/** The Active Object List */
+		TSharedPtr<SListView<TSharedPtr<FDMXConflictMonitorActiveObjectItem>>> ActiveObjectList;
+
 		/** The status of the monitor. Note status info is ment for UI purposes, and not the state of the monitor. */
 		EDMXConflictMonitorStatusInfo StatusInfo;
 
@@ -101,6 +122,9 @@ namespace UE::DMX
 		/** The analytics provider for this tool */
 		FDMXEditorToolAnalyticsProvider AnalyticsProvider;
 
+		/** Time on the game thread */
+		double TimeGameThread = 0.0;
+
 		// Slate args
 		TAttribute<double> UpdateInterval;
 	};
diff --git a/Engine/Plugins/VirtualProduction/DMX/DMXEngine/Source/DMXEditor/Private/Widgets/Monitors/SDMXConflictMonitorActiveObjectRow.cpp b/Engine/Plugins/VirtualProduction/DMX/DMXEngine/Source/DMXEditor/Private/Widgets/Monitors/SDMXConflictMonitorActiveObjectRow.cpp
new file mode 100644
--- /dev/null
+++ b/Engine/Plugins/VirtualProduction/DMX/DMXEngine/Source/DMXEditor/Private/Widgets/Monitors/SDMXConflictMonitorActiveObjectRow.cpp
@@ -0,0 +1,150 @@
+// Copyright Epic Games, Inc. All Rights Reserved.
+
+#include "SDMXConflictMonitorActiveObjectRow.h"
+
+#include "ContentBrowserModule.h"
+#include "DMXConflictMonitorActiveObjectItem.h"
+#include "DMXEditorStyle.h"
+#include "Editor.h"
+#include "IContentBrowserSingleton.h"
+#include "SDMXConflictMonitor.h"
+#include "Styling/AppStyle.h"
+#include "Subsystems/AssetEditorSubsystem.h"
+#include "Textures/SlateIcon.h"
+#include "Widgets/Images/SImage.h"
+#include "Widgets/Input/SHyperlink.h"
+#include "Widgets/Layout/SWidgetSwitcher.h"
+#include "Widgets/Text/SRichTextBlock.h"
+
+#define LOCTEXT_NAMESPACE "SDMXConflictMonitorActiveObjectRow"
+
+namespace UE::DMX
+{
+	void SDMXConflictMonitorActiveObjectRow::Construct(const FArguments& InArgs, const TSharedRef<STableViewBase>& OwnerTable, const TSharedRef<FDMXConflictMonitorActiveObjectItem>& InActiveObjectItem)
+	{
+		ActiveObjectItem = InActiveObjectItem;
+
+		SMultiColumnTableRow<TSharedPtr<FDMXConflictMonitorActiveObjectItem>>::Construct(
+			FSuperRowType::FArguments(),
+			OwnerTable
+		);
+	}
+
+	TSharedRef<SWidget> SDMXConflictMonitorActiveObjectRow::GenerateWidgetForColumn(const FName& ColumnName)
+	{
+		const FMargin Padding(4.f, 2.f, 4.f, 2.f);
+		if (ColumnName == SDMXConflictMonitor::FActiveObjectCollumnID::ObjectName)
+		{
+			return
+				SNew(SBorder)
+				.HAlign(HAlign_Left)
+				.VAlign(VAlign_Center)
+				.BorderImage(FAppStyle::GetBrush("NoBorder"))
+				.Padding(Padding)
+				[
+					SNew(STextBlock)
+					.Visibility(EVisibility::HitTestInvisible)
+					.TextStyle(FAppStyle::Get(), "MessageLog")
+					.Text_Lambda([this, SharedThis = AsShared()]()
+						{
+							return FText::FromString(ActiveObjectItem->ObjectName.ToString());
+						})
+				];
+		}
+		else if (ColumnName == SDMXConflictMonitor::FActiveObjectCollumnID::OpenAsset)
+		{
+			return 
+				SNew(SBorder)
+				.HAlign(HAlign_Left)
+				.VAlign(VAlign_Center)
+				.BorderImage(FAppStyle::GetBrush("NoBorder"))
+				.Padding(Padding)
+				[
+					SNew(SWidgetSwitcher)
+					.WidgetIndex_Lambda([this, SharedThis = AsShared()]
+						{
+							return ActiveObjectItem->ObjectPath.IsNull() ? 0 : 1;
+						})
+
+					+ SWidgetSwitcher::Slot()
+					[
+						SNew(STextBlock)
+						.Text(LOCTEXT("UnavailableAssetHyperlink", "unknown"))
+						.TextStyle(FAppStyle::Get(), "NormalFont")
+					]
+
+					+ SWidgetSwitcher::Slot()
+					[
+						SNew(SHyperlink)
+						.Text(LOCTEXT("OpenAssetHyperlink", "Open Asset"))
+						.Style(FAppStyle::Get(), "Common.GotoNativeCodeHyperlink")
+						.OnNavigate(this, &SDMXConflictMonitorActiveObjectRow::OnOpenAssetClicked)
+					]
+				];
+		}
+		else if (ColumnName == SDMXConflictMonitor::FActiveObjectCollumnID::ShowInContentBrowser)
+		{
+			return 
+				SNew(SBorder)
+				.HAlign(HAlign_Center)
+				.VAlign(VAlign_Center)
+				.BorderImage(FAppStyle::GetBrush("NoBorder"))
+				.Padding(Padding)
+				[
+					SNew(SWidgetSwitcher)
+					.WidgetIndex_Lambda([this, SharedThis = AsShared()]
+					{
+						return ActiveObjectItem->ObjectPath.IsNull() ? 0 : 1;
+					})
+
+					+ SWidgetSwitcher::Slot()
+					[
+						SNew(STextBlock)
+						.Text(LOCTEXT("UnavailableShowInContentBrowser", ""))
+						.TextStyle(FAppStyle::Get(), "NormalFont")
+					]
+
+					+ SWidgetSwitcher::Slot()
+					[
+						SNew(SButton)
+						[
+							SNew(SImage)
+							.Image(FAppStyle::Get().GetBrush("Icons.Search"))
+							.OnMouseButtonDown(this, &SDMXConflictMonitorActiveObjectRow::OnShowInContentBrowserClicked)
+						]
+					]
+				];
+		}
+		else
+		{
+			checkf(0, TEXT("Unhandled column ID"));
+		}
+
+		return SNullWidget::NullWidget;
+	}
+
+	void SDMXConflictMonitorActiveObjectRow::OnOpenAssetClicked()
+	{
+		if (UObject* AssetObject = ActiveObjectItem->ObjectPath.TryLoad())
+		{
+			TArray<UObject*> Assets;
+			Assets.Add(AssetObject);
+			GEditor->GetEditorSubsystem<UAssetEditorSubsystem>()->OpenEditorForAssets(Assets);
+		}
+	}
+
+	FReply SDMXConflictMonitorActiveObjectRow::OnShowInContentBrowserClicked(const FGeometry& MyGeometry, const FPointerEvent& MouseEvent)
+	{
+		if (UObject* AssetObject = ActiveObjectItem->ObjectPath.TryLoad())
+		{
+			FContentBrowserModule& ContentBrowserModule = FModuleManager::Get().LoadModuleChecked<FContentBrowserModule>(TEXT("ContentBrowser"));
+			TArray<FAssetData> Assets;
+			Assets.Add(AssetObject);
+			ContentBrowserModule.Get().SyncBrowserToAssets(Assets);
+		}
+
+		return FReply::Handled();
+	}
+}
+
+#undef LOCTEXT_NAMESPACE
diff --git a/Engine/Plugins/VirtualProduction/DMX/DMXEngine/Source/DMXEditor/Private/Widgets/Monitors/SDMXConflictMonitorActiveObjectRow.h b/Engine/Plugins/VirtualProduction/DMX/DMXEngine/Source/DMXEditor/Private/Widgets/Monitors/SDMXConflictMonitorActiveObjectRow.h
new file mode 100644
--- /dev/null
+++ b/Engine/Plugins/VirtualProduction/DMX/DMXEngine/Source/DMXEditor/Private/Widgets/Monitors/SDMXConflictMonitorActiveObjectRow.h
@@ -0,0 +1,42 @@
+// Copyright Epic Games, Inc. All Rights Reserved.
+
+#pragma once
+
+#include "Widgets/Views/SListView.h"
+#include "Widgets/Views/STableRow.h"
+
+class SRichTextBlock;
+
+namespace UE::DMX
+{
+	class FDMXConflictMonitorActiveObjectItem;
+
+	/** A row in the Conflict Monitor's Active Object list */
+	class SDMXConflictMonitorActiveObjectRow
+		: public SMultiColumnTableRow<TSharedPtr<FDMXConflictMonitorActiveObjectItem>>
+	{
+	public:
+		SLATE_BEGIN_ARGS(SDMXConflictMonitorActiveObjectRow)
+			{}
+
+		SLATE_END_ARGS()
+
+		/** Constructs this widget */
+		void Construct(const FArguments& InArgs, const TSharedRef<STableViewBase>& OwnerTable, const TSharedRef<FDMXConflictMonitorActiveObjectItem>& InActiveObjectItem);
+
+	protected:
+		//~ Begin SMultiColumnTableRow interface
+		virtual TSharedRef<SWidget> GenerateWidgetForColumn(const FName& ColumnName) override;
+		//~ End SMultiColumnTableRow interface
+
+	private:
+		/** Called when the open asset hyperlink was clicked */
+		void OnOpenAssetClicked();
+
+		/** Called when the show in content browser button was clicked */
+		FReply OnShowInContentBrowserClicked(const FGeometry& MyGeometry, const FPointerEvent& MouseEvent);
+
+		/** The active object item */
+		TSharedPtr<FDMXConflictMonitorActiveObjectItem> ActiveObjectItem;
+	};
+}
diff --git a/Engine/Plugins/VirtualProduction/DMX/DMXEngine/Source/DMXEditor/Private/Widgets/Monitors/SDMXConflictMonitorToolbar.cpp b/Engine/Plugins/VirtualProduction/DMX/DMXEngine/Source/DMXEditor/Private/Widgets/Monitors/SDMXConflictMonitorToolbar.cpp
--- a/Engine/Plugins/VirtualProduction/DMX/DMXEngine/Source/DMXEditor/Private/Widgets/Monitors/SDMXConflictMonitorToolbar.cpp
+++ b/Engine/Plugins/VirtualProduction/DMX/DMXEngine/Source/DMXEditor/Private/Widgets/Monitors/SDMXConflictMonitorToolbar.cpp
@@ -7,11 +7,11 @@
 #include "DMXEditorStyle.h"
 #include "Framework/Commands/UICommandList.h"
 #include "Framework/MultiBox/MultiBoxBuilder.h"
-#include "Widgets/Layout/SWrapBox.h"
-#include "Widgets/SBoxPanel.h"
+#include "Internationalization/Text.h"
 #include "Widgets/Images/SThrobber.h"
 #include "Widgets/Input/SNumericEntryBox.h"
-
+#include "Widgets/Layout/SWrapBox.h"
+#include "Widgets/SBoxPanel.h"
 
 #define LOCTEXT_NAMESPACE "SDMXConflictMonitorToolbar"
 
@@ -48,6 +48,7 @@ namespace UE::DMX
 
 					// Throbber 
 					+ SHorizontalBox::Slot()
+					.HAlign(HAlign_Left)
 					.VAlign(VAlign_Center)
 					.AutoWidth()
 					.Padding(2.f, 0.f, 6.f, 0.f)
@@ -58,12 +59,13 @@ namespace UE::DMX
 
 					// Scanning Info 
 					+ SHorizontalBox::Slot()
+					.HAlign(HAlign_Left)
 					.VAlign(VAlign_Center)
-					.AutoWidth()
+					.FillWidth(1.f)
 					.Padding(2.f, 0.f)
 					[
 						SNew(STextBlock)
-						.Text(LOCTEXT("ScanningText", "Scanning for Conflicts.."))
+						.Text(LOCTEXT("ScanningText", "Scanning.."))
 					]
 				];
 			}
@@ -86,6 +88,8 @@ namespace UE::DMX
 		CommandList = InCommandList;
 		OnDepthChanged = InArgs._OnDepthChanged;
 
+		CachedTimeGameThread = InArgs._TimeGameThread.Get();
+
 		ChildSlot
 		[
 			SNew(SHorizontalBox)
@@ -205,6 +209,47 @@ namespace UE::DMX
 		}
 		ToolbarBuilder.EndSection();
 
+		ToolbarBuilder.BeginSection("Stats");
+		{
+			const FName TutorialHighlight = NAME_None;
+			constexpr bool bSearchable = true;
+
+			ToolbarBuilder.AddWidget(
+				SNew(SHorizontalBox)
+
+				+ SHorizontalBox::Slot()
+				.FillWidth(1.f)
+				[
+					SNew(SBorder)
+					.HAlign(HAlign_Right)
+					.VAlign(VAlign_Center)
+					.BorderImage(FAppStyle::GetBrush("NoBorder"))
+					[
+						SNew(STextBlock)
+						.MinDesiredWidth(100.f)
+						.Text_Lambda([TimeGameThreadAttribute = InArgs._TimeGameThread, this]()
+							{			
+								const double DeltaTime = FSlateApplication::Get().GetDeltaTime();
+
+								const double TimeGameThread = TimeGameThreadAttribute.Get();
+
+								CachedTimeGameThread = 	CachedTimeGameThread + (TimeGameThread - CachedTimeGameThread) * (DeltaTime / 1.0);
+
+								FNumberFormattingOptions Options;
+								Options.SetMaximumFractionalDigits(2);
+								Options.SetMinimumFractionalDigits(2);
+
+								const FText ValueText = FText::AsNumber(CachedTimeGameThread, &Options);
+								return FText::Format(LOCTEXT("PerformanceText", "CPU Usage (GameThread): {0}ms"), ValueText);
+							})
+					]
+				],
+				TutorialHighlight,
+				bSearchable,
+				HAlign_Right);
+		}
+		ToolbarBuilder.EndSection();
+
 		return ToolbarBuilder.MakeWidget();
 	}
 
diff --git a/Engine/Plugins/VirtualProduction/DMX/DMXEngine/Source/DMXEditor/Private/Widgets/Monitors/SDMXConflictMonitorToolbar.h b/Engine/Plugins/VirtualProduction/DMX/DMXEngine/Source/DMXEditor/Private/Widgets/Monitors/SDMXConflictMonitorToolbar.h
--- a/Engine/Plugins/VirtualProduction/DMX/DMXEngine/Source/DMXEditor/Private/Widgets/Monitors/SDMXConflictMonitorToolbar.h
+++ b/Engine/Plugins/VirtualProduction/DMX/DMXEngine/Source/DMXEditor/Private/Widgets/Monitors/SDMXConflictMonitorToolbar.h
@@ -26,11 +26,15 @@ namespace UE::DMX
 	public:
 		SLATE_BEGIN_ARGS(SDMXConflictMonitorToolbar)
 			: _StatusInfo(EDMXConflictMonitorStatusInfo::Idle)
+			, _TimeGameThread(0.0)
 			{}
 
 			/** The status of the monitor */
 			SLATE_ATTRIBUTE(EDMXConflictMonitorStatusInfo, StatusInfo)
 
+			/** The game thread time */
+			SLATE_ATTRIBUTE(double, TimeGameThread)
+
 			/** Broadcast when the depth changed */
 			SLATE_EVENT(FSimpleDelegate, OnDepthChanged)
 
@@ -61,6 +65,9 @@ namespace UE::DMX
 		/** The current depth */
 		uint8 Depth = 3;
 
+		/** Cached time on game thread */
+		double CachedTimeGameThread = 0.0;
+
 		// Slate args
 		TAttribute<EDMXConflictMonitorStatusInfo> StatusInfo;
 		FSimpleDelegate OnDepthChanged;
diff --git a/Engine/Plugins/VirtualProduction/DMX/DMXProtocol/Source/DMXProtocol/Internal/IO/DMXConflictMonitor.h b/Engine/Plugins/VirtualProduction/DMX/DMXProtocol/Source/DMXProtocol/Internal/IO/DMXConflictMonitor.h
--- a/Engine/Plugins/VirtualProduction/DMX/DMXProtocol/Source/DMXProtocol/Internal/IO/DMXConflictMonitor.h
+++ b/Engine/Plugins/VirtualProduction/DMX/DMXProtocol/Source/DMXProtocol/Internal/IO/DMXConflictMonitor.h
@@ -93,6 +93,9 @@ namespace UE::DMX
 		/** Returns current conflicts, locking. */
 		TMap<FName, TArray<TSharedRef<FDMXMonitoredOutboundDMXData>>> GetOutboundConflictsSynchronous() const;
 
+		/** Returns currently monitored data. This contains non-conflicting data as well. */
+		TArray<TSharedRef<FDMXMonitoredOutboundDMXData>> GetMonitoredOutboundData() const { return MonitoredOutboundData; }
+
 		/** Adds a trace to the queue stack. */
 		void TraceUser(const FMinimalName& User);
 
@@ -102,6 +105,9 @@ namespace UE::DMX
 		/** Monitors outbound DMX. Can be called from the game thread only. */
 		void MonitorOutboundDMX(const TSharedRef<FDMXOutputPort>& InOutputPort, int32 InLocalUniverseID, const TMap<int32, uint8>& InChannelToValueMap);
 
+		/** Returns the time the monitor consumes on the game thread */
+		double GetTimeGameThread() const { return TimeGameThread; }
+
 	private:
 		/** Removes a user from the monitor. If all user objects left, the monitor is disabled. */
 		static void RemoveUser(const FName& UserName);
@@ -115,15 +121,15 @@ namespace UE::DMX
 		/** Current outbound data */
 		TArray<TSharedRef<FDMXMonitoredOutboundDMXData>> MonitoredOutboundData;
 
-		/** Traced instigators and the current data as a map */
-		TMap<FName, TArray<TSharedRef<FDMXMonitoredOutboundDMXData>>> Snapshots;
-
 		/** An array of conflicting traces */
 		TMap<FName, TArray<TSharedRef<FDMXMonitoredOutboundDMXData>>>  Conflicts;
 
 		/** The last frame number that was handled on the monitor thread */
 		uint64 FrameNumber = 0;
 
+		/** The time the monitor uses on the game thread */
+		double TimeGameThread = 0.0;
+
 		/** Users of the monitor */
 		static TArray<FName> UserNames;
 
diff --git a/Engine/Plugins/VirtualProduction/DMX/DMXProtocol/Source/DMXProtocol/Private/IO/DMXConflictMonitor.cpp b/Engine/Plugins/VirtualProduction/DMX/DMXProtocol/Source/DMXProtocol/Private/IO/DMXConflictMonitor.cpp
--- a/Engine/Plugins/VirtualProduction/DMX/DMXProtocol/Source/DMXProtocol/Private/IO/DMXConflictMonitor.cpp
+++ b/Engine/Plugins/VirtualProduction/DMX/DMXProtocol/Source/DMXProtocol/Private/IO/DMXConflictMonitor.cpp
@@ -8,7 +8,9 @@
 #include "Algo/Find.h"
 #include "DMXProtocolConstants.h"
 #include "IO/DMXOutputPort.h"
+#include "DMXStats.h"
 
+DECLARE_CYCLE_STAT(TEXT("DMX Conflict Monitor"), STAT_DMXConflictMonitor, STATGROUP_DMX);
 
 namespace UE::DMX
 {
@@ -109,7 +111,10 @@ namespace UE::DMX
 
 	void FDMXConflictMonitor::MonitorOutboundDMX(const TSharedRef<FDMXOutputPort>& InOutputPort, int32 InLocalUniverseID, const TMap<int32, uint8>& InChannelToValueMap)
 	{
+		SCOPE_CYCLE_COUNTER(STAT_DMXConflictMonitor)
+
 		check(IsInGameThread());
+		const double StartTime = FPlatformTime::Seconds();
 
 		// Only do work when the frame switched
 		if (FrameNumber != GFrameNumber)
@@ -168,13 +173,18 @@ namespace UE::DMX
 		}
 
 		FrameNumber = GFrameNumber;
+
+		const double EndTime = FPlatformTime::Seconds();
+		TimeGameThread = (EndTime - StartTime) * 1000.0;
 	}
 
 	void FDMXConflictMonitor::RemoveUser(const FName& UserName)
 	{
 		UserNames.Remove(UserName);
-		if (UserNames.IsEmpty())
+		if (UserNames.IsEmpty() &&
+			ensureMsgf(Instance.IsValid(), TEXT("Unexpected tried to remove a user from conflict monitor, but the conflict monitor instance is invalid.")))
 		{
+			Instance->MonitoredOutboundData.Reset();
 			Instance.Reset();
 		}
 	}
diff --git a/Engine/Plugins/VirtualProduction/DMX/DMXProtocol/Source/DMXProtocol/Private/IO/DMXTrace.cpp b/Engine/Plugins/VirtualProduction/DMX/DMXProtocol/Source/DMXProtocol/Private/IO/DMXTrace.cpp
--- a/Engine/Plugins/VirtualProduction/DMX/DMXProtocol/Source/DMXProtocol/Private/IO/DMXTrace.cpp
+++ b/Engine/Plugins/VirtualProduction/DMX/DMXProtocol/Source/DMXProtocol/Private/IO/DMXTrace.cpp
@@ -6,7 +6,7 @@
 #include "DMXStats.h"
 
 
-DECLARE_CYCLE_STAT(TEXT("DMX trace"), STAT_DMXTrace, STATGROUP_DMX);
+DECLARE_CYCLE_STAT(TEXT("DMX Trace"), STAT_DMXTrace, STATGROUP_DMX);
 
 
 #if WITH_EDITOR
