DMX - Fix Static Analysis -  'FixtureNode' pointer was utilized before it was verified against nullptr

diff --git a/Engine/Plugins/VirtualProduction/DMX/DMXEngine/Source/DMXEditor/Private/Exporters/DMXMVRExporter.cpp b/Engine/Plugins/VirtualProduction/DMX/DMXEngine/Source/DMXEditor/Private/Exporters/DMXMVRExporter.cpp
--- a/Engine/Plugins/VirtualProduction/DMX/DMXEngine/Source/DMXEditor/Private/Exporters/DMXMVRExporter.cpp
+++ b/Engine/Plugins/VirtualProduction/DMX/DMXEngine/Source/DMXEditor/Private/Exporters/DMXMVRExporter.cpp
@@ -491,9 +491,13 @@ namespace UE::DMX
 			}
 
 			UDMXMVRFixtureNode* FixtureNode = GeneralSceneDescription->WriteFixturePatch(*FixturePatch);
-			UsedMVRUUIDs.Add(FixtureNode->UUID);
+			if (!FixtureNode)
+			{
+				continue;
+			}
 
-			if (FixtureNode && ExportOptions->bUseTransformsFromLevel)
+			UsedMVRUUIDs.Add(FixtureNode->UUID);
+			if (ExportOptions->bUseTransformsFromLevel)
 			{
 				FixtureNode->SetTransformAbsolute(Actor->GetTransform());
 			}
