DMX: Add physical units and physical value ranges to Fixture Type Functions (editor only)

diff --git a/Engine/Plugins/VirtualProduction/DMX/DMXEngine/Source/DMXEditor/Private/Customizations/DMXFixtureFunctionCustomization.cpp b/Engine/Plugins/VirtualProduction/DMX/DMXEngine/Source/DMXEditor/Private/Customizations/DMXFixtureFunctionCustomization.cpp
new file mode 100644
--- /dev/null
+++ b/Engine/Plugins/VirtualProduction/DMX/DMXEngine/Source/DMXEditor/Private/Customizations/DMXFixtureFunctionCustomization.cpp
@@ -0,0 +1,43 @@
+// Copyright Epic Games, Inc. All Rights Reserved.
+
+#include "DMXFixtureFunctionCustomization.h"
+
+#include "DetailCategoryBuilder.h"
+#include "DetailLayoutBuilder.h"
+#include "IDetailChildrenBuilder.h"
+#include "Library/DMXEntityFixtureType.h"
+#include "PropertyHandle.h"
+
+namespace UE::DMX
+{
+	TSharedRef<IPropertyTypeCustomization> FDMXFixtureFunctionCustomization::MakeInstance()
+	{
+		return MakeShared<FDMXFixtureFunctionCustomization>();
+	}
+
+	void FDMXFixtureFunctionCustomization::CustomizeChildren(TSharedRef<IPropertyHandle> InPropertyHandle, IDetailChildrenBuilder& InChildBuilder, IPropertyTypeCustomizationUtils& CustomizationUtils)
+	{
+		uint32 NumChildren;
+		if (InPropertyHandle->GetNumChildren(NumChildren) != FPropertyAccess::Success)
+		{
+			return;
+		}
+
+		for (uint32 ChildIndex = 0; ChildIndex < NumChildren; ChildIndex++)
+		{
+			const TSharedPtr<IPropertyHandle> ChildHandle = InPropertyHandle->GetChildHandle(ChildIndex);
+
+			// Hide physical properties, they're editor only and confusing in blueprints
+			if (!ChildHandle.IsValid() ||
+				ChildHandle->GetProperty()->GetFName() == FDMXFixtureFunction::GetPhysicalDefaultValuePropertyName() ||
+				ChildHandle->GetProperty()->GetFName() == FDMXFixtureFunction::GetPhysicalUnitPropertyName() ||
+				ChildHandle->GetProperty()->GetFName() == FDMXFixtureFunction::GetPhysicalFromPropertyName() ||
+				ChildHandle->GetProperty()->GetFName() == FDMXFixtureFunction::GetPhysicalToPropertyName())
+			{
+				continue;
+			}
+
+			InChildBuilder.AddProperty(ChildHandle.ToSharedRef());
+		}
+	}
+}
diff --git a/Engine/Plugins/VirtualProduction/DMX/DMXEngine/Source/DMXEditor/Private/Customizations/DMXFixtureFunctionCustomization.h b/Engine/Plugins/VirtualProduction/DMX/DMXEngine/Source/DMXEditor/Private/Customizations/DMXFixtureFunctionCustomization.h
new file mode 100644
--- /dev/null
+++ b/Engine/Plugins/VirtualProduction/DMX/DMXEngine/Source/DMXEditor/Private/Customizations/DMXFixtureFunctionCustomization.h
@@ -0,0 +1,24 @@
+// Copyright Epic Games, Inc. All Rights Reserved.
+
+#pragma once
+
+#include "IPropertyTypeCustomization.h"
+
+namespace UE::DMX
+{
+	/**
+	 * Customization for the FDMXFixtureFunction struct
+	 */
+	class FDMXFixtureFunctionCustomization
+		: public IPropertyTypeCustomization
+	{
+	public:
+		/** Creates an instance of this property type customization */
+		static TSharedRef<IPropertyTypeCustomization> MakeInstance();
+
+		//~ IPropertyTypeCustomization interface begin
+		virtual void CustomizeHeader(TSharedRef<IPropertyHandle> InPropertyHandle, FDetailWidgetRow& InHeaderRow, IPropertyTypeCustomizationUtils& CustomizationUtils) override {};
+		virtual void CustomizeChildren(TSharedRef<IPropertyHandle> InPropertyHandle, IDetailChildrenBuilder& InChildBuilder, IPropertyTypeCustomizationUtils& CustomizationUtils) override;
+		//~ IPropertyTypeCustomization interface end
+	};
+}
diff --git a/Engine/Plugins/VirtualProduction/DMX/DMXEngine/Source/DMXEditor/Private/DMXEditorModule.cpp b/Engine/Plugins/VirtualProduction/DMX/DMXEngine/Source/DMXEditor/Private/DMXEditorModule.cpp
--- a/Engine/Plugins/VirtualProduction/DMX/DMXEngine/Source/DMXEditor/Private/DMXEditorModule.cpp
+++ b/Engine/Plugins/VirtualProduction/DMX/DMXEngine/Source/DMXEditor/Private/DMXEditorModule.cpp
@@ -11,6 +11,7 @@
 #include "Customizations/DMXEntityFixtureTypeDetails.h"
 #include "Customizations/DMXEntityReferenceCustomization.h"
 #include "Customizations/DMXFixtureCategoryCustomization.h"
+#include "Customizations/DMXFixtureFunctionCustomization.h"
 #include "Customizations/DMXFixtureSignalFormatCustomization.h"
 #include "Customizations/DMXLibraryPortReferencesCustomization.h"
 #include "Customizations/DMXMVRSceneActorDetails.h"
@@ -349,7 +350,9 @@ void FDMXEditorModule::RegisterClassCustomizations()
 }
 
 void FDMXEditorModule::RegisterPropertyTypeCustomizations()
-{
+{	
+	using namespace UE::DMX;
+
 	// Property type customization for the EDMXPixelMappingDistribution enum
 	RegisterCustomPropertyTypeLayout("EDMXPixelMappingDistribution", 
 		FOnGetPropertyTypeCustomizationInstance::CreateStatic(&FDMXPixelMappingDistributionCustomization::MakeInstance)
@@ -360,6 +363,11 @@ void FDMXEditorModule::RegisterPropertyTypeCustomizations()
 		FOnGetPropertyTypeCustomizationInstance::CreateStatic(&FDMXFixtureSignalFormatCustomization::MakeInstance)
 	);
 
+	// Property type customization for the FDMXFixtureFunction struct
+	RegisterCustomPropertyTypeLayout(FDMXFixtureFunction::StaticStruct()->GetFName(),
+		FOnGetPropertyTypeCustomizationInstance::CreateStatic(&FDMXFixtureFunctionCustomization::MakeInstance)
+	);
+
 	// Property type customization for the FDMXAttributeName struct
 	RegisterCustomPropertyTypeLayout(FDMXAttributeName::StaticStruct()->GetFName(), 
 		FOnGetPropertyTypeCustomizationInstance::CreateStatic(&FDMXAttributeNameCustomization::MakeInstance)
@@ -374,9 +382,6 @@ void FDMXEditorModule::RegisterPropertyTypeCustomizations()
 	RegisterCustomPropertyTypeLayout(FDMXEntityFixtureTypeRef::StaticStruct()->GetFName(), 
 		FOnGetPropertyTypeCustomizationInstance::CreateStatic(&FDMXEntityReferenceCustomization::MakeInstance)
 	);
-	RegisterCustomPropertyTypeLayout(FDMXEntityFixturePatchRef::StaticStruct()->GetFName(), 
-		FOnGetPropertyTypeCustomizationInstance::CreateStatic(&FDMXEntityReferenceCustomization::MakeInstance)
-	);
 
 	// Customization for the FDMXLibraryPortReferences struct
 	RegisterCustomPropertyTypeLayout(FDMXLibraryPortReferences::StaticStruct()->GetFName(), 
diff --git a/Engine/Plugins/VirtualProduction/DMX/DMXEngine/Source/DMXEditor/Private/DMXPhysicalUnitToDefaultValueRange.h b/Engine/Plugins/VirtualProduction/DMX/DMXEngine/Source/DMXEditor/Private/DMXPhysicalUnitToDefaultValueRange.h
new file mode 100644
--- /dev/null
+++ b/Engine/Plugins/VirtualProduction/DMX/DMXEngine/Source/DMXEditor/Private/DMXPhysicalUnitToDefaultValueRange.h
@@ -0,0 +1,44 @@
+// Copyright Epic Games, Inc. All Rights Reserved.
+
+#pragma once
+
+#include "Containers/Map.h"
+#include "GDTF/AttributeDefinitions/DMXGDTFPhysicalUnit.h"
+
+namespace UE::DMX
+{
+	struct FDMXPhysicalUnitToDefaultValueRange
+	{
+		/** Returns the value range for specified physical unit */
+		static const TPair<double, double>& GetValueRange(EDMXGDTFPhysicalUnit PhysicalUnit);
+
+	private:
+		/** A map to list all GDTF compliant Physical Units and their matching Physical Value default ranges. */
+		static inline const TMap<EDMXGDTFPhysicalUnit, TPair<double, double>> PhysicalUnitToValueRangeMap =
+		{
+			{ EDMXGDTFPhysicalUnit::None, {0.0, 1.0} },
+			{ EDMXGDTFPhysicalUnit::Percent, {0.0, 100.0} },
+			{ EDMXGDTFPhysicalUnit::Length, {0.1, 100.0} },
+			{ EDMXGDTFPhysicalUnit::Mass, {0.1, 100.0} },
+			{ EDMXGDTFPhysicalUnit::Time, {0.001, 60.0} },
+			{ EDMXGDTFPhysicalUnit::Temperature, {1000.0, 20000.0} },
+			{ EDMXGDTFPhysicalUnit::LuminousIntensity, {0.0, 100000.0} },
+			{ EDMXGDTFPhysicalUnit::Angle, {-120.0, 120.0} },
+			{ EDMXGDTFPhysicalUnit::Force, {0.1, 1000.0} },
+			{ EDMXGDTFPhysicalUnit::Frequency, {0.0, 30.0} },
+			{ EDMXGDTFPhysicalUnit::Current, {0.01, 100.0} },
+			{ EDMXGDTFPhysicalUnit::Voltage, {5.0, 400.0} },
+			{ EDMXGDTFPhysicalUnit::Power, {1.0, 10000.0} },
+			{ EDMXGDTFPhysicalUnit::Energy, {0.001, 100000.0} },
+			{ EDMXGDTFPhysicalUnit::Area, {0.0001, 10.0} },
+			{ EDMXGDTFPhysicalUnit::Volume, {0.001, 1.0} },
+			{ EDMXGDTFPhysicalUnit::Speed, {0.1, 10.0} },
+			{ EDMXGDTFPhysicalUnit::Acceleration, {0.1, 50.0} },
+			{ EDMXGDTFPhysicalUnit::AngularSpeed, {1.0, 360.0} },
+			{ EDMXGDTFPhysicalUnit::AngularAccc, {1.0, 1000.0} },
+			{ EDMXGDTFPhysicalUnit::WaveLength, {100.0, 1000000.0} },
+			{ EDMXGDTFPhysicalUnit::ColorComponent, {0.0, 100.0} },
+			{ EDMXGDTFPhysicalUnit::MaxEnumValue, {0.0, 0.0} }
+		};
+	};
+}
diff --git a/Engine/Plugins/VirtualProduction/DMX/DMXEngine/Source/DMXEditor/Private/DMXPhysicalUnitToValueRange.cpp b/Engine/Plugins/VirtualProduction/DMX/DMXEngine/Source/DMXEditor/Private/DMXPhysicalUnitToValueRange.cpp
new file mode 100644
--- /dev/null
+++ b/Engine/Plugins/VirtualProduction/DMX/DMXEngine/Source/DMXEditor/Private/DMXPhysicalUnitToValueRange.cpp
@@ -0,0 +1,11 @@
+// Copyright Epic Games, Inc. All Rights Reserved.
+
+#include "DMXPhysicalUnitToDefaultValueRange.h"
+
+namespace UE::DMX
+{
+	const TPair<double, double>& FDMXPhysicalUnitToDefaultValueRange::GetValueRange(EDMXGDTFPhysicalUnit PhysicalUnit)
+	{
+		return PhysicalUnitToValueRangeMap.FindChecked(PhysicalUnit);
+	}
+}
diff --git a/Engine/Plugins/VirtualProduction/DMX/DMXEngine/Source/DMXEditor/Private/Exporters/DMXFixtureTypeToGDTFConverter.cpp b/Engine/Plugins/VirtualProduction/DMX/DMXEngine/Source/DMXEditor/Private/Exporters/DMXFixtureTypeToGDTFConverter.cpp
--- a/Engine/Plugins/VirtualProduction/DMX/DMXEngine/Source/DMXEditor/Private/Exporters/DMXFixtureTypeToGDTFConverter.cpp
+++ b/Engine/Plugins/VirtualProduction/DMX/DMXEngine/Source/DMXEditor/Private/Exporters/DMXFixtureTypeToGDTFConverter.cpp
@@ -68,27 +68,28 @@ namespace UE::DMX::GDTF
 
 	void FDMXFixtureTypeToGDTFConverter::CreateAttributeDefinitions(const UDMXEntityFixtureType& UnrealFixtureType, const TSharedRef<FDMXGDTFFixtureType>& GDTFFixtureType)
 	{
-		TArray<FName> AttributeNames;
+		// Map of attribute names with their function. If function is nullptr, it's a matrix cell attribute
+		TMap<FName, const FDMXFixtureFunction*> AttributeNameToFunctionPtrMap;
 		for (const FDMXFixtureMode& Mode : UnrealFixtureType.Modes)
 		{
 			for (const FDMXFixtureFunction& Function : Mode.Functions)
 			{
-				AttributeNames.Add(Function.Attribute.Name);
+				AttributeNameToFunctionPtrMap.Add(Function.Attribute.Name, &Function);
 			}
 
 			if (Mode.bFixtureMatrixEnabled)
 			{
 				for (const FDMXFixtureCellAttribute& MatrixAttribute : Mode.FixtureMatrixConfig.CellAttributes)
 				{
-					AttributeNames.Add(MatrixAttribute.Attribute.Name);
+					AttributeNameToFunctionPtrMap.Add(MatrixAttribute.Attribute.Name, nullptr);
 				}
 			}
 		}
 
 		GDTFFixtureType->AttributeDefinitions = MakeShared<FDMXGDTFAttributeDefinitions>(GDTFFixtureType);
-		for (const FName& AttributeName : AttributeNames)
+		for (const TTuple<FName, const FDMXFixtureFunction*>& AttributeNameToFunctionPtrPair : AttributeNameToFunctionPtrMap)
 		{
-			const FName GDTFAttributeName = FDMXUnrealToGDTFAttributeConversion::ConvertUnrealToGDTFAttribute(AttributeName);
+			const FName GDTFAttributeName = FDMXUnrealToGDTFAttributeConversion::ConvertUnrealToGDTFAttribute(AttributeNameToFunctionPtrPair.Key);
 			
 			const FName PrettyName = FDMXUnrealToGDTFAttributeConversion::GetPrettyFromGDTFAttribute(GDTFAttributeName);
 			const FName FeatureGroupName = FDMXUnrealToGDTFAttributeConversion::GetFeatureGroupForGDTFAttribute(GDTFAttributeName);
@@ -137,9 +138,11 @@ namespace UE::DMX::GDTF
 			const TSharedRef<FDMXGDTFAttribute> GDTFAttribute = MakeShared<FDMXGDTFAttribute>(GDTFFixtureType->AttributeDefinitions.ToSharedRef());
 			GDTFFixtureType->AttributeDefinitions->Attributes.Add(GDTFAttribute);
 
+			const FDMXFixtureFunction* FunctionPtr = AttributeNameToFunctionPtrPair.Value;
+
 			GDTFAttribute->Name = GDTFAttributeName;
 			GDTFAttribute->Pretty = PrettyName.ToString();
-			GDTFAttribute->PhysicalUnit = EDMXGDTFPhysicalUnit::None;
+			GDTFAttribute->PhysicalUnit = FunctionPtr ? FunctionPtr->GetPhysicalUnit() : EDMXGDTFPhysicalUnit::None;
 			GDTFAttribute->Feature = FeatureGroupName.ToString() + TEXT(".") + FeatureName.ToString();
 		}
 	}
@@ -322,6 +325,9 @@ namespace UE::DMX::GDTF
 		const FDMXGDTFDMXValue Default = FDMXGDTFDMXValue(UnrealFunction.DefaultValue, UnrealFunction.GetNumChannels());
 		ChannelFunction->Default = Default;
 		ChannelFunction->DMXFrom = 0;
+
+		ChannelFunction->PhysicalFrom = UnrealFunction.GetPhysicalFrom();
+		ChannelFunction->PhysicalTo = UnrealFunction.GetPhysicalTo();
 	}
 
 	void FDMXFixtureTypeToGDTFConverter::CrateChannelFunction(const FDMXFixtureCellAttribute& UnrealCellAttribute, const TSharedRef<FDMXGDTFLogicalChannel>& GDTFLogicalChannel, const FString& GDTFAttribute)
diff --git a/Engine/Plugins/VirtualProduction/DMX/DMXEngine/Source/DMXEditor/Private/Factories/DMXGDTFToFixtureTypeConverter.cpp b/Engine/Plugins/VirtualProduction/DMX/DMXEngine/Source/DMXEditor/Private/Factories/DMXGDTFToFixtureTypeConverter.cpp
--- a/Engine/Plugins/VirtualProduction/DMX/DMXEngine/Source/DMXEditor/Private/Factories/DMXGDTFToFixtureTypeConverter.cpp
+++ b/Engine/Plugins/VirtualProduction/DMX/DMXEngine/Source/DMXEditor/Private/Factories/DMXGDTFToFixtureTypeConverter.cpp
@@ -12,6 +12,7 @@
 #include "DMXProtocolSettings.h"
 #include "DMXRuntimeUtils.h"
 #include "DMXZipper.h"
+#include "GDTF/AttributeDefinitions/DMXGDTFAttribute.h"
 #include "GDTF/DMXGDTFDescription.h"
 #include "GDTF/DMXGDTFFixtureType.h"
 #include "GDTF/DMXModes/DMXGDTFChannelFunction.h"
@@ -47,6 +48,10 @@ namespace UE::DMX::GDTF
 			bool bLSBMode = false;
 			uint32 NumCells = 1;
 			uint32 DefaultValue = 0;
+
+			EDMXGDTFPhysicalUnit PhysicalUnit = EDMXGDTFPhysicalUnit::None;
+			double PhysicalFrom = 0.0;
+			double PhysicalTo = 1.0;
 		};
 
 		static bool GetChannelProperties(
@@ -79,6 +84,10 @@ namespace UE::DMX::GDTF
 			OutChannelProperties.DefaultValue = GetDefaultValue();
 			GetDataType(OutChannelProperties.SignalFormat, OutChannelProperties.bLSBMode);
 
+			OutChannelProperties.PhysicalUnit = GetPhysicalUnit();
+			OutChannelProperties.PhysicalFrom = GetPhysicalFrom();
+			OutChannelProperties.PhysicalTo = GetPhysicalTo();
+
 			const int32 NumGeometries = CountGeometries();
 			if (NumGeometries > 1)
 			{				
@@ -189,6 +198,22 @@ namespace UE::DMX::GDTF
 			return 0;
 		}
 
+		EDMXGDTFPhysicalUnit GetPhysicalUnit() const
+		{
+			TSharedPtr<FDMXGDTFAttribute> Attribute = ChannelFunction.IsValid() ? ChannelFunction->ResolveAttribute() : nullptr;
+			return Attribute.IsValid() ? Attribute->PhysicalUnit : EDMXGDTFPhysicalUnit::None;
+		} 
+
+		double GetPhysicalFrom() const
+		{
+			return ChannelFunction.IsValid() ? ChannelFunction->PhysicalFrom : 0.0;
+		}
+
+		double GetPhysicalTo() const
+		{
+			return ChannelFunction.IsValid() ? ChannelFunction->PhysicalTo : 1.0;
+		}
+
 		bool GetMinOffsetOfDMXChannel(uint32& OutChannelOffset) const
 		{
 			if (DMXChannelNode->Offset.IsEmpty())
@@ -409,6 +434,8 @@ namespace UE::DMX::GDTF
 				Function.bUseLSBMode = ChannelProperties.bLSBMode;
 				Function.DataType = ChannelProperties.SignalFormat;
 				Function.DefaultValue = ChannelProperties.DefaultValue;
+				Function.SetPhysicalUnit(ChannelProperties.PhysicalUnit);
+				Function.SetPhysicalValueRange(ChannelProperties.PhysicalFrom, ChannelProperties.PhysicalTo);
 
 				Mode.Functions.Add(Function);
 
diff --git a/Engine/Plugins/VirtualProduction/DMX/DMXEngine/Source/DMXEditor/Private/Widgets/FixtureType/SDMXFixtureFunctionEditor.cpp b/Engine/Plugins/VirtualProduction/DMX/DMXEngine/Source/DMXEditor/Private/Widgets/FixtureType/SDMXFixtureFunctionEditor.cpp
--- a/Engine/Plugins/VirtualProduction/DMX/DMXEngine/Source/DMXEditor/Private/Widgets/FixtureType/SDMXFixtureFunctionEditor.cpp
+++ b/Engine/Plugins/VirtualProduction/DMX/DMXEngine/Source/DMXEditor/Private/Widgets/FixtureType/SDMXFixtureFunctionEditor.cpp
@@ -2,20 +2,101 @@
 
 #include "SDMXFixtureFunctionEditor.h"
 
+#include "DMXAttributeToDefaultPhyiscalProperties.h"
+#include "DMXConversions.h"
 #include "DMXEditor.h"
 #include "DMXFixtureTypeSharedData.h"
-#include "Library/DMXEntityFixtureType.h"
-
+#include "DMXPhysicalUnitToDefaultValueRange.h"
 #include "IStructureDetailsView.h"
+#include "Library/DMXEntityFixtureType.h"
+#include "Modules/ModuleManager.h"
 #include "PropertyEditorModule.h"
 #include "ScopedTransaction.h"
-#include "Modules/ModuleManager.h"
 #include "Widgets/Layout/SBox.h"
 #include "Widgets/Text/STextBlock.h"
 
 
 #define LOCTEXT_NAMESPACE "SDMXFixtureFunctionEditor"
 
+namespace UE::DMX::SDMXFixtureFunctionEditor::Private
+{
+	FDMXFixtureFunction* GetFixtureFunction(UDMXEntityFixtureType& FixtureType, int32 ModeIndex, int32 FunctionIndex)
+	{
+		if (!ensureMsgf(FixtureType.Modes.IsValidIndex(ModeIndex), TEXT("Trying to clamp Function Default Value, but Mode Index is not valid.")))
+		{
+			return nullptr;
+		}
+
+		FDMXFixtureMode& Mode = FixtureType.Modes[ModeIndex];
+		if (!ensureMsgf(Mode.Functions.IsValidIndex(FunctionIndex), TEXT("Trying to clamp Function Default Value, but Function Index is not valid.")))
+		{
+			return nullptr;
+		}
+
+		return &Mode.Functions[FunctionIndex];
+	}
+
+	/**
+	 * Clamps the Default Value of the Function by its Data Type
+	 *
+	 * @param ModeIndex						The Index of the Mode in which the Functions reside
+	 * @param FunctionIndex					The Index of the Function for which the Default Value is clamped
+	 */
+	void ClampFunctionDefautValueByDataType(UDMXEntityFixtureType& FixtureType, int32 ModeIndex, int32 FunctionIndex)
+	{
+		if (FDMXFixtureFunction* FunctionPtr = GetFixtureFunction(FixtureType, ModeIndex, FunctionIndex))
+		{
+			const uint32 SafeDefaultValue = FMath::Min(static_cast<int64>(TNumericLimits<uint32>::Max()), FunctionPtr->DefaultValue);
+			const uint32 ClampedDefaultValue = FDMXConversions::ClampValueBySignalFormat(SafeDefaultValue, FunctionPtr->DataType);
+
+			FunctionPtr->DefaultValue = ClampedDefaultValue;
+			FunctionPtr->UpdatePhysicalDefaultValue();
+		}
+	}
+
+	/**
+	 * Clamps the Physical Default Value of the Function by its Physical Unit^.
+	 *
+	 * @param ModeIndex						The Index of the Mode in which the Functions reside
+	 * @param FunctionIndex					The Index of the Function for which the Physical Default Value is clamped
+	 */
+	void ClampFunctionPhysicalDefautValueByPhysicalUnit(UDMXEntityFixtureType& FixtureType, int32 ModeIndex, int32 FunctionIndex)
+	{
+		if (FDMXFixtureFunction* FunctionPtr = GetFixtureFunction(FixtureType, ModeIndex, FunctionIndex))
+		{
+			const double PhysicalMin = FMath::Min(FunctionPtr->GetPhysicalFrom(), FunctionPtr->GetPhysicalTo());
+			const double PhysicalMax = FMath::Max(FunctionPtr->GetPhysicalFrom(), FunctionPtr->GetPhysicalTo());
+
+			const double ClampedPhysicalDefaultValue = FMath::Clamp(FunctionPtr->GetPhysicalDefaultValue(), PhysicalMin, PhysicalMax);
+			FunctionPtr->SetPhysicalDefaultValue(ClampedPhysicalDefaultValue);
+		}
+	}
+
+	/**
+	 * Updates the Physical Default Value of the Function from its Physical Unit.
+	 *
+	 * @param ModeIndex						The Index of the Mode in which the Functions reside
+	 * @param FunctionIndex					The Index of the Function for which the Physical Default Value is clamped
+	 */
+	void UpdateFunctionDefaultValueFromPhysicalUnit(UDMXEntityFixtureType& FixtureType, int32 ModeIndex, int32 FunctionIndex)
+	{
+		if (FDMXFixtureFunction* FunctionPtr = GetFixtureFunction(FixtureType, ModeIndex, FunctionIndex))
+		{
+			const TPair<double, double>& PhysicalUnitValueRange = FDMXPhysicalUnitToDefaultValueRange::GetValueRange(FunctionPtr->GetPhysicalUnit());
+
+			FunctionPtr->SetPhysicalValueRange(PhysicalUnitValueRange.Key, PhysicalUnitValueRange.Value);
+		}
+	}
+
+	void UpdateFunctionPhysicalUnitFromAttribute(UDMXEntityFixtureType& FixtureType, int32 ModeIndex, int32 FunctionIndex)
+	{
+		if (FDMXFixtureFunction* FunctionPtr = GetFixtureFunction(FixtureType, ModeIndex, FunctionIndex))
+		{
+			FDMXAttributeToDefaultPhyiscalProperties::ResetToDefaultPhysicalProperties(*FunctionPtr);
+		}
+	}
+}
+
 void SDMXFixtureFunctionEditor::Construct(const FArguments& InArgs, const TSharedRef<FDMXEditor>& InDMXEditor)
 {
 	WeakDMXEditor = InDMXEditor;
@@ -89,11 +170,28 @@ void SDMXFixtureFunctionEditor::NotifyPreChange(FProperty* PropertyAboutToChange
 		}
 		else if (PropertyName == GET_MEMBER_NAME_CHECKED(FDMXFixtureFunction, DataType))
 		{
-			Transaction = MakeUnique<FScopedTransaction>(LOCTEXT("SetDataTypeTransaction", "Data Type of Function"));
+			Transaction = MakeUnique<FScopedTransaction>(LOCTEXT("SetDataTypeTransaction", "Set Data Type of Function"));
+		}
+		else if (PropertyName == GET_MEMBER_NAME_CHECKED(FDMXFixtureFunction, DefaultValue))
+		{
+			Transaction = MakeUnique<FScopedTransaction>(LOCTEXT("SetDefaultValueTransaction", "Set Default Value of Function"));
 		}
 		else if (PropertyName == GET_MEMBER_NAME_CHECKED(FDMXFixtureFunction, DefaultValue))
 		{
-			Transaction = MakeUnique<FScopedTransaction>(LOCTEXT("SetDefaultValueTransaction", "Default Value of Function"));
+			Transaction = MakeUnique<FScopedTransaction>(LOCTEXT("SetDefaultValueTransaction", "Set Attribute of Function"));
+		}
+		else if (PropertyName == FDMXFixtureFunction::GetPhysicalDefaultValuePropertyName())
+		{
+			Transaction = MakeUnique<FScopedTransaction>(LOCTEXT("SetPhysicalValueTransaction", "Set Physical Default Value of Function"));
+		}
+		else if (PropertyName == FDMXFixtureFunction::GetPhysicalFromPropertyName() ||
+			PropertyName == FDMXFixtureFunction::GetPhysicalToPropertyName())
+		{
+			Transaction = MakeUnique<FScopedTransaction>(LOCTEXT("SetPhysicalValueTransaction", "Set Physical Range of Function"));
+		}
+		else if (PropertyName == FDMXFixtureFunction::GetPhysicalUnitPropertyName())
+		{
+			Transaction = MakeUnique<FScopedTransaction>(LOCTEXT("SetPhysicalValueTransaction", "Set Physical Unit of Function"));
 		}
 
 		FixtureType->Modify();
@@ -103,6 +201,8 @@ void SDMXFixtureFunctionEditor::NotifyPreChange(FProperty* PropertyAboutToChange
 
 void SDMXFixtureFunctionEditor::NotifyPostChange(const FPropertyChangedEvent& PropertyChangedEvent, FProperty* PropertyThatChanged)
 {
+	using namespace UE::DMX::SDMXFixtureFunctionEditor::Private;
+
 	UDMXEntityFixtureType* FixtureType = WeakFixtureType.Get();
 	FDMXFixtureFunction* FunctionBeingEditedPtr = GetFunctionBeingEdited();
 
@@ -126,11 +226,26 @@ void SDMXFixtureFunctionEditor::NotifyPostChange(const FPropertyChangedEvent& Pr
 			else if (PropertyName == GET_MEMBER_NAME_CHECKED(FDMXFixtureFunction, DataType))
 			{
 				FixtureType->AlignFunctionChannels(ModeIndex);
-				FixtureType->ClampFunctionDefautValueByDataType(ModeIndex, FunctionIndex);
+				ClampFunctionDefautValueByDataType(*FixtureType, ModeIndex, FunctionIndex);
 			}
 			else if (PropertyName == GET_MEMBER_NAME_CHECKED(FDMXFixtureFunction, DefaultValue))
 			{
-				FixtureType->ClampFunctionDefautValueByDataType(ModeIndex, FunctionIndex);
+				ClampFunctionDefautValueByDataType(*FixtureType, ModeIndex, FunctionIndex);
+			}
+			else if (PropertyName == GET_MEMBER_NAME_CHECKED(FDMXFixtureFunction, Attribute) ||
+				PropertyName == GET_MEMBER_NAME_CHECKED(FDMXAttribute, Name))
+			{
+				UpdateFunctionPhysicalUnitFromAttribute(*FixtureType, ModeIndex, FunctionIndex);
+			}
+			else if (PropertyName == FDMXFixtureFunction::GetPhysicalDefaultValuePropertyName() || 
+					PropertyName == FDMXFixtureFunction::GetPhysicalFromPropertyName() ||
+					PropertyName == FDMXFixtureFunction::GetPhysicalToPropertyName())
+			{
+				ClampFunctionPhysicalDefautValueByPhysicalUnit(*FixtureType, ModeIndex, FunctionIndex);
+			}
+			else if (PropertyName == FDMXFixtureFunction::GetPhysicalUnitPropertyName())
+			{
+				UpdateFunctionDefaultValueFromPhysicalUnit(*FixtureType, ModeIndex, FunctionIndex);
 			}
 
 			FPropertyChangedEvent ObjectPropertyChangedEvent(PropertyChangedEvent);
diff --git a/Engine/Plugins/VirtualProduction/DMX/DMXEngine/Source/DMXEditor/Private/Widgets/FixtureType/SDMXFixtureFunctionEditor.h b/Engine/Plugins/VirtualProduction/DMX/DMXEngine/Source/DMXEditor/Private/Widgets/FixtureType/SDMXFixtureFunctionEditor.h
--- a/Engine/Plugins/VirtualProduction/DMX/DMXEngine/Source/DMXEditor/Private/Widgets/FixtureType/SDMXFixtureFunctionEditor.h
+++ b/Engine/Plugins/VirtualProduction/DMX/DMXEngine/Source/DMXEditor/Private/Widgets/FixtureType/SDMXFixtureFunctionEditor.h
@@ -2,23 +2,19 @@
 
 #pragma once
 
-#include "Widgets/SCompoundWidget.h"
-
-#include "CoreMinimal.h"
 #include "Misc/NotifyHook.h"
 #include "ScopedTransaction.h"
+#include "Widgets/SCompoundWidget.h"
 
 class FDMXEditor;
 struct FDMXFixtureFunction;
 struct FDMXFixtureMode;
 class FDMXFixtureTypeSharedData;
-class UDMXEntityFixtureType;
-
 struct FPropertyAndParent;
 class IStructureDetailsView;
 class STextBlock;
 class SWidget;
-
+class UDMXEntityFixtureType;
 
 /** Editor for a single Function struct inside a UDMXEntityFixtureType. */
 class SDMXFixtureFunctionEditor
@@ -52,7 +48,7 @@ private:
 
 	/** Returns true if the properties are visible */
 	bool IsPropertyVisible(const FPropertyAndParent& PropertyAndParent) const;
-
+	
 	/** Returns the Mode currently being edited, or nullptr if the current selection cannot be resolved. */
 	FDMXFixtureMode* GetModeBeingEdited() const;
 
diff --git a/Engine/Plugins/VirtualProduction/DMX/DMXEngine/Source/DMXRuntime/Internal/DMXAttributeToDefaultPhyiscalProperties.h b/Engine/Plugins/VirtualProduction/DMX/DMXEngine/Source/DMXRuntime/Internal/DMXAttributeToDefaultPhyiscalProperties.h
new file mode 100644
--- /dev/null
+++ b/Engine/Plugins/VirtualProduction/DMX/DMXEngine/Source/DMXRuntime/Internal/DMXAttributeToDefaultPhyiscalProperties.h
@@ -0,0 +1,45 @@
+// Copyright Epic Games, Inc. All Rights Reserved.
+
+#pragma once
+
+// Editor only
+#if WITH_EDITOR 
+
+#include "GDTF/AttributeDefinitions/DMXGDTFPhysicalUnit.h"
+
+struct FDMXFixtureFunction;
+
+namespace UE::DMX
+{
+	struct DMXRUNTIME_API FDMXAttributeToDefaultPhyiscalProperties
+	{
+		/** Returns the value range for specified physical unit */
+		static void ResetToDefaultPhysicalProperties(FDMXFixtureFunction& InOutFunction);
+
+	private:
+		struct FDefaultPhyicalProperties
+		{
+			FDefaultPhyicalProperties(EDMXGDTFPhysicalUnit InPhysicalUnit, double InPhysicalFrom, double InPhysicalTo)
+				: PhysicalUnit(InPhysicalUnit)
+				, PhysicalFrom(InPhysicalFrom)
+				, PhysicalTo(InPhysicalTo)
+			{}
+				
+			const EDMXGDTFPhysicalUnit PhysicalUnit = EDMXGDTFPhysicalUnit::None;
+			const double PhysicalFrom = 0.0;
+			const double PhysicalTo = 1.0;
+		};
+
+		/** A map to list all GDTF compliant Physical Units and their matching Physical Value default ranges. */
+		static inline const TMap<FName, FDefaultPhyicalProperties> AttributeNameToDefaultPhysicalProperitesMap
+		{ 
+			{ "Zoom", FDefaultPhyicalProperties(EDMXGDTFPhysicalUnit::Angle, 1.0, 120.0) },
+
+			{ "Pan", FDefaultPhyicalProperties(EDMXGDTFPhysicalUnit::Angle, -120.0, 120.0) },
+			{ "Tilt", FDefaultPhyicalProperties(EDMXGDTFPhysicalUnit::Angle, -120.0, 120.0) },
+
+		};
+	};
+}
+
+#endif // WITH_EDITOR
diff --git a/Engine/Plugins/VirtualProduction/DMX/DMXEngine/Source/DMXRuntime/Private/DMXAttributeToDefaultPhyiscalProperties.cpp b/Engine/Plugins/VirtualProduction/DMX/DMXEngine/Source/DMXRuntime/Private/DMXAttributeToDefaultPhyiscalProperties.cpp
new file mode 100644
--- /dev/null
+++ b/Engine/Plugins/VirtualProduction/DMX/DMXEngine/Source/DMXRuntime/Private/DMXAttributeToDefaultPhyiscalProperties.cpp
@@ -0,0 +1,22 @@
+// Copyright Epic Games, Inc. All Rights Reserved.
+
+#if WITH_EDITOR 
+
+#include "DMXAttributeToDefaultPhyiscalProperties.h"
+
+#include "Library/DMXEntityFixtureType.h"
+
+namespace UE::DMX
+{
+	void FDMXAttributeToDefaultPhyiscalProperties::ResetToDefaultPhysicalProperties(FDMXFixtureFunction& InOutFunction)
+	{
+		const FDefaultPhyicalProperties* AttributeNameToDefaultPhysicalProperitesPtr = AttributeNameToDefaultPhysicalProperitesMap.Find(InOutFunction.Attribute.Name);
+		if (AttributeNameToDefaultPhysicalProperitesPtr)
+		{
+			InOutFunction.SetPhysicalUnit(AttributeNameToDefaultPhysicalProperitesPtr->PhysicalUnit);
+			InOutFunction.SetPhysicalValueRange(AttributeNameToDefaultPhysicalProperitesPtr->PhysicalFrom, AttributeNameToDefaultPhysicalProperitesPtr->PhysicalTo);
+		}
+	}
+}
+
+#endif // WITH_EDITOR
diff --git a/Engine/Plugins/VirtualProduction/DMX/DMXEngine/Source/DMXRuntime/Private/DMXRuntimeMainStreamObjectVersion.h b/Engine/Plugins/VirtualProduction/DMX/DMXEngine/Source/DMXRuntime/Private/DMXRuntimeMainStreamObjectVersion.h
--- a/Engine/Plugins/VirtualProduction/DMX/DMXEngine/Source/DMXRuntime/Private/DMXRuntimeMainStreamObjectVersion.h
+++ b/Engine/Plugins/VirtualProduction/DMX/DMXEngine/Source/DMXRuntime/Private/DMXRuntimeMainStreamObjectVersion.h
@@ -39,6 +39,9 @@ struct FDMXRuntimeMainStreamObjectVersion
 		// Rename GDTF import assets with invalid names
 		DMXFixGDTFImportAssetsWithInvalidNames,
 
+		// Upgrade Fixture Functions to make use of pysical properties where required (e.g. Pan, Tilt, Zoom)
+		DMXUpgradeFixtureFunctionsToUsePhysicalProperties,
+
 		VersionPlusOne,
 		LatestVersion = VersionPlusOne - 1
 	};
diff --git a/Engine/Plugins/VirtualProduction/DMX/DMXEngine/Source/DMXRuntime/Private/Library/DMXEntityFixtureType.cpp b/Engine/Plugins/VirtualProduction/DMX/DMXEngine/Source/DMXRuntime/Private/Library/DMXEntityFixtureType.cpp
--- a/Engine/Plugins/VirtualProduction/DMX/DMXEngine/Source/DMXRuntime/Private/Library/DMXEntityFixtureType.cpp
+++ b/Engine/Plugins/VirtualProduction/DMX/DMXEngine/Source/DMXRuntime/Private/Library/DMXEntityFixtureType.cpp
@@ -2,6 +2,7 @@
 
 #include "Library/DMXEntityFixtureType.h"
 
+#include "DMXAttributeToDefaultPhyiscalProperties.h"
 #include "DMXConversions.h"
 #include "DMXProtocolSettings.h"
 #include "DMXRuntimeMainStreamObjectVersion.h"
@@ -16,11 +17,77 @@
 
 #define LOCTEXT_NAMESPACE "DMXEntityFixtureType"
 
+void FDMXFixtureFunction::PostSerialize(const FArchive& Ar)
+{
+#if WITH_EDITOR
+	if (Ar.IsLoading())
+	{
+		UpdatePhysicalDefaultValue();
+
+		// Upgrade Fixture Functions that were added before physical properties were added to use these properties
+		if (Ar.CustomVer(FDMXRuntimeMainStreamObjectVersion::GUID) < FDMXRuntimeMainStreamObjectVersion::DMXUpgradeFixtureFunctionsToUsePhysicalProperties)
+		{
+			using namespace UE::DMX;
+			FDMXAttributeToDefaultPhyiscalProperties::ResetToDefaultPhysicalProperties(*this);
+		}
+	}
+#endif
+}
+
 int32 FDMXFixtureFunction::GetLastChannel() const
 {
 	return Channel + GetNumChannels() - 1;
 }
 
+#if WITH_EDITOR
+void FDMXFixtureFunction::SetPhysicalDefaultValue(double InPhysicalDefaultValue)
+{
+	const double PhysicalMin = FMath::Min(PhysicalFrom, PhysicalTo);
+	const double PhysicalMax = FMath::Max(PhysicalFrom, PhysicalTo);
+
+	PhysicalDefaultValue = FMath::Clamp(InPhysicalDefaultValue, PhysicalMin, PhysicalMax);
+
+	const double PhysicalValueRange = PhysicalTo > PhysicalFrom ? PhysicalTo - PhysicalFrom : PhysicalFrom - PhysicalTo;
+	const double RelativePhysicalDefaultValue = PhysicalDefaultValue > PhysicalFrom ? PhysicalDefaultValue - PhysicalFrom : PhysicalFrom - PhysicalDefaultValue;
+	const double NormalizedPhysicalDefaultValue = FMath::IsNearlyZero(PhysicalValueRange) ? 0.0 : RelativePhysicalDefaultValue / PhysicalValueRange;
+
+	const uint32 MaxValue = FDMXConversions::GetSignalFormatMaxValue(DataType);
+	DefaultValue = FMath::Lerp(0, MaxValue, NormalizedPhysicalDefaultValue);
+}
+#endif // WITH_EDITOR
+
+#if WITH_EDITOR
+void FDMXFixtureFunction::SetPhysicalValueRange(double InPhysicalFrom, double InPhysicalTo)
+{
+	const double PhysicalValueRange = PhysicalTo > PhysicalFrom ? PhysicalTo - PhysicalFrom : PhysicalFrom - PhysicalTo;
+	const double RelativePhysicalDefaultValue = PhysicalDefaultValue > PhysicalFrom ? PhysicalDefaultValue - PhysicalFrom : PhysicalFrom - PhysicalDefaultValue;
+	const double NormalizedPhysicalDefaultValue = FMath::IsNearlyZero(PhysicalValueRange) ? 0.0 : RelativePhysicalDefaultValue / PhysicalValueRange;
+
+	PhysicalFrom = InPhysicalFrom;
+	PhysicalTo = InPhysicalTo;
+
+	const double NewPhysicalDefaultValue = FMath::Lerp(PhysicalFrom, PhysicalTo, NormalizedPhysicalDefaultValue);
+	
+	// Set both the default and the physical default value
+	SetPhysicalDefaultValue(NewPhysicalDefaultValue);
+}
+#endif // WITH_EDITOR
+
+#if WITH_EDITOR
+void FDMXFixtureFunction::UpdatePhysicalDefaultValue()
+{
+	const uint32 MaxValue = FDMXConversions::GetSignalFormatMaxValue(DataType);
+	const double NormalizedDefaultValue = static_cast<double>(DefaultValue) / MaxValue;
+
+	const double NewPhysicalValue = FMath::Lerp(PhysicalFrom, PhysicalTo, NormalizedDefaultValue);
+
+	constexpr double Tolerance = 0.01;
+	const double RoundedPhysicalValue = FMath::RoundToDouble(NewPhysicalValue / Tolerance) * Tolerance;
+
+	PhysicalDefaultValue = RoundedPhysicalValue;
+}
+#endif // WITH_EDITOR
+
 FDMXFixtureMatrix::FDMXFixtureMatrix()
 {
 	// Add a default cell attribute	
@@ -236,6 +303,7 @@ void UDMXEntityFixtureType::Serialize(FArchive& Ar)
 		}
 
 #if WITH_EDITOR
+		// Updgrade to use a soft object ptr for GDTF
 		if (Ar.CustomVer(FDMXRuntimeMainStreamObjectVersion::GUID) < FDMXRuntimeMainStreamObjectVersion::DMXImportGDTFIsASoftObjectPtr)
 		{
 			PRAGMA_DISABLE_DEPRECATION_WARNINGS
@@ -247,6 +315,7 @@ void UDMXEntityFixtureType::Serialize(FArchive& Ar)
 			}
 			PRAGMA_ENABLE_DEPRECATION_WARNINGS
 		}
+
 #endif // WITH_EDITOR
 	}
 }
@@ -964,22 +1033,24 @@ void UDMXEntityFixtureType::SetFunctionStartingChannel(int32 InModeIndex, int32
 	}
 }
 
-void UDMXEntityFixtureType::ClampFunctionDefautValueByDataType(int32 ModeIndex, int32 FunctionToRemoveIndex)
+void UDMXEntityFixtureType::ClampFunctionDefautValueByDataType(int32 ModeIndex, int32 FunctionIndex)
 {
+	// DEPRECATED 5.5
 	if (ensureMsgf(Modes.IsValidIndex(ModeIndex), TEXT("Trying to clamp Function Default Value, but Mode Index is not valid.")))
 	{
 		FDMXFixtureMode& Mode = Modes[ModeIndex];
 
-		if (ensureMsgf(Mode.Functions.IsValidIndex(FunctionToRemoveIndex), TEXT("Trying to clamp Function Default Value, but Function Index is not valid.")))
+		if (ensureMsgf(Mode.Functions.IsValidIndex(FunctionIndex), TEXT("Trying to clamp Function Default Value, but Function Index is not valid.")))
 		{
-			FDMXFixtureFunction& Function = Mode.Functions[FunctionToRemoveIndex];
+			FDMXFixtureFunction& Function = Mode.Functions[FunctionIndex];
 			const uint32 SafeDefaultValue = FMath::Min(static_cast<int64>(TNumericLimits<uint32>::Max()), Function.DefaultValue);
 			const uint32 ClampedDefaultValue = FDMXConversions::ClampValueBySignalFormat(SafeDefaultValue, Function.DataType);
 
-			if (Function.DefaultValue != ClampedDefaultValue)
-			{
-				Function.DefaultValue = ClampedDefaultValue;
-			}
+			Function.DefaultValue = ClampedDefaultValue;
+
+#if WITH_EDITOR
+			Function.UpdatePhysicalDefaultValue();
+#endif
 		}
 	}
 }
diff --git a/Engine/Plugins/VirtualProduction/DMX/DMXEngine/Source/DMXRuntime/Public/Library/DMXEntityFixtureType.h b/Engine/Plugins/VirtualProduction/DMX/DMXEngine/Source/DMXRuntime/Public/Library/DMXEntityFixtureType.h
--- a/Engine/Plugins/VirtualProduction/DMX/DMXEngine/Source/DMXRuntime/Public/Library/DMXEntityFixtureType.h
+++ b/Engine/Plugins/VirtualProduction/DMX/DMXEngine/Source/DMXRuntime/Public/Library/DMXEntityFixtureType.h
@@ -4,6 +4,7 @@
 
 #include "DMXAttribute.h"
 #include "DMXProtocolTypes.h"
+#include "GDTF/AttributeDefinitions/DMXGDTFPhysicalUnit.h"
 #include "Library/DMXEntity.h"
 #include "Library/DMXEntityReference.h"
 #include "Modulators/DMXModulator.h"
@@ -47,21 +48,49 @@ struct DMXRUNTIME_API FDMXFixtureFunction
 	/** Constructor */
 	FDMXFixtureFunction()
 		: Attribute()
-		, FunctionName()
-		, Description()
-		, DefaultValue(0)
-		, Channel(1)
-		, ChannelOffset_DEPRECATED(0)
-		, DataType(EDMXFixtureSignalFormat::E8Bit)
-		, bUseLSBMode(false)
 	{}
 
+	/** Implementing Serialize to convert UClass to FFieldClass */
+	void PostSerialize(const FArchive& Ar);
+
 	/** Returns the number of channels the function spans, according to its data type */
 	FORCEINLINE uint8 GetNumChannels() const { return static_cast<uint8>(DataType) + 1; }
 
 	/** Returns the last channel of the Function */
 	int32 GetLastChannel() const;
 
+#if WITH_EDITOR
+	/** Gets the Physical Value of the Function */
+	double GetPhysicalDefaultValue() const { return PhysicalDefaultValue; }
+
+	/** The Physical Unit this Physical Value is based on */
+	EDMXGDTFPhysicalUnit GetPhysicalUnit() const { return PhysicalUnit; }
+
+	/** The starting value of the Physical Value range, based on the Physical Unit */
+	double GetPhysicalFrom() const { return PhysicalFrom; }
+
+	/** The ending value of the Physical Value range, based on the Physical Unit */
+	double GetPhysicalTo() const { return PhysicalTo; }
+
+	/** Sets the Physical Default Value of the Function. */
+	void SetPhysicalUnit(EDMXGDTFPhysicalUnit NewPhysicalUnit) { PhysicalUnit = NewPhysicalUnit; }
+
+	/** Sets the Physical Default Value of the Function. */
+	void SetPhysicalDefaultValue(double InPhysicalDefaultValue);
+
+	/** Sets the Physical Default Value range of the Function. */
+	void SetPhysicalValueRange(double InPhysicalFrom, double InPhysicalTo);
+
+	/** Updated the Physica Default Value of the Function by the Default Value. */
+	void UpdatePhysicalDefaultValue();
+
+	// Property Name getters
+	FORCEINLINE static FName GetPhysicalDefaultValuePropertyName() { return GET_MEMBER_NAME_CHECKED(FDMXFixtureFunction, PhysicalDefaultValue); }
+	FORCEINLINE static FName GetPhysicalUnitPropertyName() { return GET_MEMBER_NAME_CHECKED(FDMXFixtureFunction, PhysicalUnit); }
+	FORCEINLINE static FName GetPhysicalFromPropertyName() { return GET_MEMBER_NAME_CHECKED(FDMXFixtureFunction, PhysicalFrom); }
+	FORCEINLINE static FName GetPhysicalToPropertyName() { return GET_MEMBER_NAME_CHECKED(FDMXFixtureFunction, PhysicalTo); }
+#endif // WITH_EDITOR
+
 	/**
 	 * The Attribute name to map this Function to.
 	 * This is used to easily find the Function in Blueprints, using an Attribute
@@ -78,17 +107,17 @@ struct DMXRUNTIME_API FDMXFixtureFunction
 	UPROPERTY(EditAnywhere, BlueprintReadWrite, meta = (DisplayPriority = "20"), Category = "Function Settings")
 	FString Description;
 
-	/** The Default Value of the function */
+	/** The Default DMX Value of the function */
 	UPROPERTY(EditAnywhere, BlueprintReadWrite, meta = (DisplayPriority = "30"), Category = "Function Settings")
-	int64 DefaultValue;
+	int64 DefaultValue = 0;
 
 	/** This function's starting channel (use editor above to make changes) */
 	UPROPERTY(VisibleAnywhere, BlueprintReadWrite, meta = (DisplayName = "Channel Assignment", DisplayPriority = "2"), Category = "Function Settings")
-	int32 Channel;
+	int32 Channel = 1;
 
 	/** DEPRECATED 5.0. Instead the 'Channel' property is EditAnywhere so any function can be assigned freely */
 	UPROPERTY(meta = (DeprecatedProperty, DeprecationMessage = "Deprecated since the Channel property can be set in the DMX Library Editor."))
-	int32 ChannelOffset_DEPRECATED;
+	int32 ChannelOffset_DEPRECATED = 0;
 
 	/** This function's data type. Defines the used number of channels (bytes) */
 	UPROPERTY(EditAnywhere, BlueprintReadWrite, meta = (DisplayPriority = "5"), Category = "Function Settings")
@@ -108,8 +137,37 @@ struct DMXRUNTIME_API FDMXFixtureFunction
 	 */
 	UPROPERTY(EditAnywhere, BlueprintReadWrite, meta = (DisplayName = "Use LSB Mode", DisplayPriority = "29"), Category = "Function Settings")
 	bool bUseLSBMode = false;
+
+private:
+#if WITH_EDITORONLY_DATA
+	/** The Physical Value used by default, based on the Physical Unit */
+	UPROPERTY(EditAnywhere, Transient, Category = "Physical Properties")
+	double PhysicalDefaultValue = 0.0;
+
+	/** The Physical Unit this Physical Value is based on */
+	UPROPERTY(EditAnywhere, Category = "Physical Properties")
+	EDMXGDTFPhysicalUnit PhysicalUnit = EDMXGDTFPhysicalUnit::None;
+
+	/** The starting value of the Physical Value range, based on the Physical Unit */
+	UPROPERTY(EditAnywhere, Category = "Physical Properties")
+	double PhysicalFrom = 0.0;
+
+	/** The ending value of the Physical Value range, based on the Physical Unit */
+	UPROPERTY(EditAnywhere, Category = "Physical Properties")
+	double PhysicalTo = 1.0;
+#endif // WITH_EDITORONLY_DATA
+};
+
+template<>
+struct TStructOpsTypeTraits<FDMXFixtureFunction> : public TStructOpsTypeTraitsBase2<FDMXFixtureFunction>
+{
+	enum
+	{
+		WithPostSerialize = true,
+	};
 };
 
+
 USTRUCT(BlueprintType)
 struct DMXRUNTIME_API FDMXFixtureCellAttribute
 {
@@ -468,10 +526,10 @@ public:
 	 * Clamps the Default Value of the Function by its Data Type
 	 *
 	 * @param ModeIndex						The Index of the Mode in which the Functions reside
-	 * @param FunctionToRemoveIndex			The Index of the Function for which the Default Value is clamped
+	 * @param FunctionIndex					The Index of the Function for which the Default Value is clamped
 	 */
-	void ClampFunctionDefautValueByDataType(int32 ModeIndex, int32 FunctionToRemoveIndex);
-
+	UE_DEPRECATED(5.5, "Removed as physical values were introduced to FDMXFixtureFunction (editor only) . Please handle the default value of the function per use case.")
+	void ClampFunctionDefautValueByDataType(int32 ModeIndex, int32 FunctionIndex);
 
 	// Fixture Matrix related
 public:
