DMX: Fix Text conflict from LOCTEXT macro for namespace for SDMXFixtureFunctionEditor.cpp

diff --git a/Engine/Plugins/VirtualProduction/DMX/DMXEngine/Source/DMXEditor/Private/Widgets/FixtureType/SDMXFixtureFunctionEditor.cpp b/Engine/Plugins/VirtualProduction/DMX/DMXEngine/Source/DMXEditor/Private/Widgets/FixtureType/SDMXFixtureFunctionEditor.cpp
--- a/Engine/Plugins/VirtualProduction/DMX/DMXEngine/Source/DMXEditor/Private/Widgets/FixtureType/SDMXFixtureFunctionEditor.cpp
+++ b/Engine/Plugins/VirtualProduction/DMX/DMXEngine/Source/DMXEditor/Private/Widgets/FixtureType/SDMXFixtureFunctionEditor.cpp
@@ -178,7 +178,7 @@ void SDMXFixtureFunctionEditor::NotifyPreChange(FProperty* PropertyAboutToChange
 		}
 		else if (PropertyName == GET_MEMBER_NAME_CHECKED(FDMXFixtureFunction, DefaultValue))
 		{
-			Transaction = MakeUnique<FScopedTransaction>(LOCTEXT("SetDefaultValueTransaction", "Set Attribute of Function"));
+			Transaction = MakeUnique<FScopedTransaction>(LOCTEXT("SetAttributeTransaction", "Set Attribute of Function"));
 		}
 		else if (PropertyName == FDMXFixtureFunction::GetPhysicalDefaultValuePropertyName())
 		{
@@ -187,11 +187,11 @@ void SDMXFixtureFunctionEditor::NotifyPreChange(FProperty* PropertyAboutToChange
 		else if (PropertyName == FDMXFixtureFunction::GetPhysicalFromPropertyName() ||
 			PropertyName == FDMXFixtureFunction::GetPhysicalToPropertyName())
 		{
-			Transaction = MakeUnique<FScopedTransaction>(LOCTEXT("SetPhysicalValueTransaction", "Set Physical Range of Function"));
+			Transaction = MakeUnique<FScopedTransaction>(LOCTEXT("SetPhysicalRangeTransaction", "Set Physical Range of Function"));
 		}
 		else if (PropertyName == FDMXFixtureFunction::GetPhysicalUnitPropertyName())
 		{
-			Transaction = MakeUnique<FScopedTransaction>(LOCTEXT("SetPhysicalValueTransaction", "Set Physical Unit of Function"));
+			Transaction = MakeUnique<FScopedTransaction>(LOCTEXT("SetPhysicalUnitTransaction", "Set Physical Unit of Function"));
 		}
 
 		FixtureType->Modify();
