DMX: Improvments and bug fixes for MVR and GDTF

* Add an experimental option to export multi patch fixtures in MVR
* UDMXMVRGeneralSceneDescription now optionally considers a world when writing a DMX Library
* Set correct geometries for GDTF axis geometry and beam geometry
* Restore mistakenly deleted customization for FDMXEntityFixturePatchRef
* Restore new fixture functions default to 8bit
* Fix an ensure being hit when a patch tries to update while fixture type transactions are ongoing

diff --git a/Engine/Plugins/VirtualProduction/DMX/DMXEngine/Source/DMXEditor/Private/DMXEditorModule.cpp b/Engine/Plugins/VirtualProduction/DMX/DMXEngine/Source/DMXEditor/Private/DMXEditorModule.cpp
--- a/Engine/Plugins/VirtualProduction/DMX/DMXEngine/Source/DMXEditor/Private/DMXEditorModule.cpp
+++ b/Engine/Plugins/VirtualProduction/DMX/DMXEngine/Source/DMXEditor/Private/DMXEditorModule.cpp
@@ -383,6 +383,11 @@ void FDMXEditorModule::RegisterPropertyTypeCustomizations()
 		FOnGetPropertyTypeCustomizationInstance::CreateStatic(&FDMXEntityReferenceCustomization::MakeInstance)
 	);
 
+	// Customization for the FDMXEntityFixturePatchRef struct
+	RegisterCustomPropertyTypeLayout(FDMXEntityFixturePatchRef::StaticStruct()->GetFName(),
+		FOnGetPropertyTypeCustomizationInstance::CreateStatic(&FDMXEntityReferenceCustomization::MakeInstance)
+	);
+
 	// Customization for the FDMXLibraryPortReferences struct
 	RegisterCustomPropertyTypeLayout(FDMXLibraryPortReferences::StaticStruct()->GetFName(), 
 		FOnGetPropertyTypeCustomizationInstance::CreateStatic(&FDMXLibraryPortReferencesCustomization::MakeInstance)
diff --git a/Engine/Plugins/VirtualProduction/DMX/DMXEngine/Source/DMXEditor/Private/Exporters/DMXFixtureTypeToGDTFConverter.cpp b/Engine/Plugins/VirtualProduction/DMX/DMXEngine/Source/DMXEditor/Private/Exporters/DMXFixtureTypeToGDTFConverter.cpp
--- a/Engine/Plugins/VirtualProduction/DMX/DMXEngine/Source/DMXEditor/Private/Exporters/DMXFixtureTypeToGDTFConverter.cpp
+++ b/Engine/Plugins/VirtualProduction/DMX/DMXEngine/Source/DMXEditor/Private/Exporters/DMXFixtureTypeToGDTFConverter.cpp
@@ -156,7 +156,7 @@ namespace UE::DMX::GDTF
 			const TSharedRef<FDMXGDTFModel> InstanceModel = MakeShared<FDMXGDTFModel>(GDTFFixtureType);
 			GDTFFixtureType->Models.Add(InstanceModel);
 
-			InstanceModel->Name = "Layers";
+			InstanceModel->Name = FDMXFixtureTypeToGDTFGeometryFactory::CellsModelName;
 			InstanceModel->PrimitiveType = EDMXGDTFModelPrimitiveType::Cube;
 			InstanceModel->Height = 0.01f;
 			InstanceModel->Length = 1.f;
@@ -170,15 +170,17 @@ namespace UE::DMX::GDTF
 
 		FDMXFixtureTypeToGDTFGeometryFactory GeometryFactory(UnrealFixtureType, GDTFFixtureType->GeometryCollect.ToSharedRef());
 
+		ModesWithBaseGeometry = GeometryFactory.GetModesWithBaseGeometry();
 		FunctionsWithControlledGeometry = GeometryFactory.GetFunctionsWithControlledGeometry();
 	}
 
 	void FDMXFixtureTypeToGDTFConverter::CreateDMXModes(const UDMXEntityFixtureType& UnrealFixtureType, const TSharedRef<FDMXGDTFFixtureType>& GDTFFixtureType)
 	{
+		int32 ModeCount = 1;
 		for (const FDMXFixtureMode& UnrealMode : UnrealFixtureType.Modes)
 		{
-			const FDMXFixtureFunctionWithControlledGeometry* ControlledGeometryPtr = Algo::FindBy(FunctionsWithControlledGeometry, &UnrealMode, &FDMXFixtureFunctionWithControlledGeometry::ModePtr);
-			if (!ensureMsgf(ControlledGeometryPtr, TEXT("%hs: Unexpected cannot controlled root geometry for DMX Mode '%s'. Failed to convert mode to GDTF."), __FUNCTION__, *UnrealMode.ModeName))
+			const FDMXFixtureModeWithBaseGeometry* ModeWithBaseGeometryPtr = Algo::FindBy(ModesWithBaseGeometry, &UnrealMode, &FDMXFixtureModeWithBaseGeometry::ModePtr);
+			if (!ensureMsgf(ModeWithBaseGeometryPtr, TEXT("%hs: Unexpected cannot find base geometry for DMX Mode '%s'. Failed to convert mode to GDTF."), __FUNCTION__, *UnrealMode.ModeName))
 			{
 				continue;
 			}
@@ -187,11 +189,13 @@ namespace UE::DMX::GDTF
 			const TSharedRef<FDMXGDTFDMXMode> DMXMode = MakeShared<FDMXGDTFDMXMode>(GDTFFixtureType);
 			GDTFFixtureType->DMXModes.Add(DMXMode);
 
-			DMXMode->Name = *UnrealMode.ModeName;
+			DMXMode->Name = *FString::Printf(TEXT("Mode %i %s"), ModeCount, *UnrealMode.ModeName);
 			DMXMode->Description = TEXT("Unreal Engine generated DMX Mode");
-			DMXMode->Geometry = ControlledGeometryPtr->ControlledGeometry->Name;
+			DMXMode->Geometry = ModeWithBaseGeometryPtr->BaseGeometry->Name;
 
 			CreateDMXChannels(UnrealMode, DMXMode);
+
+			ModeCount++;
 		}
 	}
 
@@ -200,8 +204,8 @@ namespace UE::DMX::GDTF
 		// Create DMX Channels for non-matrix Unreal Functions
 		for (const FDMXFixtureFunction& UnrealFunction : UnrealMode.Functions)
 		{
-			const FDMXFixtureFunctionWithControlledGeometry* ControlledGeometryPtr = Algo::FindBy(FunctionsWithControlledGeometry, &UnrealFunction, &FDMXFixtureFunctionWithControlledGeometry::FunctionPtr);
-			if (!ensureMsgf(ControlledGeometryPtr, TEXT("%hs: Unexpected cannot find geometry for DMX Function '%s'. Failed to convert mode to GDTF."), __FUNCTION__, *UnrealFunction.FunctionName))
+			const FDMXFixtureFunctionWithControlledGeometry* FunctionWithControlledGeometryPtr = Algo::FindBy(FunctionsWithControlledGeometry, &UnrealFunction, &FDMXFixtureFunctionWithControlledGeometry::FunctionPtr);
+			if (!ensureMsgf(FunctionWithControlledGeometryPtr, TEXT("%hs: Unexpected cannot find controlled geometry for DMX Function '%s'. Failed to convert mode to GDTF."), __FUNCTION__, *UnrealFunction.FunctionName))
 			{
 				continue;
 			}
@@ -209,7 +213,7 @@ namespace UE::DMX::GDTF
 			const TSharedRef<FDMXGDTFDMXChannel> DMXChannel = MakeShared<FDMXGDTFDMXChannel>(GDTFDMXMode);
 			GDTFDMXMode->DMXChannels.Add(DMXChannel);
 
-			const FString ControlledGeometryName = ControlledGeometryPtr->ControlledGeometry->Name.ToString();
+			const FString ControlledGeometryName = FunctionWithControlledGeometryPtr->ControlledGeometry->Name.ToString();
 			const FString GDTFAttribute = FDMXUnrealToGDTFAttributeConversion::ConvertUnrealToGDTFAttribute(UnrealFunction.Attribute.Name).ToString();
 			const FString ChannelFunctionName = UnrealFunction.FunctionName;
 
diff --git a/Engine/Plugins/VirtualProduction/DMX/DMXEngine/Source/DMXEditor/Private/Exporters/DMXFixtureTypeToGDTFConverter.h b/Engine/Plugins/VirtualProduction/DMX/DMXEngine/Source/DMXEditor/Private/Exporters/DMXFixtureTypeToGDTFConverter.h
--- a/Engine/Plugins/VirtualProduction/DMX/DMXEngine/Source/DMXEditor/Private/Exporters/DMXFixtureTypeToGDTFConverter.h
+++ b/Engine/Plugins/VirtualProduction/DMX/DMXEngine/Source/DMXEditor/Private/Exporters/DMXFixtureTypeToGDTFConverter.h
@@ -13,6 +13,7 @@ class UDMXEntityFixtureType;
 namespace UE::DMX::GDTF
 {
 	struct FDMXFixtureFunctionWithControlledGeometry;
+	struct FDMXFixtureModeWithBaseGeometry;
 	class FDMXGDTFAttributeDefinitions;
 	class FDMXGDTFDMXChannel;
 	class FDMXGDTFDMXMode;
@@ -60,6 +61,9 @@ namespace UE::DMX::GDTF
 		/** Creates a Channel Function for the Logical Channel, using an Unreal Cell Attribute as the input */
 		void CrateChannelFunction(const FDMXFixtureCellAttribute& UnrealCellAttribute, const TSharedRef<FDMXGDTFLogicalChannel>& GDTFLogicalChannel, const FString& GDTFAttribute);
 
+		/** Map of Modes with their base geometry */
+		TArray<FDMXFixtureModeWithBaseGeometry> ModesWithBaseGeometry;
+
 		/** Map of Functions with the geometry they control */
 		TArray<FDMXFixtureFunctionWithControlledGeometry> FunctionsWithControlledGeometry;
 	};
diff --git a/Engine/Plugins/VirtualProduction/DMX/DMXEngine/Source/DMXEditor/Private/Exporters/DMXFixtureTypeToGDTFGeometryFactory.cpp b/Engine/Plugins/VirtualProduction/DMX/DMXEngine/Source/DMXEditor/Private/Exporters/DMXFixtureTypeToGDTFGeometryFactory.cpp
--- a/Engine/Plugins/VirtualProduction/DMX/DMXEngine/Source/DMXEditor/Private/Exporters/DMXFixtureTypeToGDTFGeometryFactory.cpp
+++ b/Engine/Plugins/VirtualProduction/DMX/DMXEngine/Source/DMXEditor/Private/Exporters/DMXFixtureTypeToGDTFGeometryFactory.cpp
@@ -13,12 +13,17 @@
 
 namespace UE::DMX::GDTF
 {
+	const FName FDMXFixtureTypeToGDTFGeometryFactory::CellsModelName = "Cells";
+
 	const FName FDMXFixtureTypeToGDTFGeometryFactory::BaseGeometryName = "Base";
 	const FName FDMXFixtureTypeToGDTFGeometryFactory::PanGeometryName = "Yoke";
 	const FName FDMXFixtureTypeToGDTFGeometryFactory::HeadGeometryName = "Head";
 	const FName FDMXFixtureTypeToGDTFGeometryFactory::BeamGeometryName = "Beam";
 	const FName FDMXFixtureTypeToGDTFGeometryFactory::MatrixBeamGeometryName = "Instance";
 
+	const FName FDMXFixtureTypeToGDTFGeometryFactory::PanAttributeName = "Pan";
+	const FName FDMXFixtureTypeToGDTFGeometryFactory::TiltAttributeName = "Tilt";
+
 	namespace Internal
 	{
 		FDMXUniqueModeGeometry::FDMXUniqueModeGeometry(const FDMXFixtureMode& Mode)
@@ -31,11 +36,16 @@ namespace UE::DMX::GDTF
 				TotalNumChannels += Function.GetNumChannels();
 			}
 
-			bWithPan = Algo::AnyOf(Mode.Functions, [](const FDMXFixtureFunction& Function) { return Function.Attribute.Name == TEXT("Pan"); });
-			bWithTilt = Algo::AnyOf(Mode.Functions, [](const FDMXFixtureFunction& Function) { return Function.Attribute.Name == TEXT("Tilt"); });
+			bWithPan = Algo::AnyOf(Mode.Functions, [](const FDMXFixtureFunction& Function) { return Function.Attribute.Name == FDMXFixtureTypeToGDTFGeometryFactory::PanAttributeName; });
+			bWithTilt = Algo::AnyOf(Mode.Functions, [](const FDMXFixtureFunction& Function) { return Function.Attribute.Name == FDMXFixtureTypeToGDTFGeometryFactory::TiltAttributeName; });
 		}
 	}
 
+	FDMXFixtureModeWithBaseGeometry::FDMXFixtureModeWithBaseGeometry(const FDMXFixtureMode* InModePtr, TSharedRef<const FDMXGDTFGeometry> InBaseGeometry)
+		: ModePtr(InModePtr)
+		, BaseGeometry(InBaseGeometry)
+	{}
+
 	FDMXFixtureFunctionWithControlledGeometry::FDMXFixtureFunctionWithControlledGeometry(const FDMXFixtureMode* InModePtr, const FDMXFixtureFunction* InFunctionPtr, TSharedRef<const FDMXGDTFGeometry> InControlledGeometry)
 		: ModePtr(InModePtr)
 		, FunctionPtr(InFunctionPtr)
@@ -49,24 +59,83 @@ namespace UE::DMX::GDTF
 		BuildGeometries();
 	}
 
+	TArray<FDMXFixtureModeWithBaseGeometry> FDMXFixtureTypeToGDTFGeometryFactory::GetModesWithBaseGeometry() const
+	{
+		TArray<FDMXFixtureModeWithBaseGeometry> Result;
+
+		for (const FDMXFixtureMode& Mode : FixtureType.Modes)
+		{
+			const FDMXUniqueModeGeometry UniqueModeGeometry(Mode);
+			const TSharedRef<FDMXGDTFGeometry>* BaseGeometryPtr = UniqueModeGeometryToGeometryMap.Find(UniqueModeGeometry);
+			if (!ensureMsgf(BaseGeometryPtr, TEXT("Unexpected cannot find Geometry for Mode %s."), *Mode.ModeName))
+			{
+				return Result;
+			}
+
+			Result.Add(FDMXFixtureModeWithBaseGeometry(&Mode, *BaseGeometryPtr));
+		}
+
+		return Result;
+	}
+
 	TArray<FDMXFixtureFunctionWithControlledGeometry> FDMXFixtureTypeToGDTFGeometryFactory::GetFunctionsWithControlledGeometry() const
 	{
 		TArray<FDMXFixtureFunctionWithControlledGeometry> Result;
 
 		for (const FDMXFixtureMode& Mode : FixtureType.Modes)
 		{
-			FDMXUniqueModeGeometry UniqueModeGeometry(Mode);
-			const TSharedRef<FDMXGDTFGeometry>* GeometryPtr = UniqueModeGeometryToGeometryMap.Find(UniqueModeGeometry);
-
-			if (!ensureMsgf(GeometryPtr, TEXT("Unexpected cannot find Geometry for Mode %s."), *Mode.ModeName))
+			const FDMXUniqueModeGeometry UniqueModeGeometry(Mode);
+			const TSharedRef<FDMXGDTFGeometry>* BaseGeometryPtr = UniqueModeGeometryToGeometryMap.Find(UniqueModeGeometry);
+			if (!ensureMsgf(BaseGeometryPtr, TEXT("Unexpected cannot find Geometry for Mode %s."), *Mode.ModeName))
 			{
 				return Result;
 			}
 
+			const TSharedRef<FDMXGDTFGeometry> BaseGeometry = *BaseGeometryPtr;
+
+			const bool bWithPrimaryAxis = UniqueModeGeometry.bWithPan || UniqueModeGeometry.bWithTilt;
+			const TSharedPtr<FDMXGDTFGeometry> PrimaryAxisGeometry = bWithPrimaryAxis ? BaseGeometry->AxisArray[0] : nullptr;
+
+			const bool bWithSecondaryAxisGeometry = UniqueModeGeometry.bWithPan && UniqueModeGeometry.bWithTilt && PrimaryAxisGeometry.IsValid() && !PrimaryAxisGeometry->AxisArray.IsEmpty();
+			const TSharedPtr<FDMXGDTFGeometry> SecondaryAxisGeometry = bWithSecondaryAxisGeometry ? PrimaryAxisGeometry->AxisArray[0] : nullptr;
+			
+			const TSharedPtr<FDMXGDTFGeometry> BeamGeometry =
+				bWithSecondaryAxisGeometry && !SecondaryAxisGeometry->BeamArray.IsEmpty() ? SecondaryAxisGeometry->BeamArray[0] :
+				bWithPrimaryAxis && !PrimaryAxisGeometry->BeamArray.IsEmpty() ? PrimaryAxisGeometry->BeamArray[0] :
+				!BaseGeometry->BeamArray.IsEmpty() ? BaseGeometry->BeamArray[0] :
+				nullptr;
+
 			Algo::Transform(Mode.Functions, Result,
-				[&Mode, GeometryPtr](const FDMXFixtureFunction& Function)
+				[&Mode, &UniqueModeGeometry, &BaseGeometry, &PrimaryAxisGeometry, &SecondaryAxisGeometry, &BeamGeometry](const FDMXFixtureFunction& Function)
 				{
-					return FDMXFixtureFunctionWithControlledGeometry(&Mode, &Function, *GeometryPtr);
+					// Return axis geometries for pan and tilt
+					const bool bPanFunction = Function.Attribute == PanAttributeName;
+					const bool bTiltFunction = Function.Attribute == TiltAttributeName;
+					if (bPanFunction || bTiltFunction)
+					{
+						const bool bOnlyOneAxis = UniqueModeGeometry.bWithPan != UniqueModeGeometry.bWithTilt;
+						if (PrimaryAxisGeometry.IsValid() && (bOnlyOneAxis || bPanFunction))
+						{
+							return FDMXFixtureFunctionWithControlledGeometry(&Mode, &Function, PrimaryAxisGeometry.ToSharedRef());
+						}
+						else if (SecondaryAxisGeometry.IsValid())
+						{
+							return FDMXFixtureFunctionWithControlledGeometry(&Mode, &Function, SecondaryAxisGeometry.ToSharedRef());
+						}
+
+						ensureMsgf(0, TEXT("Unexpected could not find axis geometry for Pan or Tilt"));
+						return FDMXFixtureFunctionWithControlledGeometry(&Mode, &Function, BaseGeometry);
+					}
+
+					// Other attributes control the beam if available
+					if (BeamGeometry.IsValid())
+					{
+						return FDMXFixtureFunctionWithControlledGeometry(&Mode, &Function, BeamGeometry.ToSharedRef());
+					}
+					else
+					{
+						return FDMXFixtureFunctionWithControlledGeometry(&Mode, &Function, BaseGeometry);
+					}
 				});
 		}
 
@@ -187,14 +256,17 @@ namespace UE::DMX::GDTF
 			for (int32 CellID = 0; CellID < NumCells; CellID++)
 			{
 				const int32 DMXOffset = CellID * CellSize + 1;
-				const FName GeometryName = *FString::Printf(TEXT("Layer_%i"), CellID + 1);
+
+				const int32 Row = CellID / Mode.FixtureMatrixConfig.YCells + 1;
+				const int32 Column = CellID % Mode.FixtureMatrixConfig.YCells + 1;
+				const FName GeometryReferenceName = *FString::Printf(TEXT("Cell_%i_%i"), Row, Column);
 
 				const TSharedRef<FDMXGDTFGeometryReference> GeometryReference = MakeShared<FDMXGDTFGeometryReference>(OuterGeometry);
 				OuterGeometry->GeometryReferenceArray.Add(GeometryReference);
 
-				GeometryReference->Name = GeometryName;
+				GeometryReference->Name = GeometryReferenceName;
 				GeometryReference->Geometry = GetOrCreateMatrixBeamGeometryInstance()->Name;
-				GeometryReference->Model = TEXT("Layers");
+				GeometryReference->Model = CellsModelName;
 
 				const TSharedRef<FDMXGDTFGeometryBreak> Break = MakeShared<FDMXGDTFGeometryBreak>(GeometryReference);
 				Break->DMXBreak = 1;
diff --git a/Engine/Plugins/VirtualProduction/DMX/DMXEngine/Source/DMXEditor/Private/Exporters/DMXFixtureTypeToGDTFGeometryFactory.h b/Engine/Plugins/VirtualProduction/DMX/DMXEngine/Source/DMXEditor/Private/Exporters/DMXFixtureTypeToGDTFGeometryFactory.h
--- a/Engine/Plugins/VirtualProduction/DMX/DMXEngine/Source/DMXEditor/Private/Exporters/DMXFixtureTypeToGDTFGeometryFactory.h
+++ b/Engine/Plugins/VirtualProduction/DMX/DMXEngine/Source/DMXEditor/Private/Exporters/DMXFixtureTypeToGDTFGeometryFactory.h
@@ -47,6 +47,16 @@ namespace UE::DMX::GDTF
 		};
 	}
 
+	/** A struct that links a mode to its base geometry */
+	struct FDMXFixtureModeWithBaseGeometry
+	{
+		FDMXFixtureModeWithBaseGeometry(const FDMXFixtureMode* InModePtr, TSharedRef<const FDMXGDTFGeometry> InBaseGeometry);
+
+		const FDMXFixtureMode* ModePtr = nullptr;
+		const TSharedRef<const FDMXGDTFGeometry> BaseGeometry;
+	};
+
+	/** A struct that links a function to its controlled geometry */
 	struct FDMXFixtureFunctionWithControlledGeometry
 	{
 		FDMXFixtureFunctionWithControlledGeometry(const FDMXFixtureMode* InModePtr, const FDMXFixtureFunction* InFunctionPtr, TSharedRef<const FDMXGDTFGeometry> InControlledGeometry);
@@ -64,16 +74,24 @@ namespace UE::DMX::GDTF
 
 	public:
 		FDMXFixtureTypeToGDTFGeometryFactory(const UDMXEntityFixtureType& InFixtureType, const TSharedRef<FDMXGDTFGeometryCollect>& InBaseGeometryCollect);
-		
+
+		/** Returns an array of functions with their linked geometry */
+		TArray<FDMXFixtureModeWithBaseGeometry> GetModesWithBaseGeometry() const;
+
 		/** Returns an array of functions with their linked geometry */
 		TArray<FDMXFixtureFunctionWithControlledGeometry> GetFunctionsWithControlledGeometry() const;
 
+		static const FName CellsModelName;
+
 		static const FName BaseGeometryName;
 		static const FName PanGeometryName;
 		static const FName HeadGeometryName;
 		static const FName BeamGeometryName;
 		static const FName MatrixBeamGeometryName;
 
+		static const FName PanAttributeName;
+		static const FName TiltAttributeName;
+
 	private:
 		/** Builds all geometries */
 		void BuildGeometries();
diff --git a/Engine/Plugins/VirtualProduction/DMX/DMXEngine/Source/DMXEditor/Private/Exporters/DMXMVRExportOptions.h b/Engine/Plugins/VirtualProduction/DMX/DMXEngine/Source/DMXEditor/Private/Exporters/DMXMVRExportOptions.h
--- a/Engine/Plugins/VirtualProduction/DMX/DMXEngine/Source/DMXEditor/Private/Exporters/DMXMVRExportOptions.h
+++ b/Engine/Plugins/VirtualProduction/DMX/DMXEngine/Source/DMXEditor/Private/Exporters/DMXMVRExportOptions.h
@@ -17,13 +17,17 @@ public:
 	UPROPERTY(EditAnywhere, Category = "Level Options")
 	bool bConsiderLevel = true;
 
+	/** If checked, exports patches that are not in use in the level */
+	UPROPERTY(EditAnywhere, Meta = (EditCondition = "bConsiderLevel", HideEditConditionToggle, EditConditionHides), Category = "Level Options")
+	bool bExportPatchesNotPresentInWorld = true;
+
 	/** If checked, exports the fixtures with transforms as in the current level */
 	UPROPERTY(EditAnywhere, Meta = (EditCondition = "bConsiderLevel", HideEditConditionToggle, EditConditionHides), Category = "Level Options")
 	bool bUseTransformsFromLevel = true;
 
-	/** If checked, exports patches that are not in use in the level */
-	UPROPERTY(EditAnywhere, Meta = (EditCondition = "bConsiderLevel", HideEditConditionToggle, EditConditionHides), Category = "Level Options")
-	bool bExportPatchesNotPresentInLevel = true;
+	/** If checked, creates individual MVR Fixtures for each Fixture that uses the same Patch in the Level. */
+	UPROPERTY(EditAnywhere, AdvancedDisplay, Meta = (EditCondition = "bConsiderLevel", HideEditConditionToggle, EditConditionHides, DisplayName = "Create Multi Patch Fixtures (experimental)"), Category = "Level Options")
+	bool bCreateMultiPatchFixtures = false;
 
 	/** If set to true, the import was canceled */
 	UPROPERTY(Transient)
diff --git a/Engine/Plugins/VirtualProduction/DMX/DMXEngine/Source/DMXEditor/Private/Exporters/DMXMVRExporter.cpp b/Engine/Plugins/VirtualProduction/DMX/DMXEngine/Source/DMXEditor/Private/Exporters/DMXMVRExporter.cpp
--- a/Engine/Plugins/VirtualProduction/DMX/DMXEngine/Source/DMXEditor/Private/Exporters/DMXMVRExporter.cpp
+++ b/Engine/Plugins/VirtualProduction/DMX/DMXEngine/Source/DMXEditor/Private/Exporters/DMXMVRExporter.cpp
@@ -2,7 +2,6 @@
 
 #include "Exporters/DMXMVRExporter.h"
 
-#include "Algo/AnyOf.h"
 #include "DesktopPlatformModule.h"
 #include "DMXEditorLog.h"
 #include "DMXEditorSettings.h"
@@ -29,7 +28,6 @@
 #include "Misc/Paths.h"
 #include "MVR/DMXMVRAssetImportData.h"
 #include "MVR/DMXMVRGeneralSceneDescription.h"
-#include "MVR/Types/DMXMVRFixtureNode.h"
 #include "Subsystems/UnrealEditorSubsystem.h"
 #include "Widgets/Notifications/SNotificationList.h"
 #include "Widgets/SDMXMVRExportOptions.h"
@@ -111,24 +109,25 @@ namespace UE::DMX
 				return SaveFilenames[0];
 			}();
 
+		// Get Export options
+		UUnrealEditorSubsystem* UnrealEditorSubsystem = GEditor->GetEditorSubsystem<UUnrealEditorSubsystem>();
+		const UDMXMVRExportOptions* ExportOptions = GetDefault<UDMXMVRExportOptions>();
+
+		FDMXMVRGeneralSceneDescriptionWorldParams WorldParams;
+		WorldParams.World = UnrealEditorSubsystem->GetEditorWorld();
+		WorldParams.bCreateMultiPatchFixtures = ExportOptions->bCreateMultiPatchFixtures;
+		WorldParams.bExportPatchesNotPresentInWorld = ExportOptions->bExportPatchesNotPresentInWorld;
+		WorldParams.bUseTransformsFromLevel = ExportOptions->bUseTransformsFromLevel;
+
 		// Create a copy of the library's general scene description, so transforms exported aren't written to the dmx library
-		DMXLibrary->UpdateGeneralSceneDescription();
 		UDMXMVRGeneralSceneDescription* TempGeneralSceneDescription = DuplicateObject<UDMXMVRGeneralSceneDescription>(DMXLibrary->GetLazyGeneralSceneDescription(), GetTransientPackage());
+		TempGeneralSceneDescription->WriteDMXLibrary(*DMXLibrary, WorldParams);
 		if (!ensureAlwaysMsgf(TempGeneralSceneDescription, TEXT("Trying to export DMX Library '%s' as MVR file, but its General Scene Description is invalid."), *DMXLibrary->GetName()))
 		{
 			OutErrorReason = FText::Format(LOCTEXT("MVRExportGeneralSceneDescriptionInvalidReason", "DMX Library is invalid. Cannot export {0}."), FText::FromString(OutFilePathAndName));
 			return;
 		}
 
-		const UDMXMVRExportOptions* ExportOptions = GetDefault<UDMXMVRExportOptions>();
-		if (ExportOptions->bConsiderLevel)
-		{
-			const TMap<const UDMXComponent*, const AActor*> DMXComponentToActorMap = GetDMXComponentToActorMap();
-
-			TryRemoveMVRFixturesNotPresentInLevel(TempGeneralSceneDescription, DMXLibrary, DMXComponentToActorMap);
-			TryUpdateFixtureNodesFromLevel(TempGeneralSceneDescription, DMXComponentToActorMap, ExportOptions->bUseTransformsFromLevel);
-		}
-
 		const TSharedRef<FDMXZipper> Zip = MakeShared<FDMXZipper>();
 		if (!ZipGeneralSceneDescription(Zip, TempGeneralSceneDescription, OutErrorReason))
 		{
@@ -427,83 +426,6 @@ namespace UE::DMX
 		}
 	}
 
-	void FDMXMVRExporter::TryRemoveMVRFixturesNotPresentInLevel(UDMXMVRGeneralSceneDescription* GeneralSceneDescription, UDMXLibrary* DMXLibrary, const TMap<const UDMXComponent*, const AActor*>& DMXComponentToActorMap)
-	{
-		if (!ensureMsgf(GeneralSceneDescription && DMXLibrary, TEXT("Trying to write MVR Fixture Transforms from Level to MVR General Scene Description, but General Scene Description or DMX Library are invalid")))
-		{
-			return;
-		}
-
-		const UDMXMVRExportOptions* ExportOptions = GetDefault<UDMXMVRExportOptions>();
-		if (ExportOptions->bExportPatchesNotPresentInLevel)
-		{
-			return;
-		}
-
-		const TArray<UDMXEntityFixturePatch*> FixturePatchesInLibrary = DMXLibrary->GetEntitiesTypeCast<UDMXEntityFixturePatch>();
-
-		TArray<const UDMXEntityFixturePatch*> FixturePatchesToRemove;
-		Algo::TransformIf(FixturePatchesInLibrary, FixturePatchesToRemove,
-			[&DMXComponentToActorMap](const UDMXEntityFixturePatch* FixturePatch)
-			{
-				const bool bPatchExistsInLevel = Algo::AnyOf(DMXComponentToActorMap, [FixturePatch](const TTuple<const UDMXComponent*, const AActor*>& DMXComponentToActorPair)
-					{
-						const UDMXComponent* DMXComponent = DMXComponentToActorPair.Key;
-						return
-							DMXComponent &&
-							FixturePatch &&
-							DMXComponent->GetFixturePatch() == FixturePatch;
-					});
-
-				return !bPatchExistsInLevel;
-			},
-			[](const UDMXEntityFixturePatch* FixturePatch)
-			{
-				return FixturePatch;
-			});
-
-		for (const UDMXEntityFixturePatch* FixturePatch : FixturePatchesToRemove)
-		{
-			GeneralSceneDescription->RemoveFixtureNode(FixturePatch->GetMVRFixtureUUID());
-		}
-	}
-
-	void FDMXMVRExporter::TryUpdateFixtureNodesFromLevel(UDMXMVRGeneralSceneDescription* GeneralSceneDescription, const TMap<const UDMXComponent*, const AActor*>& DMXComponentToActorMap, bool bUseTransformsFromLevel)
-	{
-		if (!ensureMsgf(GeneralSceneDescription, TEXT("Trying to write MVR Fixture Transforms from Level to MVR General Scene Description, but General Scene Description are invalid")))
-		{
-			return;
-		}		
-		
-		const UDMXMVRExportOptions* ExportOptions = GetDefault<UDMXMVRExportOptions>();
-
-		TArray<UDMXMVRFixtureNode*> FixtureNodes;
-		GeneralSceneDescription->GetFixtureNodes(FixtureNodes);
-
-		TSet<FGuid> UsedMVRUUIDs;
-		for (const TTuple<const UDMXComponent*, const AActor*>& DMXComponentToActorPair : DMXComponentToActorMap)
-		{
-			const UDMXEntityFixturePatch* FixturePatch = DMXComponentToActorPair.Key->GetFixturePatch();
-			const AActor* Actor = DMXComponentToActorPair.Value;
-			if (!FixturePatch || !Actor || UsedMVRUUIDs.Contains(FixturePatch->GetMVRFixtureUUID()))
-			{
-				continue;
-			}
-
-			UDMXMVRFixtureNode* FixtureNode = GeneralSceneDescription->WriteFixturePatch(*FixturePatch);
-			if (!FixtureNode)
-			{
-				continue;
-			}
-
-			UsedMVRUUIDs.Add(FixtureNode->UUID);
-			if (ExportOptions->bUseTransformsFromLevel)
-			{
-				FixtureNode->SetTransformAbsolute(Actor->GetTransform());
-			}
-		}
-	}
-
 	const TSharedPtr<FXmlFile> FDMXMVRExporter::CreateSourceGeneralSceneDescriptionXmlFile(const UDMXMVRGeneralSceneDescription* GeneralSceneDescription) const
 	{
 		const UDMXMVRAssetImportData* AssetImportData = GeneralSceneDescription->GetMVRAssetImportData();
diff --git a/Engine/Plugins/VirtualProduction/DMX/DMXEngine/Source/DMXEditor/Private/Exporters/DMXMVRExporter.h b/Engine/Plugins/VirtualProduction/DMX/DMXEngine/Source/DMXEditor/Private/Exporters/DMXMVRExporter.h
--- a/Engine/Plugins/VirtualProduction/DMX/DMXEngine/Source/DMXEditor/Private/Exporters/DMXMVRExporter.h
+++ b/Engine/Plugins/VirtualProduction/DMX/DMXEngine/Source/DMXEditor/Private/Exporters/DMXMVRExporter.h
@@ -47,12 +47,6 @@ namespace UE::DMX
 		/** Creates an General Scene Description Xml File from the MVR Source, as it was imported */
 		const TSharedPtr<FXmlFile> CreateSourceGeneralSceneDescriptionXmlFile(const UDMXMVRGeneralSceneDescription* GeneralSceneDescription) const;
 
-		/** Tries to remove MVR fixtures that are not present in the level, depending on export options. */
-		void TryRemoveMVRFixturesNotPresentInLevel(UDMXMVRGeneralSceneDescription* GeneralSceneDescription, UDMXLibrary* DMXLibrary, const TMap<const UDMXComponent*, const AActor*>& DMXComponentToActorMap);
-
-		/** Tries to create MVR fixtures from patches used more than once in the level and apply transforms depending on export options. */
-		void TryUpdateFixtureNodesFromLevel(UDMXMVRGeneralSceneDescription* GeneralSceneDescription, const TMap<const UDMXComponent*, const AActor*>& DMXComponentToActorMap, bool bUseTransformsFromLevel);
-
 		/**
 		 * Gets raw source data or creates (possibly empty) source data where the source data is not present.
 		 * Prior 5.1 there was no source data stored. Offers a dialog to load missing data
diff --git a/Engine/Plugins/VirtualProduction/DMX/DMXEngine/Source/DMXEditor/Private/Factories/DMXGDTFToFixtureTypeConverter.cpp b/Engine/Plugins/VirtualProduction/DMX/DMXEngine/Source/DMXEditor/Private/Factories/DMXGDTFToFixtureTypeConverter.cpp
--- a/Engine/Plugins/VirtualProduction/DMX/DMXEngine/Source/DMXEditor/Private/Factories/DMXGDTFToFixtureTypeConverter.cpp
+++ b/Engine/Plugins/VirtualProduction/DMX/DMXEngine/Source/DMXEditor/Private/Factories/DMXGDTFToFixtureTypeConverter.cpp
@@ -433,9 +433,12 @@ namespace UE::DMX::GDTF
 				Function.Channel = ChannelProperties.FirstChannel;
 				Function.bUseLSBMode = ChannelProperties.bLSBMode;
 				Function.DataType = ChannelProperties.SignalFormat;
-				Function.DefaultValue = ChannelProperties.DefaultValue;
+
+				// Set physical properties (order matters here)
 				Function.SetPhysicalUnit(ChannelProperties.PhysicalUnit);
 				Function.SetPhysicalValueRange(ChannelProperties.PhysicalFrom, ChannelProperties.PhysicalTo);
+				Function.DefaultValue = ChannelProperties.DefaultValue;
+				Function.UpdatePhysicalDefaultValue();
 
 				Mode.Functions.Add(Function);
 
diff --git a/Engine/Plugins/VirtualProduction/DMX/DMXEngine/Source/DMXEditor/Private/Widgets/FixtureType/DMXFixtureTypeFunctionsEditorFunctionItem.cpp b/Engine/Plugins/VirtualProduction/DMX/DMXEngine/Source/DMXEditor/Private/Widgets/FixtureType/DMXFixtureTypeFunctionsEditorFunctionItem.cpp
--- a/Engine/Plugins/VirtualProduction/DMX/DMXEngine/Source/DMXEditor/Private/Widgets/FixtureType/DMXFixtureTypeFunctionsEditorFunctionItem.cpp
+++ b/Engine/Plugins/VirtualProduction/DMX/DMXEngine/Source/DMXEditor/Private/Widgets/FixtureType/DMXFixtureTypeFunctionsEditorFunctionItem.cpp
@@ -3,10 +3,10 @@
 #include "DMXFixtureTypeFunctionsEditorFunctionItem.h"
 
 #include "DMXAttribute.h"
+#include "DMXAttributeToDefaultPhyiscalProperties.h"
 #include "DMXEditor.h"
 #include "DMXRuntimeUtils.h"
 #include "Library/DMXEntityFixtureType.h"
-
 #include "ScopedTransaction.h"
 
 
@@ -112,6 +112,8 @@ FDMXAttributeName FDMXFixtureTypeFunctionsEditorFunctionItem::GetAttributeName()
 
 void FDMXFixtureTypeFunctionsEditorFunctionItem::SetAttributeName(const FDMXAttributeName& AttributeName) const
 {
+	using namespace UE::DMX;
+
 	if (ensureMsgf(FixtureType.IsValid() && FixtureType->Modes.IsValidIndex(ModeIndex) && FixtureType->Modes[ModeIndex].Functions.IsValidIndex(FunctionIndex), TEXT("Invalid Fixture Type, Mode or Function in FDMXFixtureTypeFunctionsEditorFunctionItem.")))
 	{
 		const FScopedTransaction SetAttributeNamelTransaction(LOCTEXT("SetAttributeNamelTransaction", "Set Fixture Function Attribute Name"));
@@ -120,6 +122,8 @@ void FDMXFixtureTypeFunctionsEditorFunctionItem::SetAttributeName(const FDMXAttr
 		FDMXFixtureFunction& Function = FixtureType->Modes[ModeIndex].Functions[FunctionIndex];
 		Function.Attribute = AttributeName;
 
+		FDMXAttributeToDefaultPhyiscalProperties::ResetToDefaultPhysicalProperties(Function);
+
 		FixtureType->PostEditChange();
 	}
 }
diff --git a/Engine/Plugins/VirtualProduction/DMX/DMXEngine/Source/DMXEditor/Private/Widgets/FixtureType/SDMXFixtureFunctionEditor.cpp b/Engine/Plugins/VirtualProduction/DMX/DMXEngine/Source/DMXEditor/Private/Widgets/FixtureType/SDMXFixtureFunctionEditor.cpp
--- a/Engine/Plugins/VirtualProduction/DMX/DMXEngine/Source/DMXEditor/Private/Widgets/FixtureType/SDMXFixtureFunctionEditor.cpp
+++ b/Engine/Plugins/VirtualProduction/DMX/DMXEngine/Source/DMXEditor/Private/Widgets/FixtureType/SDMXFixtureFunctionEditor.cpp
@@ -15,7 +15,6 @@
 #include "Widgets/Layout/SBox.h"
 #include "Widgets/Text/STextBlock.h"
 
-
 #define LOCTEXT_NAMESPACE "SDMXFixtureFunctionEditor"
 
 namespace UE::DMX::SDMXFixtureFunctionEditor::Private
@@ -122,6 +121,8 @@ void SDMXFixtureFunctionEditor::Construct(const FArguments& InArgs, const TShare
 
 	FStructureDetailsViewArgs StructureDetailsViewArgs;
 
+	using namespace UE::DMX;
+
 	FPropertyEditorModule& PropertyModule = FModuleManager::LoadModuleChecked<FPropertyEditorModule>("PropertyEditor");
 	StructDetailsView = PropertyModule.CreateStructureDetailView(DetailsViewArgs, StructureDetailsViewArgs, nullptr);
 	StructDetailsView->GetDetailsView()->SetIsPropertyVisibleDelegate(FIsPropertyVisible::CreateSP(this, &SDMXFixtureFunctionEditor::IsPropertyVisible));
@@ -176,10 +177,6 @@ void SDMXFixtureFunctionEditor::NotifyPreChange(FProperty* PropertyAboutToChange
 		{
 			Transaction = MakeUnique<FScopedTransaction>(LOCTEXT("SetDefaultValueTransaction", "Set Default Value of Function"));
 		}
-		else if (PropertyName == GET_MEMBER_NAME_CHECKED(FDMXFixtureFunction, DefaultValue))
-		{
-			Transaction = MakeUnique<FScopedTransaction>(LOCTEXT("SetAttributeTransaction", "Set Attribute of Function"));
-		}
 		else if (PropertyName == FDMXFixtureFunction::GetPhysicalDefaultValuePropertyName())
 		{
 			Transaction = MakeUnique<FScopedTransaction>(LOCTEXT("SetPhysicalValueTransaction", "Set Physical Default Value of Function"));
@@ -226,17 +223,12 @@ void SDMXFixtureFunctionEditor::NotifyPostChange(const FPropertyChangedEvent& Pr
 			else if (PropertyName == GET_MEMBER_NAME_CHECKED(FDMXFixtureFunction, DataType))
 			{
 				FixtureType->AlignFunctionChannels(ModeIndex);
-				ClampFunctionDefautValueByDataType(*FixtureType, ModeIndex, FunctionIndex);
+				UpdateFunctionDefaultValueFromPhysicalUnit(*FixtureType, ModeIndex, FunctionIndex);
 			}
 			else if (PropertyName == GET_MEMBER_NAME_CHECKED(FDMXFixtureFunction, DefaultValue))
 			{
 				ClampFunctionDefautValueByDataType(*FixtureType, ModeIndex, FunctionIndex);
 			}
-			else if (PropertyName == GET_MEMBER_NAME_CHECKED(FDMXFixtureFunction, Attribute) ||
-				PropertyName == GET_MEMBER_NAME_CHECKED(FDMXAttribute, Name))
-			{
-				UpdateFunctionPhysicalUnitFromAttribute(*FixtureType, ModeIndex, FunctionIndex);
-			}
 			else if (PropertyName == FDMXFixtureFunction::GetPhysicalDefaultValuePropertyName() || 
 					PropertyName == FDMXFixtureFunction::GetPhysicalFromPropertyName() ||
 					PropertyName == FDMXFixtureFunction::GetPhysicalToPropertyName())
diff --git a/Engine/Plugins/VirtualProduction/DMX/DMXEngine/Source/DMXRuntime/Internal/DMXAttributeToDefaultPhyiscalProperties.h b/Engine/Plugins/VirtualProduction/DMX/DMXEngine/Source/DMXRuntime/Internal/DMXAttributeToDefaultPhyiscalProperties.h
--- a/Engine/Plugins/VirtualProduction/DMX/DMXEngine/Source/DMXRuntime/Internal/DMXAttributeToDefaultPhyiscalProperties.h
+++ b/Engine/Plugins/VirtualProduction/DMX/DMXEngine/Source/DMXRuntime/Internal/DMXAttributeToDefaultPhyiscalProperties.h
@@ -38,6 +38,7 @@ namespace UE::DMX
 			{ "Pan", FDefaultPhyicalProperties(EDMXGDTFPhysicalUnit::Angle, -120.0, 120.0) },
 			{ "Tilt", FDefaultPhyicalProperties(EDMXGDTFPhysicalUnit::Angle, -120.0, 120.0) },
 
+			{ "Angle", FDefaultPhyicalProperties(EDMXGDTFPhysicalUnit::Angle, 0.0, 120.0) },
 		};
 	};
 }
diff --git a/Engine/Plugins/VirtualProduction/DMX/DMXEngine/Source/DMXRuntime/Private/DMXAttributeToDefaultPhyiscalProperties.cpp b/Engine/Plugins/VirtualProduction/DMX/DMXEngine/Source/DMXRuntime/Private/DMXAttributeToDefaultPhyiscalProperties.cpp
--- a/Engine/Plugins/VirtualProduction/DMX/DMXEngine/Source/DMXRuntime/Private/DMXAttributeToDefaultPhyiscalProperties.cpp
+++ b/Engine/Plugins/VirtualProduction/DMX/DMXEngine/Source/DMXRuntime/Private/DMXAttributeToDefaultPhyiscalProperties.cpp
@@ -13,6 +13,11 @@ namespace UE::DMX
 		const FDefaultPhyicalProperties* AttributeNameToDefaultPhysicalProperitesPtr = AttributeNameToDefaultPhysicalProperitesMap.Find(InOutFunction.Attribute.Name);
 		if (AttributeNameToDefaultPhysicalProperitesPtr)
 		{
+			if (InOutFunction.GetPhysicalUnit() == AttributeNameToDefaultPhysicalProperitesPtr->PhysicalUnit)
+			{
+				return;
+			}
+
 			InOutFunction.SetPhysicalUnit(AttributeNameToDefaultPhysicalProperitesPtr->PhysicalUnit);
 			InOutFunction.SetPhysicalValueRange(AttributeNameToDefaultPhysicalProperitesPtr->PhysicalFrom, AttributeNameToDefaultPhysicalProperitesPtr->PhysicalTo);
 		}
diff --git a/Engine/Plugins/VirtualProduction/DMX/DMXEngine/Source/DMXRuntime/Private/MVR/DMXMVRGeneralSceneDescription.cpp b/Engine/Plugins/VirtualProduction/DMX/DMXEngine/Source/DMXRuntime/Private/MVR/DMXMVRGeneralSceneDescription.cpp
--- a/Engine/Plugins/VirtualProduction/DMX/DMXEngine/Source/DMXRuntime/Private/MVR/DMXMVRGeneralSceneDescription.cpp
+++ b/Engine/Plugins/VirtualProduction/DMX/DMXEngine/Source/DMXRuntime/Private/MVR/DMXMVRGeneralSceneDescription.cpp
@@ -3,11 +3,13 @@
 #include "MVR/DMXMVRGeneralSceneDescription.h"
 
 #include "Algo/AnyOf.h"
+#include "Algo/Find.h"
 #include "Algo/MaxElement.h"
+#include "EngineUtils.h"
+#include "Game/DMXComponent.h"
 #include "Library/DMXEntityFixturePatch.h"
 #include "Library/DMXEntityFixtureType.h"
 #include "Library/DMXImportGDTF.h"
-#include "Library/DMXGDTFAssetImportData.h"
 #include "Library/DMXLibrary.h"
 #include "MVR/DMXMVRAssetImportData.h"
 #include "MVR/Types/DMXMVRChildListNode.h"
@@ -76,8 +78,10 @@ void UDMXMVRGeneralSceneDescription::WriteDMXLibraryToGeneralSceneDescription(co
 	// Deprecated 5.5, renamed to WriteDMXLibrary
 	WriteDMXLibrary(DMXLibrary);
 }
+#endif // WITH_EDITOR
 
-void UDMXMVRGeneralSceneDescription::WriteDMXLibrary(const UDMXLibrary& DMXLibrary)
+#if WITH_EDITOR
+void UDMXMVRGeneralSceneDescription::WriteDMXLibrary(const UDMXLibrary& DMXLibrary, FDMXMVRGeneralSceneDescriptionWorldParams WorldParams)
 {
 	TArray<UDMXEntityFixturePatch*> FixturePatches = DMXLibrary.GetEntitiesTypeCast<UDMXEntityFixturePatch>();
 	Algo::Sort(FixturePatches, [](const UDMXEntityFixturePatch* FixturePatchA, const UDMXEntityFixturePatch* FixturePatchB)
@@ -87,6 +91,7 @@ void UDMXMVRGeneralSceneDescription::WriteDMXLibrary(const UDMXLibrary& DMXLibra
 
 			return AbsoluteAddressA <= AbsoluteAddressB;
 		});
+	FixturePatches.Remove(nullptr);
 
 	// Remove Fixture Nodes no longer defined in the DMX Library
 	TArray<FGuid> MVRFixtureUUIDsInUse;
@@ -94,6 +99,7 @@ void UDMXMVRGeneralSceneDescription::WriteDMXLibrary(const UDMXLibrary& DMXLibra
 	{
 		MVRFixtureUUIDsInUse.Add(FixturePatch->GetMVRFixtureUUID());
 	}
+
 	TArray<UDMXMVRFixtureNode*> FixtureNodes;
 	RootNode->GetFixtureNodes(FixtureNodes);
 	for (UDMXMVRFixtureNode* FixtureNode : FixtureNodes)
@@ -104,7 +110,7 @@ void UDMXMVRGeneralSceneDescription::WriteDMXLibrary(const UDMXLibrary& DMXLibra
 		}
 	}
 
-	// Create or update MVR Fixtures for each Fixture Patch's MVR Fixture UUIDs
+	// Write all Patches in library
 	for (UDMXEntityFixturePatch* FixturePatch : FixturePatches)
 	{
 		if (FixturePatch)
@@ -112,61 +118,77 @@ void UDMXMVRGeneralSceneDescription::WriteDMXLibrary(const UDMXLibrary& DMXLibra
 			WriteFixturePatch(*FixturePatch);
 		}
 	}
-}
-#endif // WITH_EDITOR
-
-#if WITH_EDITOR
-UDMXMVRFixtureNode* UDMXMVRGeneralSceneDescription::WriteFixturePatch(const UDMXEntityFixturePatch& FixturePatch)
-{
-	const UDMXLibrary* DMXLibrary = FixturePatch.GetParentLibrary();
-	if (!DMXLibrary)
+	
+	// Only export the library if there's no world
+	if (!WorldParams.World)
 	{
-		return nullptr;
+		return;
 	}
 
-	const FGuid& MVRFixtureUUID = FixturePatch.GetMVRFixtureUUID();
-	const TObjectPtr<UDMXMVRParametricObjectNodeBase>* ParametricObjectNodePtr = RootNode->FindParametricObjectNodeByUUID(MVRFixtureUUID);
-	UDMXMVRFixtureNode* MVRFixtureNode = ParametricObjectNodePtr ? Cast<UDMXMVRFixtureNode>(*ParametricObjectNodePtr) : nullptr;
+	// Consider the world 
+	const TMap<const UDMXComponent*, const AActor*> DMXComponentToActorMap = GetDMXComponentToActorMap(WorldParams.World);
 
-	if (!MVRFixtureNode)
+	// Create fixtures and write transforms if desired
+	for (const UDMXEntityFixturePatch* FixturePatchInLibrary : FixturePatches)
 	{
-		UDMXMVRChildListNode& AnyChildList = RootNode->GetOrCreateFirstChildListNode();
-		MVRFixtureNode = AnyChildList.CreateParametricObject<UDMXMVRFixtureNode>();
+		FGuid MultiPatchUUID;
+		for (const TTuple<const UDMXComponent*, const AActor*>& DMXComponentToActorPair : DMXComponentToActorMap)
+		{
+			const UDMXEntityFixturePatch* FixturePatchInWorld = DMXComponentToActorPair.Key ? DMXComponentToActorPair.Key->GetFixturePatch() : nullptr;
+			const AActor* Actor = DMXComponentToActorPair.Value;
+			if (!FixturePatchInWorld || !Actor || FixturePatchInWorld != FixturePatchInLibrary)
+			{
+				continue;
+			}
 
-		MVRFixtureNode->Name = FixturePatch.Name;
-		MVRFixtureNode->UUID = MVRFixtureUUID;
-		MVRFixtureNode->FixtureID = FString::FromInt(FixturePatch.GetFixtureID());
-	}
-	check(MVRFixtureNode);
+			if (!MultiPatchUUID.IsValid())
+			{
+				// Remember the MVR UUID of the first patch as multi patch UUID
+				MultiPatchUUID = FixturePatchInWorld->GetMVRFixtureUUID();
 
-	MVRFixtureNode->SetUniverseID(FixturePatch.GetUniverseID());
-	MVRFixtureNode->SetStartingChannel(FixturePatch.GetStartingChannel());
+				const TOptional<FTransform> OptionalTransform = WorldParams.bUseTransformsFromLevel ? Actor->GetTransform() : TOptional<FTransform>();
+				WriteFixturePatch(*FixturePatchInLibrary, OptionalTransform);
+			}
+			else if (WorldParams.bCreateMultiPatchFixtures)
+			{
+				const TOptional<FTransform> OptionalTransform = WorldParams.bUseTransformsFromLevel ? Actor->GetTransform() : TOptional<FTransform>();
+				WriteFixturePatch(*FixturePatchInLibrary, OptionalTransform, MultiPatchUUID);
+			}
+		}
+	}
+	
+	// Update fixture nodes to contain newly added multi patch fixtures
+	RootNode->GetFixtureNodes(FixtureNodes); 
 
-	UDMXEntityFixtureType* FixtureType = FixturePatch.GetFixtureType();
-	const int32 ModeIndex = FixturePatch.GetActiveModeIndex();
-	bool bSetGDTFSpec = false;
 
-	if (FixtureType && FixtureType->Modes.IsValidIndex(ModeIndex))
+	// Remove patches not present in the world if desired
+	if (!WorldParams.bExportPatchesNotPresentInWorld)
 	{
-		// Instead refer to the generated file name
-		MVRFixtureNode->GDTFMode = FixtureType->Modes[ModeIndex].ModeName;
-
-		constexpr bool bWithExtension = false;
-		MVRFixtureNode->GDTFSpec = FixtureType->GetCleanGDTFFileNameSynchronous(bWithExtension);
-
-		bSetGDTFSpec = true;
-	}
+		for (const UDMXEntityFixturePatch* FixturePatchInLibrary : FixturePatches)
+		{
+			const bool bPatchUsedInWorldAndLibrary = Algo::AnyOf(DMXComponentToActorMap, [FixturePatchInLibrary](const TTuple<const UDMXComponent*, const AActor*>& DMXComponentToActorPair)
+				{
+					return DMXComponentToActorPair.Key && DMXComponentToActorPair.Key->GetFixturePatch() == FixturePatchInLibrary;
+				});					
 
-	if (!bSetGDTFSpec)
-	{
-		// Don't set a mode when there's no GDTF
-		MVRFixtureNode->GDTFMode = TEXT("");
-		MVRFixtureNode->GDTFSpec = TEXT("");
+			if (!bPatchUsedInWorldAndLibrary && FixturePatchInLibrary)
+			{
+				for (UDMXMVRFixtureNode* FixtureNode : FixtureNodes)
+				{
+					if (!FixtureNode)
+					{
+						continue;
+					}
+
+					if (FixtureNode->UUID == FixturePatchInLibrary->GetMVRFixtureUUID() ||
+						(FixtureNode->MultiPatch.IsSet() && FixtureNode->MultiPatch.GetValue() == FixturePatchInLibrary->GetMVRFixtureUUID()))
+					{
+						RootNode->RemoveParametricObjectNode(FixtureNode);
+					}
+				}
+			}
+		}
 	}
-
-	SanetizeFixtureNode(*MVRFixtureNode);
-
-	return MVRFixtureNode;
 }
 #endif // WITH_EDITOR
 
@@ -207,38 +229,135 @@ TSharedPtr<FXmlFile> UDMXMVRGeneralSceneDescription::CreateXmlFile() const
 }
 #endif // WITH_EDITOR
 
+#if WITH_EDITOR
+void UDMXMVRGeneralSceneDescription::WriteFixturePatch(const UDMXEntityFixturePatch& FixturePatch, const TOptional<FTransform>& OptionalTransform, const FGuid& MultiPatchUUID)
+{
+	TArray<UDMXMVRFixtureNode*> FixtureNodes;
+	RootNode->GetFixtureNodes(FixtureNodes);
+
+	UDMXMVRFixtureNode* MVRFixtureNode = nullptr;
+	const UDMXMVRFixtureNode* const* ParentMultiPatchFixtureNodePtr = [&FixtureNodes, &MultiPatchUUID]() -> const UDMXMVRFixtureNode* const*
+		{
+			if (MultiPatchUUID.IsValid())
+			{
+				const UDMXMVRFixtureNode* const* ResultPtr = Algo::FindByPredicate(FixtureNodes, [&MultiPatchUUID](const UDMXMVRFixtureNode* Other)
+					{
+						return Other->UUID == MultiPatchUUID;
+					});
+				ensureMsgf(ResultPtr, TEXT("MultiPatch UUID provided but no corresponding node can be found"));
+
+				return ResultPtr;
+			}
+			
+			return nullptr;
+		}(); 
+
+	if (ParentMultiPatchFixtureNodePtr)
+	{
+		// Create a multipatch fixture
+		UDMXMVRChildListNode& AnyChildList = RootNode->GetOrCreateFirstChildListNode();
+		MVRFixtureNode = AnyChildList.CreateParametricObject<UDMXMVRFixtureNode>();
+
+		MVRFixtureNode->Name = FixturePatch.Name;
+		MVRFixtureNode->UUID = FGuid::NewGuid();
+		MVRFixtureNode->MultiPatch = (*ParentMultiPatchFixtureNodePtr)->UUID;
+	}
+	else
+	{
+		const FGuid& MVRFixtureUUID = FixturePatch.GetMVRFixtureUUID();
+
+		// Find an existing fixture
+		const TObjectPtr<UDMXMVRParametricObjectNodeBase>* ParametricObjectNodePtr = RootNode->FindParametricObjectNodeByUUID(MVRFixtureUUID);
+		MVRFixtureNode = ParametricObjectNodePtr ? Cast<UDMXMVRFixtureNode>(*ParametricObjectNodePtr) : nullptr;
+
+		if (!MVRFixtureNode)
+		{
+			// Create a new fixture
+			UDMXMVRChildListNode& AnyChildList = RootNode->GetOrCreateFirstChildListNode();
+			MVRFixtureNode = AnyChildList.CreateParametricObject<UDMXMVRFixtureNode>();
+
+			MVRFixtureNode->Name = FixturePatch.Name;
+			MVRFixtureNode->Name = FixturePatch.Name;
+			MVRFixtureNode->UUID = MVRFixtureUUID;
+			MVRFixtureNode->FixtureID = FString::FromInt(FixturePatch.GetFixtureID());
+		}
+	}
+	check(MVRFixtureNode)
+
+	if(OptionalTransform.IsSet())
+	{
+		MVRFixtureNode->SetTransformAbsolute(OptionalTransform.GetValue());
+	}
+
+	MVRFixtureNode->SetUniverseID(FixturePatch.GetUniverseID());
+	MVRFixtureNode->SetStartingChannel(FixturePatch.GetStartingChannel());
+
+	UDMXEntityFixtureType* FixtureType = FixturePatch.GetFixtureType();
+	const int32 ModeIndex = FixturePatch.GetActiveModeIndex();
+	bool bSetGDTFSpec = false;
+
+	if (FixtureType && FixtureType->Modes.IsValidIndex(ModeIndex))
+	{
+		// Instead refer to the generated file name
+		MVRFixtureNode->GDTFMode = FixtureType->Modes[ModeIndex].ModeName;
+
+		constexpr bool bWithExtension = false;
+		MVRFixtureNode->GDTFSpec = FixtureType->GetCleanGDTFFileNameSynchronous(bWithExtension);
+
+		bSetGDTFSpec = true;
+	}
+
+	if (!bSetGDTFSpec)
+	{
+		// Don't set a mode when there's no GDTF
+		MVRFixtureNode->GDTFMode = TEXT("");
+		MVRFixtureNode->GDTFSpec = TEXT("");
+	}
+
+	SanetizeFixtureNode(*MVRFixtureNode);
+}
+#endif // WITH_EDITOR
+
 #if WITH_EDITOR
 void UDMXMVRGeneralSceneDescription::SanetizeFixtureNode(UDMXMVRFixtureNode& FixtureNode)
 {
 	TArray<UDMXMVRFixtureNode*> FixtureNodes;
 	GetFixtureNodes(FixtureNodes);
 
-	const bool bInvalidFixtureID = Algo::AnyOf(FixtureNodes, [&FixtureNode](const UDMXMVRFixtureNode* Other)
-		{
-			return
-				Other != &FixtureNode &&
-				Other->FixtureID == FixtureNode.FixtureID;
-		});
-	if (bInvalidFixtureID)
+	if (FixtureNode.MultiPatch.IsSet())
+	{
+		FixtureNode.FixtureID = TEXT("");
+		FixtureNode.CustomId.Reset();
+	}
+	else
 	{
-		const UDMXMVRFixtureNode* const* MaxFixtureNodePtr = Algo::MaxElementBy(FixtureNodes, [](const UDMXMVRFixtureNode* FixtureNode)
+		const bool bInvalidFixtureID = Algo::AnyOf(FixtureNodes, [&FixtureNode](const UDMXMVRFixtureNode* Other)
 			{
-				int32 IntegralFixtureID;
-				return LexTryParseString(IntegralFixtureID, *FixtureNode->FixtureID) ? IntegralFixtureID : 1;
+				return
+					Other != &FixtureNode &&
+					Other->FixtureID == FixtureNode.FixtureID;
 			});
+		if (bInvalidFixtureID)
+		{
+			const UDMXMVRFixtureNode* const* MaxFixtureNodePtr = Algo::MaxElementBy(FixtureNodes, [](const UDMXMVRFixtureNode* FixtureNode)
+				{
+					int32 IntegralFixtureID;
+					return LexTryParseString(IntegralFixtureID, *FixtureNode->FixtureID) ? IntegralFixtureID : 1;
+				});
 
-		FixtureNode.FixtureID = MaxFixtureNodePtr ? (*MaxFixtureNodePtr)->FixtureID : TEXT("1");
-	}
+			FixtureNode.FixtureID = MaxFixtureNodePtr ? (*MaxFixtureNodePtr)->FixtureID : TEXT("1");
+		}
 
-	const bool bInvalidMVRUUID = Algo::AnyOf(FixtureNodes, [&FixtureNode](const UDMXMVRFixtureNode* Other)
+		const bool bInvalidMVRUUID = Algo::AnyOf(FixtureNodes, [&FixtureNode](const UDMXMVRFixtureNode* Other)
+			{
+				return
+					Other != &FixtureNode &&
+					Other->UUID == FixtureNode.UUID;
+			});
+		if (bInvalidMVRUUID)
 		{
-			return
-				Other != &FixtureNode &&
-				Other->UUID == FixtureNode.UUID;
-		});
-	if (bInvalidMVRUUID)
-	{
-		FixtureNode.UUID = FGuid::NewGuid();
+			FixtureNode.UUID = FGuid::NewGuid();
+		}
 	}
 }
 #endif // WITH_EDITOR
@@ -252,6 +371,48 @@ bool UDMXMVRGeneralSceneDescription::ParseGeneralSceneDescriptionXml(const TShar
 }
 #endif // WITH_EDITOR
 
+#if WITH_EDITOR
+TMap<const UDMXComponent*, const AActor*> UDMXMVRGeneralSceneDescription::GetDMXComponentToActorMap(const UWorld* World) const
+{
+	TMap<const UDMXComponent*, const AActor*> DMXComponentToActorMap;
+	if (!World)
+	{
+		return DMXComponentToActorMap;
+	}
+
+	for (TActorIterator<AActor> It(World, AActor::StaticClass()); It; ++It)
+	{
+		const AActor* Actor = *It;
+		if (!Actor)
+		{
+			continue;
+		}
+
+		const TSet<UActorComponent*> Components = Actor->GetComponents();
+
+		TArray<const UDMXComponent*> DMXComponents;
+		Algo::TransformIf(Components, DMXComponents,
+			[](const UActorComponent* Component)
+			{
+				return Component &&
+					Component->IsA(UDMXComponent::StaticClass()) &&
+					CastChecked<UDMXComponent>(Component)->GetFixturePatch();
+			},
+			[](const UActorComponent* Component)
+			{
+				return CastChecked<UDMXComponent>(Component);
+			});
+
+		for (const UDMXComponent* DMXComponent : DMXComponents)
+		{
+			DMXComponentToActorMap.Add(TTuple<const UDMXComponent*, const AActor*>(DMXComponent, Actor));
+		}
+	}
+
+	return DMXComponentToActorMap;
+}
+#endif // WITH_EDITOR
+
 TArray<int32> UDMXMVRGeneralSceneDescription::GetNumericalFixtureIDsInUse(const UDMXLibrary& DMXLibrary)
 {
 	checkf(RootNode, TEXT("Unexpected: MVR General Scene Description Root Node is invalid."));
diff --git a/Engine/Plugins/VirtualProduction/DMX/DMXEngine/Source/DMXRuntime/Private/MVR/Types/DMXMVRFixtureNode.cpp b/Engine/Plugins/VirtualProduction/DMX/DMXEngine/Source/DMXRuntime/Private/MVR/Types/DMXMVRFixtureNode.cpp
--- a/Engine/Plugins/VirtualProduction/DMX/DMXEngine/Source/DMXRuntime/Private/MVR/Types/DMXMVRFixtureNode.cpp
+++ b/Engine/Plugins/VirtualProduction/DMX/DMXEngine/Source/DMXRuntime/Private/MVR/Types/DMXMVRFixtureNode.cpp
@@ -594,6 +594,14 @@ void UDMXMVRFixtureNode::CreateXmlNodeInParent(FXmlNode& ParentNode) const
 	FXmlAttribute UUIDAttribute(UUIDAttributeName, UUID.ToString(EGuidFormats::DigitsWithHyphens));
 	Attributes.Add(UUIDAttribute);
 
+	// MultiPatch attribute
+	if (MultiPatch.IsSet())
+	{
+		constexpr TCHAR MultiPatchAttributeName[] = TEXT("multipatch");
+		FXmlAttribute MultiPatchAttribute(MultiPatchAttributeName, MultiPatch.GetValue().ToString(EGuidFormats::DigitsWithHyphens));
+		Attributes.Add(MultiPatchAttribute);
+	}
+
 	FixtureXmlNode->SetAttributes(Attributes);
 
 
diff --git a/Engine/Plugins/VirtualProduction/DMX/DMXEngine/Source/DMXRuntime/Public/Library/DMXEntityFixtureType.h b/Engine/Plugins/VirtualProduction/DMX/DMXEngine/Source/DMXRuntime/Public/Library/DMXEntityFixtureType.h
--- a/Engine/Plugins/VirtualProduction/DMX/DMXEngine/Source/DMXRuntime/Public/Library/DMXEntityFixtureType.h
+++ b/Engine/Plugins/VirtualProduction/DMX/DMXEngine/Source/DMXRuntime/Public/Library/DMXEntityFixtureType.h
@@ -121,7 +121,7 @@ struct DMXRUNTIME_API FDMXFixtureFunction
 
 	/** This function's data type. Defines the used number of channels (bytes) */
 	UPROPERTY(EditAnywhere, BlueprintReadWrite, meta = (DisplayPriority = "5"), Category = "Function Settings")
-	EDMXFixtureSignalFormat DataType;
+	EDMXFixtureSignalFormat DataType = EDMXFixtureSignalFormat::E8Bit;
 
 	/**
 	 * Least Significant Byte mode makes the individual bytes (channels) of the function be
diff --git a/Engine/Plugins/VirtualProduction/DMX/DMXEngine/Source/DMXRuntime/Public/MVR/DMXMVRGeneralSceneDescription.h b/Engine/Plugins/VirtualProduction/DMX/DMXEngine/Source/DMXRuntime/Public/MVR/DMXMVRGeneralSceneDescription.h
--- a/Engine/Plugins/VirtualProduction/DMX/DMXEngine/Source/DMXRuntime/Public/MVR/DMXMVRGeneralSceneDescription.h
+++ b/Engine/Plugins/VirtualProduction/DMX/DMXEngine/Source/DMXRuntime/Public/MVR/DMXMVRGeneralSceneDescription.h
@@ -10,6 +10,7 @@
 #include "DMXMVRGeneralSceneDescription.generated.h"
 
 class FXmlFile;
+class UDMXComponent;
 class UDMXEntityFixturePatch;
 class UDMXLibrary;
 class UDMXMVRAssetImportData;
@@ -17,6 +18,23 @@ class UDMXMVRClassNode;
 class UDMXMVRRootNode;
 class UDMXMVRFixtureNode;
 
+namespace UE::DMX
+{
+	struct FDMXMVRGeneralSceneDescriptionWorldParams
+	{
+		UWorld* World = nullptr;
+
+		/** If checked, exports patches that are not in use in the level */
+		bool bExportPatchesNotPresentInWorld = false;
+
+		/** If checked, creates individual MVR Fixtures for each Fixture that uses the same Patch in the Level. */
+		bool bCreateMultiPatchFixtures = false;
+
+		/** If checked, exports the fixtures with transforms as in the current level */
+		bool bUseTransformsFromLevel = true;
+	};
+}
+
 
 /** MVR General Scene Description Object */
 UCLASS()
@@ -25,6 +43,8 @@ class DMXRUNTIME_API UDMXMVRGeneralSceneDescription
 {
 	GENERATED_BODY()
 
+	using FDMXMVRGeneralSceneDescriptionWorldParams = UE::DMX::FDMXMVRGeneralSceneDescriptionWorldParams;
+
 public:
 	/** Constructor */
 	UDMXMVRGeneralSceneDescription();
@@ -49,15 +69,10 @@ public:
 	/**
 	 * Writes the Library to the General Scene Description, effectively removing inexisting and adding
 	 * new MVR Fixtures, according to what MVR Fixture UUIDs the Fixture Patches of the Library contain.
+	 * 
+	 * If world params are specified, considers these.
 	 */
-	void WriteDMXLibrary(const UDMXLibrary& DMXLibrary);
-
-	/**
-	 * Writes the Fixture Patch to the General Scene Description.
-	 *
-	 * Returns the fixture node or nullptr if the node could not be created.
-	 */
-	UDMXMVRFixtureNode* WriteFixturePatch(const UDMXEntityFixturePatch& FixturePatch);
+	void WriteDMXLibrary(const UDMXLibrary& DMXLibrary, FDMXMVRGeneralSceneDescriptionWorldParams WorldParams = FDMXMVRGeneralSceneDescriptionWorldParams());
 
 	/** Removes a fixture node from the General Scene Description. */
 	void RemoveFixtureNode(const FGuid& FixtureUUID);
@@ -72,11 +87,24 @@ public:
 	FORCEINLINE UDMXMVRAssetImportData* GetMVRAssetImportData() const { return MVRAssetImportData; }
 
 private:
+	/** 	
+	 * Writes the Fixture Patch to the General Scene Description.
+	 * 
+	 * Optionally a transform can be specified for the patch.
+	 * 
+	 * Optionally a MultiPatch UUID can be passed so the patch is added as a multi patch.
+	 * When a MultiPatch UUID is provided, the parent with related UUID is expected to exist already.
+	 */
+	void WriteFixturePatch(const UDMXEntityFixturePatch& FixturePatch, const TOptional<FTransform>& OptionalTransform = TOptional<FTransform>(), const FGuid& MultiPatchUUID = FGuid());
+
 	/** Makes sure the node has a unique MVR UUID and Fixture ID */
 	void  SanetizeFixtureNode(UDMXMVRFixtureNode& FixtureNode);
 
 	/** Parses a General Scene Description Xml File. Only ment to be used for initialization (ensured) */
 	[[nodiscard]] bool ParseGeneralSceneDescriptionXml(const TSharedRef<FXmlFile>& GeneralSceneDescriptionXml);
+
+	/** Returns a map of DMX components with their actor from a world */
+	TMap<const UDMXComponent*, const AActor*> GetDMXComponentToActorMap(const UWorld* World) const;
 #endif
 
 	/** Returns the Fixture IDs currently in use, as number, sorted from lowest to highest */
diff --git a/Engine/Plugins/VirtualProduction/DMX/DMXEngine/Source/DMXRuntime/Public/MVR/Types/DMXMVRFixtureNode.h b/Engine/Plugins/VirtualProduction/DMX/DMXEngine/Source/DMXRuntime/Public/MVR/Types/DMXMVRFixtureNode.h
--- a/Engine/Plugins/VirtualProduction/DMX/DMXEngine/Source/DMXRuntime/Public/MVR/Types/DMXMVRFixtureNode.h
+++ b/Engine/Plugins/VirtualProduction/DMX/DMXEngine/Source/DMXRuntime/Public/MVR/Types/DMXMVRFixtureNode.h
@@ -163,6 +163,13 @@ public:
 	UPROPERTY()
 	FString Name;
 
+	/** 
+	 * The unique identifier of the parent multipatch fixture. When this value is set, you may not define a FixtureID or CustomID for this fixture. 
+	 * The FixtureID and CustomID from the object defined as multi parent also applies to this object. 
+	 */
+	UPROPERTY()
+	FDMXOptionalGuid MultiPatch;
+
 	/** The Class the object belongs to. */
 	UPROPERTY()
 	FDMXOptionalGuid Classing;
diff --git a/Engine/Plugins/VirtualProduction/DMX/DMXGDTF/Source/DMXGDTF/Internal/GDTF/Geometries/DMXGDTFGeometryReference.h b/Engine/Plugins/VirtualProduction/DMX/DMXGDTF/Source/DMXGDTF/Internal/GDTF/Geometries/DMXGDTFGeometryReference.h
--- a/Engine/Plugins/VirtualProduction/DMX/DMXGDTF/Source/DMXGDTF/Internal/GDTF/Geometries/DMXGDTFGeometryReference.h
+++ b/Engine/Plugins/VirtualProduction/DMX/DMXGDTF/Source/DMXGDTF/Internal/GDTF/Geometries/DMXGDTFGeometryReference.h
@@ -44,7 +44,7 @@ namespace UE::DMX::GDTF
 		 * point is Models Collect. If model is not set, the model is taken from
 		 * the referenced geometry.
 		 */
-		FString Model;
+		FName Model;
 
 		/**
 		 * As children, the Geometry Type Reference has a list of a breaks. The count of the children depends on the
diff --git a/Engine/Plugins/VirtualProduction/DMX/DMXGDTF/Source/DMXGDTF/Private/GDTF/Geometries/DMXGDTFGeometryReference.cpp b/Engine/Plugins/VirtualProduction/DMX/DMXGDTF/Source/DMXGDTF/Private/GDTF/Geometries/DMXGDTFGeometryReference.cpp
--- a/Engine/Plugins/VirtualProduction/DMX/DMXGDTF/Source/DMXGDTF/Private/GDTF/Geometries/DMXGDTFGeometryReference.cpp
+++ b/Engine/Plugins/VirtualProduction/DMX/DMXGDTF/Source/DMXGDTF/Private/GDTF/Geometries/DMXGDTFGeometryReference.cpp
@@ -29,7 +29,7 @@ namespace UE::DMX::GDTF
 
 	FXmlNode* FDMXGDTFGeometryReference::CreateXmlNode(FXmlNode& Parent)
 	{
-		const FString DefaultModel = TEXT("");
+		const FName DefaultModel = "";
 
 		const FDMXGDTFXmlNodeBuilder ChildBuilder = FDMXGDTFXmlNodeBuilder(Parent, *this)
 			.SetAttribute(TEXT("Name"), Name)
