DMX: Fix the GDTF exporter to set the correct GDTF mode name again

diff --git a/Engine/Plugins/VirtualProduction/DMX/DMXEngine/Source/DMXEditor/Private/Exporters/DMXFixtureTypeToGDTFConverter.cpp b/Engine/Plugins/VirtualProduction/DMX/DMXEngine/Source/DMXEditor/Private/Exporters/DMXFixtureTypeToGDTFConverter.cpp
--- a/Engine/Plugins/VirtualProduction/DMX/DMXEngine/Source/DMXEditor/Private/Exporters/DMXFixtureTypeToGDTFConverter.cpp
+++ b/Engine/Plugins/VirtualProduction/DMX/DMXEngine/Source/DMXEditor/Private/Exporters/DMXFixtureTypeToGDTFConverter.cpp
@@ -176,11 +176,10 @@ namespace UE::DMX::GDTF
 
 	void FDMXFixtureTypeToGDTFConverter::CreateDMXModes(const UDMXEntityFixtureType& UnrealFixtureType, const TSharedRef<FDMXGDTFFixtureType>& GDTFFixtureType)
 	{
-		int32 ModeCount = 1;
 		for (const FDMXFixtureMode& UnrealMode : UnrealFixtureType.Modes)
 		{
 			const FDMXFixtureModeWithBaseGeometry* ModeWithBaseGeometryPtr = Algo::FindBy(ModesWithBaseGeometry, &UnrealMode, &FDMXFixtureModeWithBaseGeometry::ModePtr);
-			if (!ensureMsgf(ModeWithBaseGeometryPtr, TEXT("%hs: Unexpected cannot find base geometry for DMX Mode '%s'. Failed to convert mode to GDTF."), __FUNCTION__, *UnrealMode.ModeName))
+			if (!ensureMsgf(ModeWithBaseGeometryPtr, TEXT("%hs: Unexpected cannot controlled base geometry for DMX Mode '%s'. Failed to convert mode to GDTF."), __FUNCTION__, *UnrealMode.ModeName))
 			{
 				continue;
 			}
@@ -189,13 +188,11 @@ namespace UE::DMX::GDTF
 			const TSharedRef<FDMXGDTFDMXMode> DMXMode = MakeShared<FDMXGDTFDMXMode>(GDTFFixtureType);
 			GDTFFixtureType->DMXModes.Add(DMXMode);
 
-			DMXMode->Name = *FString::Printf(TEXT("Mode %i %s"), ModeCount, *UnrealMode.ModeName);
+			DMXMode->Name = *UnrealMode.ModeName;
 			DMXMode->Description = TEXT("Unreal Engine generated DMX Mode");
 			DMXMode->Geometry = ModeWithBaseGeometryPtr->BaseGeometry->Name;
 
 			CreateDMXChannels(UnrealMode, DMXMode);
-
-			ModeCount++;
 		}
 	}
 
