DMX: Add an option to display only the latest revision of GDTF modes in the Fixture Type editor

diff --git a/Engine/Plugins/VirtualProduction/DMX/DMXEngine/Source/DMXEditor/Private/Widgets/FixtureType/DMXFixtureTypeModesEditorModeItem.cpp b/Engine/Plugins/VirtualProduction/DMX/DMXEngine/Source/DMXEditor/Private/Widgets/FixtureType/DMXFixtureTypeModesEditorModeItem.cpp
--- a/Engine/Plugins/VirtualProduction/DMX/DMXEngine/Source/DMXEditor/Private/Widgets/FixtureType/DMXFixtureTypeModesEditorModeItem.cpp
+++ b/Engine/Plugins/VirtualProduction/DMX/DMXEngine/Source/DMXEditor/Private/Widgets/FixtureType/DMXFixtureTypeModesEditorModeItem.cpp
@@ -6,8 +6,8 @@
 #include "DMXFixtureTypeSharedData.h"
 #include "DMXRuntimeUtils.h"
 #include "Library/DMXEntityFixtureType.h"
-
 #include "ScopedTransaction.h"
+#include "Internationalization/Regex.h"
 
 
 #define LOCTEXT_NAMESPACE "DMXFixtureTypeModesEditorModeItem"
@@ -18,7 +18,20 @@ FDMXFixtureTypeModesEditorModeItem::FDMXFixtureTypeModesEditorModeItem(const TSh
 	, SharedData(InDMXEditor->GetFixtureTypeSharedData())
 	, WeakDMXEditor(InDMXEditor)
 {
-	ensureMsgf(FixtureType.IsValid() && FixtureType->Modes.IsValidIndex(ModeIndex), TEXT("Invalid Fixture Type or Mode in FDMXFixtureTypeModesEditorModeItem."));
+	if (!ensureMsgf(FixtureType.IsValid() && FixtureType->Modes.IsValidIndex(ModeIndex), TEXT("Invalid Fixture Type or Mode in FDMXFixtureTypeModesEditorModeItem.")))
+	{
+		return;
+	}
+
+	if (!FixtureType->GDTFSource.IsNull())
+	{
+		const FRegexPattern ModeNumberRegexPattern(TEXT("mode (\\d+)"), ERegexPatternFlags::CaseInsensitive);
+		FRegexMatcher ModeNumberRegexMatcher(ModeNumberRegexPattern, FixtureType->Modes[ModeIndex].ModeName);
+		if (ModeNumberRegexMatcher.FindNext())
+		{
+			LexTryParseString(GDTFModeNumber, *ModeNumberRegexMatcher.GetCaptureGroup(1));
+		}
+	}
 }
 
 int32 FDMXFixtureTypeModesEditorModeItem::GetModeIndex() const
diff --git a/Engine/Plugins/VirtualProduction/DMX/DMXEngine/Source/DMXEditor/Private/Widgets/FixtureType/DMXFixtureTypeModesEditorModeItem.h b/Engine/Plugins/VirtualProduction/DMX/DMXEngine/Source/DMXEditor/Private/Widgets/FixtureType/DMXFixtureTypeModesEditorModeItem.h
--- a/Engine/Plugins/VirtualProduction/DMX/DMXEngine/Source/DMXEditor/Private/Widgets/FixtureType/DMXFixtureTypeModesEditorModeItem.h
+++ b/Engine/Plugins/VirtualProduction/DMX/DMXEngine/Source/DMXEditor/Private/Widgets/FixtureType/DMXFixtureTypeModesEditorModeItem.h
@@ -24,6 +24,12 @@ public:
 	/** Returns the Mode Name as Text */
 	FText GetModeName() const;
 
+	/** 
+	 * GDTF modes typically start with "Mode" and a number, e.g. "Mode 3", the latest revision is the last entry.
+	 * This function returns the number of the mode or INDEX_NONE if the mode is not enumerated.
+	 */
+	int32 GetGDTFModeNumber() const { return GDTFModeNumber; }
+
 	/** Returns if the Name is valid for the Mode. If the Mode Name is not valid, returns a non-empty OutInvalidReason */
 	bool IsValidModeName(const FText& InModeName, FText& OutInvalidReason);
 
@@ -34,6 +40,9 @@ private:
 	/** The index of the mode in the Fixture Type's Modes array */
 	int32 ModeIndex = INDEX_NONE;
 
+	/** The number of the mode or INDEX_NONE if the mode is not enumerated. See GetGDTFModeNumber. */
+	int32 GDTFModeNumber = INDEX_NONE;
+
 	/** The Fixture Type which owns the Mode */
 	TWeakObjectPtr<UDMXEntityFixtureType> FixtureType;
 
diff --git a/Engine/Plugins/VirtualProduction/DMX/DMXEngine/Source/DMXEditor/Private/Widgets/FixtureType/SDMXFixtureTypeModesEditor.cpp b/Engine/Plugins/VirtualProduction/DMX/DMXEngine/Source/DMXEditor/Private/Widgets/FixtureType/SDMXFixtureTypeModesEditor.cpp
--- a/Engine/Plugins/VirtualProduction/DMX/DMXEngine/Source/DMXEditor/Private/Widgets/FixtureType/SDMXFixtureTypeModesEditor.cpp
+++ b/Engine/Plugins/VirtualProduction/DMX/DMXEngine/Source/DMXEditor/Private/Widgets/FixtureType/SDMXFixtureTypeModesEditor.cpp
@@ -2,22 +2,22 @@
 
 #include "SDMXFixtureTypeModesEditor.h"
 
+#include "Algo/AnyOf.h"
 #include "DMXEditor.h"
 #include "DMXFixtureTypeSharedData.h"
-#include "SDMXFixtureTypeModesEditorCategoryRow.h"
-#include "SDMXFixtureTypeModesEditorModeRow.h"
-#include "Library/DMXEntityFixtureType.h"
-#include "Widgets/FixtureType/DMXFixtureTypeModesEditorModeItem.h"
-
-#include "Styling/AppStyle.h"
-#include "ScopedTransaction.h"
 #include "Framework/Commands/GenericCommands.h"
 #include "Framework/Commands/UICommandList.h"
 #include "Framework/MultiBox/MultiBoxBuilder.h"
+#include "Library/DMXEntityFixtureType.h"
+#include "ScopedTransaction.h"
+#include "SDMXFixtureTypeModesEditorCategoryRow.h"
+#include "SDMXFixtureTypeModesEditorModeRow.h"
+#include "Styling/AppStyle.h"
+#include "Widgets/FixtureType/DMXFixtureTypeModesEditorModeItem.h"
+#include "Widgets/Input/SCheckBox.h"
 #include "Widgets/SBoxPanel.h"
 #include "Widgets/Views/SListView.h"
 
-
 #define LOCTEXT_NAMESPACE "SDMXFixtureTypeModesEditor"
 
 void SDMXFixtureTypeModesEditor::Construct(const FArguments& InArgs, const TSharedRef<FDMXEditor>& InDMXEditor)
@@ -58,6 +58,42 @@ void SDMXFixtureTypeModesEditor::Construct(const FArguments& InArgs, const TShar
 			.HAlign(HAlign_Fill)
 			.VAlign(VAlign_Fill)
 		]
+
+		// Revision
+		+ SVerticalBox::Slot()
+		.HAlign(HAlign_Fill)
+		.VAlign(VAlign_Fill)
+		.AutoHeight()
+		[
+			SNew(SBorder)
+			.Padding(2.0f)
+			.BorderImage(FAppStyle::GetBrush("ToolPanel.GroupBorder"))
+			.ToolTipText(LOCTEXT("ShowAllRevisionsTooltip", "When checked, all revisions of the GDTF Modes in this GDTF are displayed, otherwise only the latest revisions are displayed. \nThis option is available as this Fixture Type is based on a GDTF. \n\n"))
+			.Visibility(this, &SDMXFixtureTypeModesEditor::GetRevisionCheckBoxVisibility)
+			[
+				SNew(SHorizontalBox)
+
+				+ SHorizontalBox::Slot()
+				.HAlign(HAlign_Right)
+				.VAlign(VAlign_Center)
+				.Padding(8.f, 0.f)
+				.FillWidth(1.f)
+				[
+					SNew(STextBlock)
+					.Font(FAppStyle::GetFontStyle("PropertyWindow.NormalFont"))
+					.Text(LOCTEXT("ShowAllRevisionsCheckboxLabel", "Show all revisions"))
+				]
+
+				+ SHorizontalBox::Slot()
+				.HAlign(HAlign_Right)
+				.VAlign(VAlign_Center)
+				.AutoWidth()
+				[
+					SAssignNew(RevisionCheckBox, SCheckBox)
+					.OnCheckStateChanged(this, &SDMXFixtureTypeModesEditor::OnRevisionCheckBoxStateChanged) 			
+				]
+			]
+		]
 	];
 
 	RegisterCommands();
@@ -102,18 +138,35 @@ void SDMXFixtureTypeModesEditor::RebuildList()
 		// Only single fixture type editing is supported
 		if (SelectedFixtureTypes.Num() == 1 && SelectedFixtureTypes[0].IsValid())
 		{
+			const UDMXEntityFixtureType& FixtureType = *SelectedFixtureTypes[0];
+
+			// Update the revision check box visibility 
+			bHasModeRevisions = false;
+
 			// Rebuild the list source
-			TWeakObjectPtr<UDMXEntityFixtureType> WeakFixtureType = SelectedFixtureTypes[0];
-			if (UDMXEntityFixtureType* FixtureType = WeakFixtureType.Get())
+			const FString SearchString = SearchText.ToString();
+			
+			for (int32 ModeIndex = 0; ModeIndex < FixtureType.Modes.Num(); ModeIndex++)
 			{
-				const FString SearchString = SearchText.ToString();
-				for (int32 ModeIndex = 0; ModeIndex < FixtureType->Modes.Num(); ModeIndex++)
+				if (SearchString.IsEmpty() || FixtureType.Modes[ModeIndex].ModeName.Contains(SearchString, ESearchCase::IgnoreCase))
 				{
-					if (SearchString.IsEmpty() || FixtureType->Modes[ModeIndex].ModeName.Contains(SearchString, ESearchCase::IgnoreCase))
-					{
-						TSharedRef<FDMXFixtureTypeModesEditorModeItem> ModeItem = MakeShared<FDMXFixtureTypeModesEditorModeItem>(DMXEditor.ToSharedRef(), FixtureType, ModeIndex);
-						ListSource.Add(ModeItem);
+					const TSharedRef<FDMXFixtureTypeModesEditorModeItem> ModeItem = MakeShared<FDMXFixtureTypeModesEditorModeItem>(DMXEditor.ToSharedRef(), SelectedFixtureTypes[0], ModeIndex);
+
+					const bool bIsNewerRevision = 
+						!ListSource.IsEmpty() &&
+						!FixtureType.GDTFSource.IsNull() &&
+						ModeItem->GetGDTFModeNumber() != INDEX_NONE &&
+						ModeItem->GetGDTFModeNumber() == ListSource.Last()->GetGDTFModeNumber();
+
+					bHasModeRevisions |= bIsNewerRevision;
+
+					// If only the latest revision should be shown, pop the previous item if it uses the same GDTF mode number.
+					if (bHasModeRevisions && bIsNewerRevision && FixtureType.bShowOnlyLatestGDTFModeRevisions)
+					{				
+						ListSource.Pop();
 					}
+
+					ListSource.Add(ModeItem);
 				}
 			}
 
@@ -132,6 +185,8 @@ void SDMXFixtureTypeModesEditor::RebuildList()
 		}
 		else if (SelectedFixtureTypes.Num() == 0)
 		{
+			bHasModeRevisions = false;
+
 			// Show a warning when no fixture type is selected
 			ListContentBorder->SetContent
 			(
@@ -146,6 +201,8 @@ void SDMXFixtureTypeModesEditor::RebuildList()
 		}
 		else
 		{
+			bHasModeRevisions = false;
+
 			// Show a warning when multiple fixture types are selected
 			ListContentBorder->SetContent
 			(
@@ -242,6 +299,23 @@ void SDMXFixtureTypeModesEditor::OnFixtureTypeSharedDataSelectedModes()
 	}
 }
 
+void SDMXFixtureTypeModesEditor::OnRevisionCheckBoxStateChanged(ECheckBoxState NewCheckState)
+{
+	const TArray<TWeakObjectPtr<UDMXEntityFixtureType>> SelectedFixtureTypes = FixtureTypeSharedData->GetSelectedFixtureTypes();
+	if (ensureMsgf(SelectedFixtureTypes.Num() == 1, TEXT("Trying to Duplicate Modes, but many or no Fixture Type is selected. This should not be possible.")) &&
+		SelectedFixtureTypes[0].IsValid())
+	{
+		SelectedFixtureTypes[0]->bShowOnlyLatestGDTFModeRevisions = NewCheckState != ECheckBoxState::Checked;
+
+		RebuildList();
+	}
+}
+
+EVisibility SDMXFixtureTypeModesEditor::GetRevisionCheckBoxVisibility() const
+{
+	return bHasModeRevisions ? EVisibility::Visible : EVisibility::Collapsed;
+}
+
 TArray<int32> SDMXFixtureTypeModesEditor::GetListSelectionAsModeIndices() const
 {
 	TArray<TSharedPtr<FDMXFixtureTypeModesEditorModeItem>> SelectedModeItems;
diff --git a/Engine/Plugins/VirtualProduction/DMX/DMXEngine/Source/DMXEditor/Private/Widgets/FixtureType/SDMXFixtureTypeModesEditor.h b/Engine/Plugins/VirtualProduction/DMX/DMXEngine/Source/DMXEditor/Private/Widgets/FixtureType/SDMXFixtureTypeModesEditor.h
--- a/Engine/Plugins/VirtualProduction/DMX/DMXEngine/Source/DMXEditor/Private/Widgets/FixtureType/SDMXFixtureTypeModesEditor.h
+++ b/Engine/Plugins/VirtualProduction/DMX/DMXEngine/Source/DMXEditor/Private/Widgets/FixtureType/SDMXFixtureTypeModesEditor.h
@@ -2,24 +2,23 @@
 
 #pragma once
 
-#include "CoreMinimal.h"
+#include "Styling/SlateTypes.h"
 #include "Widgets/SCompoundWidget.h"
 
 class FDMXEditor;
 struct FDMXFixtureMode;
 class FDMXFixtureTypeSharedData;
 class FDMXFixtureTypeModesEditorModeItem;
-class SDMXFixtureTypeModesEditorModeRow;
-class UDMXEntityFixtureType;
-class UDMXLibrary;
-
 class FUICommandList;
 class ITableRow;
 class SBorder;
+class SCheckBox;
+class SDMXFixtureTypeModesEditorModeRow;
 template <typename ItemType> class SListView;
 class STableViewBase;
 class SVerticalBox;
-
+class UDMXEntityFixtureType;
+class UDMXLibrary;
 
 /** The editor for the Fixure Type Modes array */
 class SDMXFixtureTypeModesEditor
@@ -63,7 +62,13 @@ private:
 
 	/** Called when modes were selected in Fixture Type Shared Data */
 	void OnFixtureTypeSharedDataSelectedModes();
-	
+
+	/** Called when the show all revisions check box state changed */
+	void OnRevisionCheckBoxStateChanged(ECheckBoxState NewCheckState);
+
+	/** Returns the visibility of the Revisions section */
+	EVisibility GetRevisionCheckBoxVisibility() const;
+
 	/** Returns the selected Modes in the List as Indices in the Modes array */
 	TArray<int32> GetListSelectionAsModeIndices() const;
 
@@ -85,6 +90,12 @@ private:
 	/** Rows of Mode widgets in the list */
 	TArray<TSharedPtr<SDMXFixtureTypeModesEditorModeRow>> ModeRows;
 
+	/** The revision check box */
+	TSharedPtr<SCheckBox> RevisionCheckBox;
+
+	/** True if mode revisions are available */
+	bool bHasModeRevisions = false;
+
 	/** Search Text entered in the header row */
 	FText SearchText;
 
diff --git a/Engine/Plugins/VirtualProduction/DMX/DMXEngine/Source/DMXRuntime/Public/Library/DMXEntityFixtureType.h b/Engine/Plugins/VirtualProduction/DMX/DMXEngine/Source/DMXRuntime/Public/Library/DMXEntityFixtureType.h
--- a/Engine/Plugins/VirtualProduction/DMX/DMXEngine/Source/DMXRuntime/Public/Library/DMXEntityFixtureType.h
+++ b/Engine/Plugins/VirtualProduction/DMX/DMXEngine/Source/DMXRuntime/Public/Library/DMXEntityFixtureType.h
@@ -586,6 +586,10 @@ public:
 	/** If checked, generates a new GDTF instead of exporting the imported GDTF. This adopts changes in editor but in most cases will result in data loss and is not recommended. */
 	UPROPERTY(EditAnywhere, AdvancedDisplay, Category = "Fixture Settings", Meta = (DisplayName = "Generate GDTF (not recommended)"))
 	bool bExportGeneratedGDTF = false;
+
+	/** If true only shows latest GDTF mode revisions in editor */
+	UPROPERTY()
+	bool bShowOnlyLatestGDTFModeRevisions = true;
 #endif
 
 	//////////////////////////////////////////////////
