Added .Get() to some TObjectPtrs when compared against TSoftObjectPtr, as it results in deprecation warnings on some Clang versions with the updated TSoftObjectPtr.

diff --git a/Engine/Plugins/PCG/Source/PCG/Private/Tests/Elements/PCGStaticMeshSpawnerTest.cpp b/Engine/Plugins/PCG/Source/PCG/Private/Tests/Elements/PCGStaticMeshSpawnerTest.cpp
--- a/Engine/Plugins/PCG/Source/PCG/Private/Tests/Elements/PCGStaticMeshSpawnerTest.cpp
+++ b/Engine/Plugins/PCG/Source/PCG/Private/Tests/Elements/PCGStaticMeshSpawnerTest.cpp
@@ -199,7 +199,7 @@ void TestMeshSelectorByAttribute(
 		UStaticMesh* StaticMesh = ISMC->GetStaticMesh();
 
 		const FPCGByAttributeValidationData* Entry = ValidationDataEntries.FindByPredicate([StaticMesh](const FPCGByAttributeValidationData& Entry) {
-			return Entry.Mesh == StaticMesh;
+			return Entry.Mesh == StaticMesh.Get();
 			});
 
 		if (Test->TestNotNull("Validate instanced mesh exists in MeshEntries", Entry))
diff --git a/Engine/Plugins/VirtualProduction/DMX/DMXEngine/Source/DMXRuntime/Private/Sequencer/MovieSceneDMXLibrarySection.cpp b/Engine/Plugins/VirtualProduction/DMX/DMXEngine/Source/DMXRuntime/Private/Sequencer/MovieSceneDMXLibrarySection.cpp
--- a/Engine/Plugins/VirtualProduction/DMX/DMXEngine/Source/DMXRuntime/Private/Sequencer/MovieSceneDMXLibrarySection.cpp
+++ b/Engine/Plugins/VirtualProduction/DMX/DMXEngine/Source/DMXRuntime/Private/Sequencer/MovieSceneDMXLibrarySection.cpp
@@ -767,7 +767,7 @@ void UMovieSceneDMXLibrarySection::UpdateChannelProxy(bool bResetDefaultChannelV
 	Algo::TransformIf(DMXLibraries, RelevantDMXLibraries,
 		[this](const TSoftObjectPtr<UDMXLibrary>& DMXLibrary)
 		{
-			return DMXLibrary.IsValid() && Algo::FindBy(FixturePatchChannels, DMXLibrary, &FDMXFixturePatchChannel::DMXLibrary) != nullptr;
+			return DMXLibrary.IsValid() && Algo::FindBy(FixturePatchChannels, DMXLibrary, [](const FDMXFixturePatchChannel& Channel) { return Channel.DMXLibrary.Get(); }) != nullptr;
 		},
 		[](const TSoftObjectPtr<UDMXLibrary>& DMXLibrary)
 		{
diff --git a/Engine/Source/Runtime/Engine/Private/HLOD/HLODProxyDesc.cpp b/Engine/Source/Runtime/Engine/Private/HLOD/HLODProxyDesc.cpp
--- a/Engine/Source/Runtime/Engine/Private/HLOD/HLODProxyDesc.cpp
+++ b/Engine/Source/Runtime/Engine/Private/HLOD/HLODProxyDesc.cpp
@@ -374,7 +374,7 @@ ALODActor* UHLODProxyDesc::SpawnLODActor(ULevel* InLevel) const
 	{
 		if (ALODActor* SubLODActor = Cast<ALODActor>(Actor))
 		{
-			if (SubLODActor->ProxyDesc && SubHLODDescs.Contains(SubLODActor->ProxyDesc))
+			if (SubLODActor->ProxyDesc && SubHLODDescs.Contains(TSoftObjectPtr<UHLODProxyDesc>(SubLODActor->ProxyDesc.Get())))
 			{
 				check(SubLODActor != LODActor);
 				SubActorsToAdd.Add(SubLODActor);
