DMX: Fix a bug where the Patch Tool does not show newly created DMX Libraries, fix an issue where UDMXSubsystem::GetDMXLibraries does not return newly created DMX Libraries

diff --git a/Engine/Plugins/VirtualProduction/DMX/DMXEngine/Source/DMXEditor/Private/Widgets/PatchTool/SDMXPatchTool.cpp b/Engine/Plugins/VirtualProduction/DMX/DMXEngine/Source/DMXEditor/Private/Widgets/PatchTool/SDMXPatchTool.cpp
--- a/Engine/Plugins/VirtualProduction/DMX/DMXEngine/Source/DMXEditor/Private/Widgets/PatchTool/SDMXPatchTool.cpp
+++ b/Engine/Plugins/VirtualProduction/DMX/DMXEngine/Source/DMXEditor/Private/Widgets/PatchTool/SDMXPatchTool.cpp
@@ -3,6 +3,7 @@
 #include "SDMXPatchTool.h"
 
 #include "Algo/Transform.h"
+#include "AssetRegistry/AssetRegistryModule.h"
 #include "DMXPatchToolItem.h"
 #include "DMXSubsystem.h"
 #include "Editor.h"
@@ -25,17 +26,49 @@ namespace UE::DMX
 	{}
 
 	SDMXPatchTool::~SDMXPatchTool()
+	{}
+
+	void SDMXPatchTool::Construct(const FArguments& InArgs)
+	{
+		Refresh();
+
+		// Listen to assets being added or removed
+		FAssetRegistryModule& AssetRegistryModule = FModuleManager::GetModuleChecked<FAssetRegistryModule>(TEXT("AssetRegistry"));
+		AssetRegistryModule.Get().OnAssetAdded().AddSP(this, &SDMXPatchTool::OnAssetAddedOrRemoved);
+		AssetRegistryModule.Get().OnAssetRemoved().AddSP(this, &SDMXPatchTool::OnAssetAddedOrRemoved);
+
+		// Listen to DMX Library changes
+		UDMXLibrary::GetOnEntitiesAdded().AddSP(this, &SDMXPatchTool::OnEntitiesAddedOrRemoved);
+		UDMXLibrary::GetOnEntitiesRemoved().AddSP(this, &SDMXPatchTool::OnEntitiesAddedOrRemoved);
+	}
+
+	void SDMXPatchTool::AddReferencedObjects(FReferenceCollector& Collector)
 	{
-		// Unbind from library changes
-		if (PreviouslySelectedLibrary.IsValid())
+		Collector.AddReferencedObject(DMXLibrary);
+		Collector.AddReferencedObject(FixturePatch);
+		Collector.AddReferencedObjects(FixturePatchSource);
+	}
+
+	FString SDMXPatchTool::GetReferencerName() const
+	{
+		return TEXT("SDMXPatchTool");
+	}
+
+	void SDMXPatchTool::RequestRefresh()
+	{
+		if (!RefreshTimerHandle.IsValid())
 		{
-			PreviouslySelectedLibrary->GetOnEntitiesAdded().RemoveAll(this);
-			PreviouslySelectedLibrary->GetOnEntitiesRemoved().RemoveAll(this);
+			RefreshTimerHandle = GEditor->GetTimerManager()->SetTimerForNextTick(FTimerDelegate::CreateSP(this, &SDMXPatchTool::Refresh));
 		}
 	}
 
-	void SDMXPatchTool::Construct(const FArguments& InArgs)
+	void SDMXPatchTool::Refresh()
 	{
+		RefreshTimerHandle.Invalidate();
+
+		// Update the library source and make a valid selection
+		LibrarySource.Reset();
+
 		UDMXSubsystem* Subsystem = UDMXSubsystem::GetDMXSubsystem_Pure();
 		check(Subsystem);
 
@@ -45,6 +78,31 @@ namespace UE::DMX
 				return MakeShared<FDMXPatchToolItem>(SoftDMXLibrary);
 			});
 
+		const bool bCanSelectDMXLibrary = DMXLibrary && Algo::FindBy(LibrarySource, DMXLibrary.Get(), &FDMXPatchToolItem::SoftDMXLibrary) != nullptr;
+		if (!bCanSelectDMXLibrary)
+		{
+			DMXLibrary = !LibrarySource.IsEmpty() ? LibrarySource[0]->SoftDMXLibrary.LoadSynchronous() : nullptr;
+		}
+
+		const TSharedPtr<FDMXPatchToolItem>* SelectDMXLibraryItemPtr = Algo::FindBy(LibrarySource, DMXLibrary.Get(), &FDMXPatchToolItem::SoftDMXLibrary);
+		const TSharedPtr<FDMXPatchToolItem> SelectDMXLibraryItem = SelectDMXLibraryItemPtr ? *SelectDMXLibraryItemPtr : nullptr;
+
+
+		// Update the fixture patch source and make a valid selection
+		FixturePatchSource.Reset();
+		if (DMXLibrary)
+		{
+			FixturePatchSource = ObjectPtrWrap(DMXLibrary->GetEntitiesTypeCast<UDMXEntityFixturePatch>());
+		}
+
+		const bool bCanSelectFixturePatch = FixturePatch && FixturePatchSource.Contains(FixturePatch);
+		if (!bCanSelectFixturePatch)
+		{
+			FixturePatch = !FixturePatchSource.IsEmpty() ? FixturePatchSource[0] : nullptr;
+		}
+
+
+		// Rebuild the widget
 		ChildSlot
 		[
 			SNew(SVerticalBox)
@@ -89,8 +147,15 @@ namespace UE::DMX
 						.OnGenerateWidget(this, &SDMXPatchTool::GenerateLibraryComboBoxEntry)
 						.OnSelectionChanged(this, &SDMXPatchTool::OnLibrarySelected)
 						.OptionsSource(&ObjectPtrDecay(LibrarySource))
+						.InitiallySelectedItem(SelectDMXLibraryItem)
 						[
-							SAssignNew(SelectedLibraryTextBlock, STextBlock)
+							SNew(STextBlock)
+							.Text_Lambda([this]()
+								{
+									return DMXLibrary ?
+										FText::FromString(DMXLibrary->GetName()) :
+										LOCTEXT("NoDMXLibraryAvailableInfo", "No DMX Library available");
+								})
 						]
 					]
 				]
@@ -128,8 +193,15 @@ namespace UE::DMX
 						.OnGenerateWidget(this, &SDMXPatchTool::GenerateFixturePatchComboBoxEntry)
 						.OnSelectionChanged(this, &SDMXPatchTool::OnFixturePatchSelected)
 						.OptionsSource(&ObjectPtrDecay(FixturePatchSource))
+						.InitiallySelectedItem(FixturePatch)
 						[
-							SAssignNew(SelectedFixturePatchTextBlock, STextBlock)
+							SNew(STextBlock)
+							.Text_Lambda([this]()
+								{
+									return FixturePatch ?
+										FText::FromString(FixturePatch->Name) :
+										LOCTEXT("NoPatchAvailableInfo", "No DMX Library selected");
+								})
 						]
 					]
 				]
@@ -185,52 +257,6 @@ namespace UE::DMX
 				]
 			]
 		];
-
-		// Make an initial selection
-		UpdateLibrarySelection();
-		UpdateFixturePatchSelection(DMXLibrary);
-	}
-
-	void SDMXPatchTool::AddReferencedObjects(FReferenceCollector& Collector)
-	{
-		Collector.AddReferencedObject(DMXLibrary);
-		Collector.AddReferencedObjects(FixturePatchSource);
-	}
-
-	void SDMXPatchTool::UpdateLibrarySelection()
-	{
-		check(LibraryComboBox.IsValid());
-		check(FixturePatchComboBox.IsValid());
-
-		if (LibrarySource.Num() > 0)
-		{
-			LibraryComboBox->SetSelectedItem(LibrarySource[0]);
-		}
-		else
-		{
-			SelectedLibraryTextBlock->SetText(LOCTEXT("NoLibraryAvailableAfterUpdate", "No DMX library available"));
-		}
-	}
-
-	void SDMXPatchTool::UpdateFixturePatchSelection(UDMXLibrary* InDMXLibrary)
-	{
-		if (InDMXLibrary)
-		{
-			FixturePatchSource = ObjectPtrWrap(InDMXLibrary->GetEntitiesTypeCast<UDMXEntityFixturePatch>());
-
-			if (FixturePatchSource.Num() > 0)
-			{
-				FixturePatchComboBox->SetSelectedItem(FixturePatchSource[0]);
-			}
-			else
-			{
-				SelectedFixturePatchTextBlock->SetText(LOCTEXT("NoFixturePatchAvailable", "No Fixture Patch available in Library"));
-			}
-		}
-		else
-		{
-			FixturePatchSource.Reset();
-		}
 	}
 
 	FReply SDMXPatchTool::OnAddressIncrementalClicked()
@@ -248,8 +274,8 @@ namespace UE::DMX
 					{
 						if (FixturePatchSource.IsValidIndex(IndexOfPatch))
 						{
-							UDMXEntityFixturePatch* FixturePatch = FixturePatchSource[IndexOfPatch];
-							Component->SetFixturePatch(FixturePatch);
+							UDMXEntityFixturePatch* NextFixturePatch = FixturePatchSource[IndexOfPatch];
+							Component->SetFixturePatch(NextFixturePatch);
 
 							IndexOfPatch++;
 						}
@@ -270,16 +296,13 @@ namespace UE::DMX
 		UDMXEntityFixturePatch* SelectedFixturePatch = FixturePatchComboBox->GetSelectedItem();
 		if (SelectedFixturePatch)
 		{
-			int32 IndexOfPatch = FixturePatchSource.IndexOfByKey(SelectedFixturePatch);
-
 			for (FSelectionIterator It(GEditor->GetSelectedActorIterator()); It; ++It)
 			{
 				if (AActor* Actor = Cast<AActor>(*It))
 				{
 					for (UDMXComponent* Component : TInlineComponentArray<UDMXComponent*>(Actor))
 					{
-						UDMXEntityFixturePatch* FixturePatch = FixturePatchSource[IndexOfPatch];
-						Component->SetFixturePatch(FixturePatch);
+						Component->SetFixturePatch(SelectedFixturePatch);
 					}
 				}
 			}
@@ -303,11 +326,11 @@ namespace UE::DMX
 					{
 						if (FixturePatchSource.IsValidIndex(IndexOfPatch))
 						{
-							UDMXEntityFixturePatch* FixturePatch = FixturePatchSource[IndexOfPatch];
-							Component->SetFixturePatch(FixturePatch);
+							UDMXEntityFixturePatch* NextFixturePatch = FixturePatchSource[IndexOfPatch];
+							Component->SetFixturePatch(NextFixturePatch);
 
 							// Rename
-							Actor->SetActorLabel(FixturePatch->Name);
+							Actor->SetActorLabel(NextFixturePatch->Name);
 
 							IndexOfPatch++;
 						}
@@ -325,7 +348,7 @@ namespace UE::DMX
 
 	TSharedRef<SWidget> SDMXPatchTool::GenerateLibraryComboBoxEntry(TSharedPtr<FDMXPatchToolItem> ItemToAdd)
 	{
-		FText LibraryName = FText::FromString(ItemToAdd->SoftDMXLibrary.GetAssetName());
+		const FText LibraryName = FText::FromString(ItemToAdd->SoftDMXLibrary.GetAssetName());
 
 		return
 			SNew(STextBlock)
@@ -334,40 +357,17 @@ namespace UE::DMX
 
 	void SDMXPatchTool::OnLibrarySelected(TSharedPtr<FDMXPatchToolItem> SelectedItem, ESelectInfo::Type SelectInfo)
 	{
-		// Ignore unchanged selections
-		if (PreviouslySelectedLibrary.IsValid() && PreviouslySelectedLibrary.Get() == SelectedItem->SoftDMXLibrary)
-		{
-			return;
-		}
-
-		// Unbind from previously selected library changes
-		if (PreviouslySelectedLibrary.IsValid())
-		{
-			PreviouslySelectedLibrary->GetOnEntitiesAdded().RemoveAll(this);
-			PreviouslySelectedLibrary->GetOnEntitiesRemoved().RemoveAll(this);
-		}
-		PreviouslySelectedLibrary = DMXLibrary;
-
-		DMXLibrary = Cast<UDMXLibrary>(SelectedItem->SoftDMXLibrary.LoadSynchronous());
-		if (DMXLibrary)
+		if (DMXLibrary.Get() != SelectedItem->SoftDMXLibrary)
 		{
-			// Bind to library edits
-			DMXLibrary->GetOnEntitiesAdded().AddSP(this, &SDMXPatchTool::OnEntitiesAddedOrRemoved);
-			DMXLibrary->GetOnEntitiesRemoved().AddSP(this, &SDMXPatchTool::OnEntitiesAddedOrRemoved);
+			DMXLibrary = SelectedItem->SoftDMXLibrary.LoadSynchronous();
 
-			SelectedLibraryTextBlock->SetText(FText::FromString(DMXLibrary->GetName()));
+			RequestRefresh();
 		}
-		else
-		{
-			SelectedLibraryTextBlock->SetText(LOCTEXT("NoLibraryAvailableAfterSelectionChange", "No DMX Library available"));
-		}
-
-		UpdateFixturePatchSelection(DMXLibrary);
 	}
 
 	TSharedRef<SWidget> SDMXPatchTool::GenerateFixturePatchComboBoxEntry(UDMXEntityFixturePatch* FixturePatchToAdd)
 	{
-		FText FixturePatchName = FText::FromString(FixturePatchToAdd->Name);
+		const FText FixturePatchName = FText::FromString(FixturePatchToAdd->Name);
 
 		return
 			SNew(STextBlock)
@@ -376,44 +376,17 @@ namespace UE::DMX
 
 	void SDMXPatchTool::OnFixturePatchSelected(UDMXEntityFixturePatch* SelectedFixturePatch, ESelectInfo::Type SelectInfo)
 	{
-		if (IsValid(SelectedFixturePatch))
-		{
-			SelectedFixturePatchTextBlock->SetText(FText::FromString(SelectedFixturePatch->Name));
-		}
-		else
-		{
-			SelectedFixturePatchTextBlock->SetText(LOCTEXT("NoFixturePatchAvailable", "No Fixture Patch available in Library"));
-		}
+		FixturePatch = SelectedFixturePatch;
 	}
 
-	void SDMXPatchTool::OnEntitiesAddedOrRemoved(UDMXLibrary* ChangedDMXLibrary, TArray<UDMXEntity*> Entities)
+	void SDMXPatchTool::OnAssetAddedOrRemoved(const FAssetData& AssetData)
 	{
-		UpdateLibrarySelection();
-
-		UDMXEntityFixturePatch* PreviouslySelectedFixturePatch = FixturePatchComboBox->GetSelectedItem();
-
-		if (ChangedDMXLibrary == DMXLibrary)
-		{
-			FixturePatchSource.Reset();
-
-			if (IsValid(DMXLibrary))
-			{
-				FixturePatchSource = ObjectPtrWrap(DMXLibrary->GetEntitiesTypeCast<UDMXEntityFixturePatch>());
-
-				FixturePatchComboBox->RefreshOptions();
+		RequestRefresh();
+	}
 
-				if (IsValid(PreviouslySelectedFixturePatch) && FixturePatchSource.Contains(PreviouslySelectedFixturePatch))
-				{
-					// Restore the previous selection
-					FixturePatchComboBox->SetSelectedItem(PreviouslySelectedFixturePatch);
-				}
-				else
-				{
-					// Update the selection
-					UpdateFixturePatchSelection(DMXLibrary);
-				}
-			}
-		}
+	void SDMXPatchTool::OnEntitiesAddedOrRemoved(UDMXLibrary* ChangedDMXLibrary, TArray<UDMXEntity*> Entities)
+	{
+		RequestRefresh();
 	}
 }
 
diff --git a/Engine/Plugins/VirtualProduction/DMX/DMXEngine/Source/DMXEditor/Private/Widgets/PatchTool/SDMXPatchTool.h b/Engine/Plugins/VirtualProduction/DMX/DMXEngine/Source/DMXEditor/Private/Widgets/PatchTool/SDMXPatchTool.h
--- a/Engine/Plugins/VirtualProduction/DMX/DMXEngine/Source/DMXEditor/Private/Widgets/PatchTool/SDMXPatchTool.h
+++ b/Engine/Plugins/VirtualProduction/DMX/DMXEngine/Source/DMXEditor/Private/Widgets/PatchTool/SDMXPatchTool.h
@@ -3,15 +3,17 @@
 #pragma once
 
 #include "Analytics/DMXEditorToolAnalyticsProvider.h"
+#include "Engine/TimerHandle.h"
 #include "UObject/GCObject.h"
 #include "Widgets/DeclarativeSyntaxSupport.h"
 #include "Widgets/SCompoundWidget.h"
 
-template<typename OptionType> class SComboBox;
 class STextBlock;
 class UDMXEntity;
 class UDMXEntityFixturePatch;
 class UDMXLibrary;
+struct FAssetData;
+template<typename OptionType> class SComboBox;
 
 namespace UE::DMX
 {
@@ -39,17 +41,14 @@ namespace UE::DMX
 	private:
 		// ~Begin FGCObject interface
 		virtual void AddReferencedObjects(FReferenceCollector& Collector) override;
-		virtual FString GetReferencerName() const override
-		{
-			return TEXT("SDMXPatchTool");
-		}
+		virtual FString GetReferencerName() const override;
 		// ~End FGCObject interface
 
-		/** Updates the selected library. Useful on initialization or when assets changed */
-		void UpdateLibrarySelection();
+		/** Refreshes the widget on the next tick */
+		void RequestRefresh();
 
-		/** Updates the available fixture patches from the provided DMX Library */
-		void UpdateFixturePatchSelection(UDMXLibrary* InDMXLibrary);
+		/** Refreshes the widget */
+		void Refresh();
 
 		/** Called when the Address Incremental Button was clicked */
 		FReply OnAddressIncrementalClicked();
@@ -69,37 +68,35 @@ namespace UE::DMX
 		/** Combobox to select a library */
 		TSharedPtr<SComboBox<TSharedPtr<FDMXPatchToolItem>>> LibraryComboBox;
 
-		/** Text block showing the selected library */
-		TSharedPtr<STextBlock> SelectedLibraryTextBlock;
-
-		/** Source for the library combo box */
-		TArray<TSharedPtr<FDMXPatchToolItem>> LibrarySource;
-
-	private:
 		/** Generates an entry in the library combo box */
 		TSharedRef<SWidget> GenerateFixturePatchComboBoxEntry(UDMXEntityFixturePatch* FixturePatchToAdd);
 
 		/** Called when a fixture patch was slected */
 		void OnFixturePatchSelected(UDMXEntityFixturePatch* SelectedFixturePatch, ESelectInfo::Type SelectInfo);
 
+		/** Called when an asset was added or removed */
+		void OnAssetAddedOrRemoved(const FAssetData& AssetData);
+
+		/** Called when the library was edited */
+		void OnEntitiesAddedOrRemoved(UDMXLibrary* Library, TArray<UDMXEntity*> Entities);
+
+		/** Source for the library combo box */
+		TArray<TSharedPtr<FDMXPatchToolItem>> LibrarySource;
+
 		/** Combobox to select a patch within the library */
 		TSharedPtr<SComboBox<UDMXEntityFixturePatch*>> FixturePatchComboBox;
 
-		/** Text block showing the selected fixture patch */
-		TSharedPtr<STextBlock> SelectedFixturePatchTextBlock;
-
 		/** Source for the fixture patch combo box */
 		TArray<TObjectPtr<UDMXEntityFixturePatch>> FixturePatchSource;
 
-	private:
-		/** Called when the library was edited */
-		void OnEntitiesAddedOrRemoved(UDMXLibrary* Library, TArray<UDMXEntity*> Entities);
-
 		/** The currently selected DMX Library. Useful to GC */
 		TObjectPtr<UDMXLibrary> DMXLibrary;
 
-		/** The previously selected library, to unbind from library changes */
-		TWeakObjectPtr<UDMXLibrary> PreviouslySelectedLibrary;
+		/** The currently selected fixture patch */
+		TObjectPtr<UDMXEntityFixturePatch> FixturePatch;
+
+		/** Timer handle for the request refresh method */
+		FTimerHandle RefreshTimerHandle;
 
 		/** The analytics provider for this tool */
 		FDMXEditorToolAnalyticsProvider AnalyticsProvider;
diff --git a/Engine/Plugins/VirtualProduction/DMX/DMXEngine/Source/DMXRuntime/Private/DMXSubsystem.cpp b/Engine/Plugins/VirtualProduction/DMX/DMXEngine/Source/DMXRuntime/Private/DMXSubsystem.cpp
--- a/Engine/Plugins/VirtualProduction/DMX/DMXEngine/Source/DMXRuntime/Private/DMXSubsystem.cpp
+++ b/Engine/Plugins/VirtualProduction/DMX/DMXEngine/Source/DMXRuntime/Private/DMXSubsystem.cpp
@@ -3,6 +3,7 @@
 #include "DMXSubsystem.h"
 
 #include "AssetRegistry/AssetData.h"
+#include "AssetRegistry/AssetRegistryModule.h"
 #include "Async/Async.h"
 #include "DMXAttribute.h"
 #include "DMXConversions.h"
@@ -681,23 +682,15 @@ TArray<UDMXLibrary*> UDMXSubsystem::LoadDMXLibrariesSynchronous() const
 
 TArray<TSoftObjectPtr<UDMXLibrary>> UDMXSubsystem::GetDMXLibraries() const
 {
-	constexpr bool bHasBlueprintClasses = true;
-	UObjectLibrary* LibraryOfDMXLibraries = UObjectLibrary::CreateLibrary(UDMXLibrary::StaticClass(), bHasBlueprintClasses, GIsEditor);
-	LibraryOfDMXLibraries->LoadAssetDataFromPath(TEXT("/Game"));
-	LibraryOfDMXLibraries->LoadAssetsFromAssetData();
-
-	TArray<FAssetData> AssetDatas;
-	LibraryOfDMXLibraries->GetAssetDataList(AssetDatas);
+	TArray<FAssetData> AssetDataArray;
+	const FAssetRegistryModule& AssetRegistryModule = FModuleManager::LoadModuleChecked<FAssetRegistryModule>(TEXT("AssetRegistry"));
+	AssetRegistryModule.Get().GetAssetsByClass(UDMXLibrary::StaticClass()->GetClassPathName(), AssetDataArray);
 
 	TArray<TSoftObjectPtr<UDMXLibrary>> DMXLibraries;
-	Algo::TransformIf(AssetDatas, DMXLibraries,
-		[](const FAssetData& AssetData)
-		{
-			return AssetData.IsValid();
-		},
+	Algo::Transform(AssetDataArray, DMXLibraries,
 		[](const FAssetData& AssetData)
 		{
-			return TSoftObjectPtr<UDMXLibrary>(AssetData.ToSoftObjectPath());
+			return TSoftObjectPtr<UDMXLibrary>(AssetData.GetSoftObjectPath());
 		});
 
 	return DMXLibraries;
