DMX: The DMXEditorModule now offers a public method to export DMX Libraries as MVR files

diff --git a/Engine/Plugins/VirtualProduction/DMX/DMXEngine/Source/DMXEditor/Private/Customizations/DMXFixtureFunctionCustomization.cpp b/Engine/Plugins/VirtualProduction/DMX/DMXEngine/Source/DMXEditor/Private/Customizations/DMXFixtureFunctionCustomization.cpp
--- a/Engine/Plugins/VirtualProduction/DMX/DMXEngine/Source/DMXEditor/Private/Customizations/DMXFixtureFunctionCustomization.cpp
+++ b/Engine/Plugins/VirtualProduction/DMX/DMXEngine/Source/DMXEditor/Private/Customizations/DMXFixtureFunctionCustomization.cpp
@@ -29,6 +29,7 @@ namespace UE::DMX
 
 			// Hide physical properties, they're editor only and confusing in blueprints
 			if (!ChildHandle.IsValid() ||
+				!ChildHandle->IsValidHandle() ||
 				ChildHandle->GetProperty()->GetFName() == FDMXFixtureFunction::GetPhysicalDefaultValuePropertyName() ||
 				ChildHandle->GetProperty()->GetFName() == FDMXFixtureFunction::GetPhysicalUnitPropertyName() ||
 				ChildHandle->GetProperty()->GetFName() == FDMXFixtureFunction::GetPhysicalFromPropertyName() ||
diff --git a/Engine/Plugins/VirtualProduction/DMX/DMXEngine/Source/DMXEditor/Private/DMXEditorModule.cpp b/Engine/Plugins/VirtualProduction/DMX/DMXEngine/Source/DMXEditor/Private/DMXEditorModule.cpp
--- a/Engine/Plugins/VirtualProduction/DMX/DMXEngine/Source/DMXEditor/Private/DMXEditorModule.cpp
+++ b/Engine/Plugins/VirtualProduction/DMX/DMXEngine/Source/DMXEditor/Private/DMXEditorModule.cpp
@@ -23,6 +23,7 @@
 #include "DMXEditorTabNames.h"
 #include "DMXProtocolBlueprintLibrary.h"
 #include "DMXProtocolTypes.h"
+#include "Exporters/DMXMVRExporter.h"
 #include "Framework/Commands/UICommandList.h"
 #include "Framework/Docking/TabManager.h"
 #include "Framework/MultiBox/MultiBoxBuilder.h"
@@ -114,6 +115,14 @@ TSharedRef<FDMXEditor> FDMXEditorModule::CreateEditor(const EToolkitMode::Type M
 	return NewDMXEditor;
 }
 
+void FDMXEditorModule::ExportDMXLibraryAsMVRFile(UDMXLibrary* DMXLibrary, const FString& DesiredName) const
+{
+	if (DMXLibrary)
+	{
+		UE::DMX::FDMXMVRExporter::Export(DMXLibrary, DesiredName);
+	}
+}
+
 void FDMXEditorModule::FixGDTFAssetNames()
 {
 	if (!UObjectInitialized())
diff --git a/Engine/Plugins/VirtualProduction/DMX/DMXEngine/Source/DMXEditor/Private/Exporters/DMXMVRExporter.cpp b/Engine/Plugins/VirtualProduction/DMX/DMXEngine/Source/DMXEditor/Private/Exporters/DMXMVRExporter.cpp
--- a/Engine/Plugins/VirtualProduction/DMX/DMXEngine/Source/DMXEditor/Private/Exporters/DMXMVRExporter.cpp
+++ b/Engine/Plugins/VirtualProduction/DMX/DMXEngine/Source/DMXEditor/Private/Exporters/DMXMVRExporter.cpp
@@ -37,12 +37,12 @@
 
 namespace UE::DMX
 {
-	void FDMXMVRExporter::Export(UDMXLibrary* DMXLibrary)
+	void FDMXMVRExporter::Export(UDMXLibrary* DMXLibrary, const FString& DesiredName)
 	{
 		FDMXMVRExporter Instance;
 		FText ErrorReason;
 		FString FilePathAndName;
-		Instance.ExportInternal(DMXLibrary, ErrorReason, FilePathAndName);
+		Instance.ExportInternal(DMXLibrary, DesiredName, ErrorReason, FilePathAndName);
 
 		if (ErrorReason.IsEmpty())
 		{
@@ -61,26 +61,29 @@ namespace UE::DMX
 		}
 	}
 
-	void FDMXMVRExporter::ExportInternal(UDMXLibrary* DMXLibrary, FText& OutErrorReason, FString& OutFilePathAndName)
+	void FDMXMVRExporter::ExportInternal(UDMXLibrary* InDMXLibrary, const FString& InDesiredName, FText& OutErrorReason, FString& OutFilePathAndName)
 	{
-		if (!ensureAlwaysMsgf(DMXLibrary, TEXT("Trying to export DMX Library '%s' as MVR file, but the DMX Library is invalid."), *DMXLibrary->GetName()))
+		if (!ensureAlwaysMsgf(InDMXLibrary, TEXT("Trying to export DMX Library '%s' as MVR file, but the DMX Library is invalid."), *InDMXLibrary->GetName()))
 		{
 			OutErrorReason = LOCTEXT("MVRExportDMXLibraryInvalidReason", "DMX Library {0} is invalid. Cannot export MVR file.");
 			return;
 		}
 
-		UpdateExportOptions(*DMXLibrary);
+		UpdateExportOptions(*InDMXLibrary);
 		if (GetDefault<UDMXMVRExportOptions>()->bCanceled)
 		{
 			OutErrorReason = LOCTEXT("MVRExportCanceledInvalidReason", "Canceled MVR export.");
 			return;
 		}
 
-		OutFilePathAndName = [DMXLibrary]()
+		OutFilePathAndName = [InDMXLibrary, &InDesiredName]()
 			{					
 				UDMXEditorSettings* DMXEditorSettings = GetMutableDefault<UDMXEditorSettings>();
 				const FString LastMVRExportPath = DMXEditorSettings->LastMVRExportPath;
 				const FString DefaultPath = FPaths::DirectoryExists(LastMVRExportPath) ? LastMVRExportPath : FPaths::ProjectSavedDir();
+				const FString DefaultFileName = InDesiredName.IsEmpty() ?
+					InDMXLibrary->GetName() + TEXT(".mvr") :
+					FPaths::GetBaseFilename(InDesiredName) + TEXT(".mvr");
 
 				IDesktopPlatform* DesktopPlatform = FDesktopPlatformModule::Get();
 				if (!DesktopPlatform)
@@ -93,7 +96,7 @@ namespace UE::DMX
 					FSlateApplication::Get().FindBestParentWindowHandleForDialogs(nullptr),
 					LOCTEXT("ExportMVR", "Export MVR").ToString(),
 					DefaultPath,
-					DMXLibrary->GetName() + TEXT(".mvr"),
+					DefaultFileName,
 					TEXT("My Virtual Rig (*.mvr)|*.mvr"),
 					EFileDialogFlags::None,
 					SaveFilenames);
@@ -110,8 +113,8 @@ namespace UE::DMX
 			}();
 
 		// Create a copy of the library's general scene description, so transforms exported aren't written to the dmx library
-		UDMXMVRGeneralSceneDescription* TempGeneralSceneDescription = DuplicateObject<UDMXMVRGeneralSceneDescription>(DMXLibrary->GetLazyGeneralSceneDescription(), GetTransientPackage());
-		if (!ensureAlwaysMsgf(TempGeneralSceneDescription, TEXT("Trying to export DMX Library '%s' as MVR file, but its General Scene Description is invalid."), *DMXLibrary->GetName()))
+		UDMXMVRGeneralSceneDescription* TempGeneralSceneDescription = DuplicateObject<UDMXMVRGeneralSceneDescription>(InDMXLibrary->GetLazyGeneralSceneDescription(), GetTransientPackage());
+		if (!ensureAlwaysMsgf(TempGeneralSceneDescription, TEXT("Trying to export DMX Library '%s' as MVR file, but its General Scene Description is invalid."), *InDMXLibrary->GetName()))
 		{
 			OutErrorReason = FText::Format(LOCTEXT("MVRExportGeneralSceneDescriptionInvalidReason", "DMX Library is invalid. Cannot export {0}."), FText::FromString(OutFilePathAndName));
 			return;
@@ -127,7 +130,7 @@ namespace UE::DMX
 		WorldParams.bExportPatchesNotPresentInWorld = ExportOptions->bExportPatchesNotPresentInWorld;
 		WorldParams.bUseTransformsFromLevel = ExportOptions->bUseTransformsFromLevel;
 
-		TempGeneralSceneDescription->WriteDMXLibrary(*DMXLibrary, WorldParams);
+		TempGeneralSceneDescription->WriteDMXLibrary(*InDMXLibrary, WorldParams);
 
 		const TSharedRef<FDMXZipper> Zip = MakeShared<FDMXZipper>();
 		if (!ZipGeneralSceneDescription(Zip, TempGeneralSceneDescription, OutErrorReason))
@@ -135,7 +138,7 @@ namespace UE::DMX
 			return;
 		}
 
-		if (!ZipGDTFs(Zip, DMXLibrary))
+		if (!ZipGDTFs(Zip, InDMXLibrary))
 		{
 			OutErrorReason = FText::Format(LOCTEXT("MVRExportZipGDTFsFailedReason", "Some Fixture Types could not be converted to GDTF. Exported MVR to {0}."), FText::FromString(OutFilePathAndName));
 			// Allow continuation of export
diff --git a/Engine/Plugins/VirtualProduction/DMX/DMXEngine/Source/DMXEditor/Private/Exporters/DMXMVRExporter.h b/Engine/Plugins/VirtualProduction/DMX/DMXEngine/Source/DMXEditor/Private/Exporters/DMXMVRExporter.h
--- a/Engine/Plugins/VirtualProduction/DMX/DMXEngine/Source/DMXEditor/Private/Exporters/DMXMVRExporter.h
+++ b/Engine/Plugins/VirtualProduction/DMX/DMXEngine/Source/DMXEditor/Private/Exporters/DMXMVRExporter.h
@@ -23,11 +23,11 @@ namespace UE::DMX
 	{
 	public:
 		/** Exports the DMX Library as MVR File */
-		static void Export(UDMXLibrary* DMXLibrary);
+		static void Export(UDMXLibrary* DMXLibrary, const FString& DesiredName = TEXT(""));
 
 	private:
 		/** Exports the DMX Library as MVR File. If OutErrorReason is not empty there were issues with the export. */
-		void ExportInternal(UDMXLibrary* DMXLibrary, FText& OutErrorReason, FString& OutFilePathAndName);
+		void ExportInternal(UDMXLibrary* InDMXLibrary, const FString& InDesiredName, FText& OutErrorReason, FString& OutFilePathAndName);
 
 		/** Updates the MVR export options. Returns false if the import was canceled */
 		void UpdateExportOptions(const UDMXLibrary& DMXLibrary) const;
diff --git a/Engine/Plugins/VirtualProduction/DMX/DMXEngine/Source/DMXEditor/Private/Widgets/FixturePatch/SDMXPatchedUniverse.cpp b/Engine/Plugins/VirtualProduction/DMX/DMXEngine/Source/DMXEditor/Private/Widgets/FixturePatch/SDMXPatchedUniverse.cpp
--- a/Engine/Plugins/VirtualProduction/DMX/DMXEngine/Source/DMXEditor/Private/Widgets/FixturePatch/SDMXPatchedUniverse.cpp
+++ b/Engine/Plugins/VirtualProduction/DMX/DMXEngine/Source/DMXEditor/Private/Widgets/FixturePatch/SDMXPatchedUniverse.cpp
@@ -438,12 +438,15 @@ void SDMXPatchedUniverse::SetUniverseIDInternal(int32 NewUniverseID)
 		for (UDMXEntityFixturePatch* FixturePatch : PatchesInUniverse)
 		{
 			TSharedPtr<FDMXFixturePatchNode> Node = FindPatchNode(FixturePatch);
-			if (!Node.IsValid())
+			if (!Node.IsValid() && IsValid(FixturePatch))
 			{
 				Node = FDMXFixturePatchNode::Create(WeakDMXEditor, FixturePatch);
 			}
 
-			FindOrAdd(Node.ToSharedRef());
+			if (Node.IsValid())
+			{
+				FindOrAdd(Node.ToSharedRef());
+			}
 		}
 
 		// Update the channel connectors' Universe ID
diff --git a/Engine/Plugins/VirtualProduction/DMX/DMXEngine/Source/DMXEditor/Private/Widgets/SNameListPicker.cpp b/Engine/Plugins/VirtualProduction/DMX/DMXEngine/Source/DMXEditor/Private/Widgets/SNameListPicker.cpp
--- a/Engine/Plugins/VirtualProduction/DMX/DMXEngine/Source/DMXEditor/Private/Widgets/SNameListPicker.cpp
+++ b/Engine/Plugins/VirtualProduction/DMX/DMXEngine/Source/DMXEditor/Private/Widgets/SNameListPicker.cpp
@@ -303,7 +303,7 @@ FText SNameListPicker::GetCurrentNameLabel() const
 		return LOCTEXT("MultipleValuesText", "<multiple values>");
 	}
 
-	const FName CurrentName = ValueAttribute.Get();
+	const FName CurrentName = ValueAttribute.IsSet() ? ValueAttribute.Get() : NAME_None;
 	if (CurrentName.IsEqual(NAME_None))
 	{
 		return NoneLabel;
diff --git a/Engine/Plugins/VirtualProduction/DMX/DMXEngine/Source/DMXEditor/Public/DMXEditorModule.h b/Engine/Plugins/VirtualProduction/DMX/DMXEngine/Source/DMXEditor/Public/DMXEditorModule.h
--- a/Engine/Plugins/VirtualProduction/DMX/DMXEngine/Source/DMXEditor/Public/DMXEditorModule.h
+++ b/Engine/Plugins/VirtualProduction/DMX/DMXEngine/Source/DMXEditor/Public/DMXEditorModule.h
@@ -55,6 +55,9 @@ public:
 	 */
 	TSharedRef<FDMXEditor> CreateEditor( const EToolkitMode::Type Mode, const TSharedPtr< class IToolkitHost >& InitToolkitHost, UDMXLibrary* DMXLibrary );
 
+	/** Exports a DMX Library as MVR file */
+	void ExportDMXLibraryAsMVRFile(UDMXLibrary* DMXLibrary, const FString& DesiredName = TEXT("")) const;
+
 	/** Get Level Editor Toolbar Menu extender */
 	TSharedPtr<FExtender> GetLevelEditorToolbarDMXMenuExtender() const { return LevelEditorToolbarDMXMenuExtender; }
 
diff --git a/Engine/Plugins/VirtualProduction/DMX/DMXEngine/Source/DMXRuntime/Private/Library/DMXEntityFixturePatch.cpp b/Engine/Plugins/VirtualProduction/DMX/DMXEngine/Source/DMXRuntime/Private/Library/DMXEntityFixturePatch.cpp
--- a/Engine/Plugins/VirtualProduction/DMX/DMXEngine/Source/DMXRuntime/Private/Library/DMXEntityFixturePatch.cpp
+++ b/Engine/Plugins/VirtualProduction/DMX/DMXEngine/Source/DMXRuntime/Private/Library/DMXEntityFixturePatch.cpp
@@ -139,7 +139,6 @@ void UDMXEntityFixturePatch::RemoveFixturePatchFromLibrary(FDMXEntityFixturePatc
 			DMXLibrary->Modify();
 			FixturePatch->Modify();
 
-			FixturePatch->SetFixtureType(nullptr);
 			FixturePatch->Destroy();
 		}
 	}
diff --git a/Engine/Plugins/VirtualProduction/DMX/DMXEngine/Source/DMXRuntime/Private/Library/DMXLibrary.cpp b/Engine/Plugins/VirtualProduction/DMX/DMXEngine/Source/DMXRuntime/Private/Library/DMXLibrary.cpp
--- a/Engine/Plugins/VirtualProduction/DMX/DMXEngine/Source/DMXRuntime/Private/Library/DMXLibrary.cpp
+++ b/Engine/Plugins/VirtualProduction/DMX/DMXEngine/Source/DMXRuntime/Private/Library/DMXLibrary.cpp
@@ -28,14 +28,14 @@ FDMXOnEntityArrayChangedDelegate UDMXLibrary::OnEntitiesRemovedDelegate;
 UDMXLibrary::UDMXLibrary()
 {
 	const FName GeneralSceneDescriptionName = FName(GetName() + TEXT("_MVRGeneralSceneDescription"));
-	GeneralSceneDescription = NewObject<UDMXMVRGeneralSceneDescription>(this, GeneralSceneDescriptionName);
+	GeneralSceneDescription = NewObject<UDMXMVRGeneralSceneDescription>(this, GeneralSceneDescriptionName, RF_Public);
 }
 
 void UDMXLibrary::PostInitProperties()
 {
 	Super::PostInitProperties();
 
-	if (!HasAnyFlags(RF_ClassDefaultObject))
+	if (!IsTemplate())
 	{
 		FDMXPortManager::Get().OnPortsChanged.AddUObject(this, &UDMXLibrary::UpdatePorts);
 		UpdatePorts();
