DMX: Update modes list to only show the 'revision' option when the fixture type is based on a GDTF, improve tooltip

#rnx

diff --git a/Engine/Plugins/VirtualProduction/DMX/DMXEngine/Source/DMXEditor/Private/Widgets/FixtureType/SDMXFixtureTypeModesEditor.cpp b/Engine/Plugins/VirtualProduction/DMX/DMXEngine/Source/DMXEditor/Private/Widgets/FixtureType/SDMXFixtureTypeModesEditor.cpp
--- a/Engine/Plugins/VirtualProduction/DMX/DMXEngine/Source/DMXEditor/Private/Widgets/FixtureType/SDMXFixtureTypeModesEditor.cpp
+++ b/Engine/Plugins/VirtualProduction/DMX/DMXEngine/Source/DMXEditor/Private/Widgets/FixtureType/SDMXFixtureTypeModesEditor.cpp
@@ -68,7 +68,7 @@ void SDMXFixtureTypeModesEditor::Construct(const FArguments& InArgs, const TShar
 			SNew(SBorder)
 			.Padding(2.0f)
 			.BorderImage(FAppStyle::GetBrush("ToolPanel.GroupBorder"))
-			.ToolTipText(LOCTEXT("ShowAllRevisionsTooltip", "When checked, all revisions of the GDTF Modes in this GDTF are displayed, otherwise only the latest revisions are displayed. \nThis option is available as this Fixture Type is based on a GDTF. \n\n"))
+			.ToolTipText(LOCTEXT("ShowAllRevisionsTooltip", "When checked, all revisions of the GDTF Modes in this GDTF are displayed, otherwise only the latest revisions are displayed.\nThis option is available as this Fixture Type is based on a GDTF and contains Modes with more than one revision."))
 			.Visibility(this, &SDMXFixtureTypeModesEditor::GetRevisionCheckBoxVisibility)
 			[
 				SNew(SHorizontalBox)
@@ -302,7 +302,8 @@ void SDMXFixtureTypeModesEditor::OnFixtureTypeSharedDataSelectedModes()
 void SDMXFixtureTypeModesEditor::OnRevisionCheckBoxStateChanged(ECheckBoxState NewCheckState)
 {
 	const TArray<TWeakObjectPtr<UDMXEntityFixtureType>> SelectedFixtureTypes = FixtureTypeSharedData->GetSelectedFixtureTypes();
-	if (ensureMsgf(SelectedFixtureTypes.Num() == 1, TEXT("Trying to Duplicate Modes, but many or no Fixture Type is selected. This should not be possible.")) &&
+	if (ensureMsgf(SelectedFixtureTypes.Num() == 1, TEXT("Trying to set if latest mode revisions are shown for fixture type, but fixture type is invalid.")) &&
+		ensureMsgf(!SelectedFixtureTypes[0]->GDTFSource.IsNull(), TEXT("Trying to set if latest mode revisions are shown for fixture type, but fixture type has no GDTF.")) &&
 		SelectedFixtureTypes[0].IsValid())
 	{
 		SelectedFixtureTypes[0]->bShowOnlyLatestGDTFModeRevisions = NewCheckState != ECheckBoxState::Checked;
