DMX: Clean up DMX Engine blueprint nodes, improve comments, clean up API

diff --git a/Engine/Plugins/VirtualProduction/DMX/DMXEngine/Config/BaseDMXEngine.ini b/Engine/Plugins/VirtualProduction/DMX/DMXEngine/Config/BaseDMXEngine.ini
--- a/Engine/Plugins/VirtualProduction/DMX/DMXEngine/Config/BaseDMXEngine.ini
+++ b/Engine/Plugins/VirtualProduction/DMX/DMXEngine/Config/BaseDMXEngine.ini
@@ -5,3 +5,5 @@
 ; UE5.4 to UE5.5
 +FunctionRedirects=(OldName="GetAllDMXLibraries",NewName="/Script/DMXEngine.DMXSubsystem.LoadDMXLibrariesSynchronous")
 +StructRedirects=(OldName="DMXColorCIE1931xyY", NewName="/Script/DMXGDTF.DMXGDTFColorCIE1931xyY")
++FunctionRedirects=(OldName="DMXEntityFixturePatch.GetAttributesValues",NewName="DMXEntityFixturePatch.GetAttributeValues")
++FunctionRedirects=(OldName="DMXEntityFixturePatch.GetNormalizedAttributesValues",NewName="DMXEntityFixturePatch.GetNormalizedAttributeValues")
diff --git a/Engine/Plugins/VirtualProduction/DMX/DMXEngine/Source/DMXBlueprintGraph/Private/K2Node_CastPatchToType.cpp b/Engine/Plugins/VirtualProduction/DMX/DMXEngine/Source/DMXBlueprintGraph/Private/K2Node_CastPatchToType.cpp
--- a/Engine/Plugins/VirtualProduction/DMX/DMXEngine/Source/DMXBlueprintGraph/Private/K2Node_CastPatchToType.cpp
+++ b/Engine/Plugins/VirtualProduction/DMX/DMXEngine/Source/DMXBlueprintGraph/Private/K2Node_CastPatchToType.cpp
@@ -355,37 +355,37 @@ void UDEPRECATED_K2Node_CastPatchToType::ExpandNode(class FKismetCompilerContext
  	LastThenPin = Branch_Node->GetThenPin();
 
 	// NODE 4. UDMXSubsystem::GetAttributesMap
-	static const FName FuncName_GetAttributesMapForPatch = GET_FUNCTION_NAME_CHECKED(UDMXSubsystem, GetFunctionsMapForPatch);
-	UFunction* FuncPtr_GetAttributesMapForPatch = FindUField<UFunction>(UDMXSubsystem::StaticClass(), FuncName_GetAttributesMapForPatch);
-	check(FuncPtr_GetAttributesMapForPatch);
+	static const FName FuncName_GetFunctionsMap = GET_FUNCTION_NAME_CHECKED(UDMXSubsystem, GetFunctionsMap);
+	UFunction* FuncPtr_GetFunctionsMap = FindUField<UFunction>(UDMXSubsystem::StaticClass(), FuncName_GetFunctionsMap);
+	check(FuncPtr_GetFunctionsMap);
 
-	UK2Node_CallFunction* GetAttributesMapForPatch_Node = CompilerContext.SpawnIntermediateNode<UK2Node_CallFunction>(this, SourceGraph);
-	GetAttributesMapForPatch_Node->SetFromFunction(FuncPtr_GetAttributesMapForPatch);
-	GetAttributesMapForPatch_Node->AllocateDefaultPins();
+	UK2Node_CallFunction* GetFunctionsMap_Node = CompilerContext.SpawnIntermediateNode<UK2Node_CallFunction>(this, SourceGraph);
+	GetFunctionsMap_Node->SetFromFunction(FuncPtr_GetFunctionsMap);
+	GetFunctionsMap_Node->AllocateDefaultPins();
 
-	UEdGraphPin* GetAttributesMap_In_Self = GetAttributesMapForPatch_Node->FindPin(UEdGraphSchema_K2::PN_Self);
+	UEdGraphPin* GetAttributesMap_In_Self = GetFunctionsMap_Node->FindPin(UEdGraphSchema_K2::PN_Self);
 	if (GetAttributesMap_In_Self == nullptr)
 	{
 		CompilerContext.MessageLog.Error(*LOCTEXT("SelfPin", "Self: Pin doesn't exists. @@").ToString(), this);
 		return;
 	}
 
-	UEdGraphPin* GetAttributesMap_In_Exec = GetAttributesMapForPatch_Node->GetExecPin();
-	UEdGraphPin* GetAttributesMap_In_FixturePatch = GetAttributesMapForPatch_Node->FindPin(TEXT("InFixturePatch"));
+	UEdGraphPin* GetAttributesMap_In_Exec = GetFunctionsMap_Node->GetExecPin();
+	UEdGraphPin* GetAttributesMap_In_FixturePatch = GetFunctionsMap_Node->FindPin(TEXT("InFixturePatch"));
 	if (GetAttributesMap_In_FixturePatch == nullptr)
 	{
 		CompilerContext.MessageLog.Error(*LOCTEXT("InFixturePatchPin", "InFixturePatch: Pin doesn't exists. @@").ToString(), this);
 		return;
 	}
 
-	UEdGraphPin* GetAttributesMap_InOut_AttributesMap = GetAttributesMapForPatch_Node->FindPin(TEXT("OutAttributesMap"));
+	UEdGraphPin* GetAttributesMap_InOut_AttributesMap = GetFunctionsMap_Node->FindPin(TEXT("OutAttributesMap"));
 	if (GetAttributesMap_InOut_AttributesMap == nullptr)
 	{
 		CompilerContext.MessageLog.Error(*LOCTEXT("OutAttributesMapPin", "OutAttributesMap: Pin doesn't exists. @@").ToString(), this);
 		return;
 	}
 
-	UEdGraphPin* GetAttributesMap_Out_Then = GetAttributesMapForPatch_Node->GetThenPin();
+	UEdGraphPin* GetAttributesMap_Out_Then = GetFunctionsMap_Node->GetThenPin();
 
 	// inputs
  	K2Schema->TryCreateConnection(GetAttributesMap_In_Self, DMXSubsytem_ReturnValue);
@@ -430,7 +430,10 @@ void UDEPRECATED_K2Node_CastPatchToType::ExpandNode(class FKismetCompilerContext
 			UEdGraphPin* IN_FuncName = Map_IN_Names[PairIndex];
 			UEdGraphPin* OUT_FuncInt = Map_OUT_Ints[PairIndex];
 
+			PRAGMA_DISABLE_DEPRECATION_WARNINGS
 			const FName FuncName_GetAttributesValue = GET_FUNCTION_NAME_CHECKED(UDMXSubsystem, GetFunctionsValue);
+			PRAGMA_ENABLE_DEPRECATION_WARNINGS 
+
 			UFunction* FuncPtr_GetAttributesValue = FindUField<UFunction>(UDMXSubsystem::StaticClass(), FuncName_GetAttributesValue);
 			check(FuncPtr_GetAttributesValue);
 
diff --git a/Engine/Plugins/VirtualProduction/DMX/DMXEngine/Source/DMXEditor/Private/DMXEditorUtils.cpp b/Engine/Plugins/VirtualProduction/DMX/DMXEngine/Source/DMXEditor/Private/DMXEditorUtils.cpp
--- a/Engine/Plugins/VirtualProduction/DMX/DMXEngine/Source/DMXEditor/Private/DMXEditorUtils.cpp
+++ b/Engine/Plugins/VirtualProduction/DMX/DMXEngine/Source/DMXEditor/Private/DMXEditorUtils.cpp
@@ -84,156 +84,6 @@ protected:
 	//~ End FCustomizableTextObjectFactory implementation
 };
 
-
-FString FDMXEditorUtils::GenerateUniqueNameFromExisting(const TSet<FString>& InExistingNames, const FString& InBaseName)
-{
-	// DEPRECATED 5.0
-
-	if (!InBaseName.IsEmpty() && !InExistingNames.Contains(InBaseName))
-	{
-		return InBaseName;
-	}
-
-	FString FinalName;
-	FString BaseName;
-
-	int32 Index = 0;
-	if (InBaseName.IsEmpty())
-	{
-		BaseName = TEXT("Default name");
-	}
-	else
-	{
-		// If there's an index at the end of the name, start from there
-		FDMXRuntimeUtils::GetNameAndIndexFromString(InBaseName, BaseName, Index);
-	}
-
-	int32 Count = (Index == 0) ? 1 : Index;
-	FinalName = BaseName;
-	// Add Count to the BaseName, increasing Count, until it's a non-existent name
-	do
-	{
-		// Calculate the number of digits in the number, adding 2 (1 extra to correctly count digits, another to account for the '_' that will be added to the name
-		int32 CountLength = Count > 0 ? (int32)FGenericPlatformMath::LogX(10.0f, Count) + 2 : 2;
-
-		// If the length of the final string will be too long, cut off the end so we can fit the number
-		if (CountLength + BaseName.Len() >= NAME_SIZE)
-		{
-			BaseName = BaseName.Left(NAME_SIZE - CountLength - 1);
-		}
-
-		FinalName = FString::Printf(TEXT("%s_%d"), *BaseName, Count);
-		++Count;
-	} while (InExistingNames.Contains(FinalName));
-
-	return FinalName;
-}
-
-FString FDMXEditorUtils::FindUniqueEntityName(const UDMXLibrary* InLibrary, TSubclassOf<UDMXEntity> InEntityClass, const FString& InBaseName /*= TEXT("")*/)
-{
-	// DEPRECATED 5.0
-
-	check(InLibrary != nullptr);
-
-	// Get existing names for the current entity type
-	TSet<FString> EntityNames;
-	InLibrary->ForEachEntityOfType(InEntityClass, [&EntityNames](UDMXEntity* Entity)
-		{
-			EntityNames.Add(Entity->GetDisplayName());
-		});
-
-	FString BaseName = InBaseName;
-
-	// If no base name was set, use the entity class name as base
-	if (BaseName.IsEmpty() && InEntityClass.Get())
-	{
-		BaseName = InEntityClass.Get()->GetDisplayNameText().ToString();
-	}
-
-	return FDMXRuntimeUtils::GenerateUniqueNameFromExisting(EntityNames, BaseName);
-}
-
-void FDMXEditorUtils::SetNewFixtureFunctionsNames(UDMXEntityFixtureType* InFixtureType)
-{
-	// DEPRECATED 5.0
-
-	check(InFixtureType != nullptr);
-
-	// We'll only populate this Set if we find an item with no name.
-	// Otherwise we can save some for loops.
-	TSet<FString> ModesNames;
-
-	// Iterate over all of the Fixture's Modes and Functions creating names for the
-	// ones with a blank name.
-	for (FDMXFixtureMode& Mode : InFixtureType->Modes)
-	{
-		// Do we need to name this mode?
-		if (Mode.ModeName.IsEmpty())
-		{
-			// Cache existing names only once, when needed.
-			if (ModesNames.Num() == 0)
-			{
-				// Cache the existing modes' names
-				for (FDMXFixtureMode& NamedMode : InFixtureType->Modes)
-				{
-					if (!NamedMode.ModeName.IsEmpty())
-					{
-						ModesNames.Add(NamedMode.ModeName);
-					}
-				}
-			}
-
-			Mode.ModeName = FDMXRuntimeUtils::GenerateUniqueNameFromExisting(ModesNames, TEXT("Mode"));
-			ModesNames.Add(Mode.ModeName);
-		}
-
-		// Name this mode's functions
-		TSet<FString> FunctionsNames;
-		for (FDMXFixtureFunction& Function : Mode.Functions)
-		{
-			if (Function.FunctionName.IsEmpty())
-			{
-				// Cache existing names only once, when needed.
-				if (FunctionsNames.Num() == 0)
-				{
-					for (FDMXFixtureFunction& NamedFunction : Mode.Functions)
-					{
-						if (!NamedFunction.FunctionName.IsEmpty())
-						{
-							FunctionsNames.Add(NamedFunction.FunctionName);
-						}
-					}
-				}
-
-				Function.FunctionName = FDMXRuntimeUtils::GenerateUniqueNameFromExisting(FunctionsNames, TEXT("Function"));
-				FunctionsNames.Add(Function.FunctionName);
-			}
-		}
-	}
-}
-
-bool FDMXEditorUtils::AddEntity(UDMXLibrary* InLibrary, const FString& NewEntityName, TSubclassOf<UDMXEntity> NewEntityClass, UDMXEntity** OutNewEntity /*= nullptr*/)
-{
-	// DEPRECATED 5.0
-
-	// Don't allow entities with empty names
-	if (NewEntityName.IsEmpty())
-	{
-		return false;
-	}
-
-	// Mark library as pending save and store current state for undo
-	const FScopedTransaction NewEntityTransaction(LOCTEXT("NewEntityTransaction", "Add new Entity to DMX Library"));
-	InLibrary->Modify();
-
-	// Create new entity 
-	PRAGMA_DISABLE_DEPRECATION_WARNINGS
-	*OutNewEntity = InLibrary->GetOrCreateEntityObject(NewEntityName, NewEntityClass);
-	PRAGMA_ENABLE_DEPRECATION_WARNINGS
-
-	return true;
-}
-
 bool FDMXEditorUtils::ValidateEntityName(const FString& NewEntityName, const UDMXLibrary* InLibrary, UClass* InEntityClass, FText& OutReason)
 {
 	if (NewEntityName.Len() >= NAME_SIZE)
@@ -317,87 +167,6 @@ bool FDMXEditorUtils::IsEntityUsed(const UDMXLibrary* InLibrary, const UDMXEntit
 	return false;
 }
 
-void FDMXEditorUtils::RemoveEntities(UDMXLibrary* InLibrary, const TArray<UDMXEntity*>& InEntities)
-{
-	if (InLibrary)
-	{
-		TArray<UDMXEntity*> EntitiesInUse;
-		for (UDMXEntity* Entity : InEntities)
-		{
-			if (FDMXEditorUtils::IsEntityUsed(InLibrary, Entity))
-			{
-				EntitiesInUse.Add(Entity);
-			}
-		}
-
-		// Confirm deletion of Entities in use, if any
-		if (EntitiesInUse.Num() > 0)
-		{
-			FText ConfirmDelete;
-
-			// Confirmation text for a single entity in use
-			if (EntitiesInUse.Num() == 1)
-			{
-				ConfirmDelete = FText::Format(LOCTEXT("ConfirmDeleteEntityInUse", "Entity \"{0}\" is in use! Do you really want to delete it?"),
-					FText::FromString(EntitiesInUse[0]->Name));
-			}
-			// Confirmation text for when all of the selected entities are in use
-			else if (EntitiesInUse.Num() == InEntities.Num())
-			{
-				ConfirmDelete = LOCTEXT("ConfirmDeleteAllEntitiesInUse", "All selected entities are in use! Do you really want to delete them?");
-			}
-			// Confirmation text for multiple entities, but not so much that would make the dialog huge
-			else if (EntitiesInUse.Num() > 1 && EntitiesInUse.Num() <= 10)
-			{
-				FString EntitiesNames;
-				for (UDMXEntity* Entity : EntitiesInUse)
-				{
-					EntitiesNames += TEXT("\t") + Entity->GetDisplayName() + TEXT("\n");
-				}
-
-				ConfirmDelete = FText::Format(LOCTEXT("ConfirmDeleteSomeEntitiesInUse", "The Entities below are in use!\n{0}\nDo you really want to delete them?"),
-					FText::FromString(EntitiesNames));
-			}
-			// Confirmation text for several entities. Displaying each of their names would make a huge dialog
-			else
-			{
-				ConfirmDelete = FText::Format(LOCTEXT("ConfirmDeleteManyEntitiesInUse", "{0} of the selected entities are in use!\nDo you really want to delete them?"),
-					FText::AsNumber(EntitiesInUse.Num()));
-			}
-
-			// Warn the user that this may result in data loss
-			FSuppressableWarningDialog::FSetupInfo Info(ConfirmDelete, LOCTEXT("DeleteEntities", "Delete Entities"), "DeleteEntitiesInUse_Warning");
-			Info.ConfirmText = LOCTEXT("DeleteEntities_Yes", "Yes");
-			Info.CancelText = LOCTEXT("DeleteEntities_No", "No");
-
-			FSuppressableWarningDialog DeleteEntitiesInUse(Info);
-			if (DeleteEntitiesInUse.ShowModal() == FSuppressableWarningDialog::Cancel)
-			{
-				return;
-			}
-		}
-
-		const FScopedTransaction Transaction(InEntities.Num() > 1 ? LOCTEXT("RemoveEntities", "Remove Entities") : LOCTEXT("RemoveEntity", "Remove Entity"));
-
-		for (UDMXEntity* EntityToDelete : InEntities)
-		{
-			// Fix references to this Entity
-			if (UDMXEntityFixtureType* AsFixtureType = Cast<UDMXEntityFixtureType>(EntityToDelete))
-			{
-				// Find Fixture Patches using this Fixture Type and null their templates
-				InLibrary->ForEachEntityOfType<UDMXEntityFixturePatch>([&AsFixtureType](UDMXEntityFixturePatch* Patch)
-					{
-						Patch->SetFixtureType(nullptr);
-					});
-			}
-
-			InLibrary->Modify();
-			EntityToDelete->Modify(); // Take a snapshot of the entity before setting its ParentLibrary to null
-			EntityToDelete->Destroy();
-		}
-	}
-}
-
 void FDMXEditorUtils::CopyEntities(const TArray<UDMXEntity*>&& EntitiesToCopy)
 {
 	// Clear the mark state for saving.
@@ -431,12 +200,6 @@ bool FDMXEditorUtils::CanPasteEntities(UDMXLibrary* ParentLibrary)
 	return FDMXEntityObjectTextFactory::CanCreate(ClipboardContent);
 }
 
-void FDMXEditorUtils::GetEntitiesFromClipboard(TArray<UDMXEntity*>& OutNewObjects)
-{
-	// DEPRECATED 5.0, no longer creates a meaningful result
-	OutNewObjects = TArray<UDMXEntity*>();
-}
-
 TArray<UDMXEntity*> FDMXEditorUtils::CreateEntitiesFromClipboard(UDMXLibrary* ParentLibrary)
 {
 	// Get the text from the clipboard
@@ -554,414 +317,6 @@ FText FDMXEditorUtils::GetEntityTypeNameText(TSubclassOf<UDMXEntity> EntityClass
 	}
 }
 
-bool FDMXEditorUtils::TryAutoAssignToUniverses(UDMXEntityFixturePatch* Patch, const TSet<int32>& AllowedUniverses)
-{
-	PRAGMA_DISABLE_DEPRECATION_WARNINGS
-	check(Patch->IsAutoAssignAddress());
-
-	Patch->Modify();
-	const int32 UniverseToRestore = Patch->GetUniverseID();
-	const int32 AutoAddressToRestore = Patch->GetAutoStartingAddress();
-
-	for (auto UniverseIt = AllowedUniverses.CreateConstIterator(); UniverseIt; ++UniverseIt)
-	{
-		// Don't auto assign to a universe smaller than the initial one
-		if (Patch->GetUniverseID() > *UniverseIt)
-		{
-			continue;
-		}
-
-		Patch->SetUniverseID(*UniverseIt);
-		const FUnassignedPatchesArray UnassignedPatches = AutoAssignedAddresses({ Patch }, 1, false);
-
-		const bool bWasPatchAssignedToUniverse = UnassignedPatches.Num() == 0;
-		if (bWasPatchAssignedToUniverse)
-		{
-			return true;
-		}
-	}
-
-	Patch->SetUniverseID(UniverseToRestore);
-	Patch->SetAutoStartingAddress(AutoAddressToRestore);
-
-	PRAGMA_ENABLE_DEPRECATION_WARNINGS
-	return false;
-}
-
-void FDMXEditorUtils::AutoAssignedAddresses(UDMXEntityFixtureType* ChangedParentFixtureType)
-{
-	PRAGMA_DISABLE_DEPRECATION_WARNINGS
-	if (ChangedParentFixtureType)
-	{
-		UDMXLibrary* Library = ChangedParentFixtureType->GetParentLibrary();
-		check(Library);
-
-		TArray<UDMXEntityFixturePatch*> FixturePatches;
-		Library->ForEachEntityOfType<UDMXEntityFixturePatch>([&FixturePatches, ChangedParentFixtureType](UDMXEntityFixturePatch* Patch)
-			{
-				if (Patch->GetFixtureType() == ChangedParentFixtureType)
-				{
-					FixturePatches.Add(Patch);
-				}
-			});
-
-		AutoAssignedAddresses(FixturePatches);
-	}
-	PRAGMA_ENABLE_DEPRECATION_WARNINGS
-}
-
-FDMXEditorUtils::FUnassignedPatchesArray FDMXEditorUtils::AutoAssignedAddresses(
-	const TArray<UDMXEntityFixturePatch*>& ChangedFixturePatches,
-	int32 MinimumAddress,
-	bool bCanChangePatchUniverses)
-{
-	PRAGMA_DISABLE_DEPRECATION_WARNINGS
-	struct Local
-	{
-		static void AssignAddressesInUniverse(
-			TArray<UDMXEntityFixturePatch*>& SortedPatchesWithSetAddress,
-			TArray<UDMXEntityFixturePatch*>& PatchesToAssign,
-			int32 MinimumAddress,
-			const int32 UniverseToAssignTo,
-			TArray<UDMXEntityFixturePatch*>& UnassignedPatches)
-		{
-			if (PatchesToAssign.Num() == 0)
-			{
-				return;
-			}
-
-			int32 IndexOfFirstPatchWithMinAddress =
-				SortedPatchesWithSetAddress.IndexOfByPredicate([MinimumAddress](UDMXEntityFixturePatch* Other)
-					{
-						const int32 ChannelEnd = Other->GetStartingChannel() + Other->GetChannelSpan();
-						return ChannelEnd >= MinimumAddress;
-					});
-
-			for (UDMXEntityFixturePatch* ToAssign : PatchesToAssign)
-			{
-				const bool bIsUniverseEmpty = IndexOfFirstPatchWithMinAddress == INDEX_NONE;
-				if (bIsUniverseEmpty)
-				{
-					if (HasEnoughSpaceToUniverseEnd(ToAssign, MinimumAddress))
-					{
-						IndexOfFirstPatchWithMinAddress = 0;
-						AssignPatchTo(ToAssign, MinimumAddress, UniverseToAssignTo);
-						SortedPatchesWithSetAddress.Add(ToAssign);
-					}
-					else
-					{
-						UnassignedPatches.Add(ToAssign);
-					}
-					continue;
-				}
-
-				const bool bFoundGap = FillIntoFirstGap(ToAssign, UniverseToAssignTo, MinimumAddress, IndexOfFirstPatchWithMinAddress, SortedPatchesWithSetAddress);
-				if (!bFoundGap)
-				{
-					UDMXEntityFixturePatch* LastPatchInUniverse = SortedPatchesWithSetAddress[SortedPatchesWithSetAddress.Num() - 1];
-					if (HasEnoughSpaceToUniverseEnd(ToAssign, GetEndAddressOf(LastPatchInUniverse) + 1))
-					{
-						AssignPatchTo(ToAssign, GetEndAddressOf(LastPatchInUniverse) + 1, UniverseToAssignTo);
-						SortedPatchesWithSetAddress.Add(ToAssign);
-					}
-					else
-					{
-						UnassignedPatches.Add(ToAssign);
-					}
-				}
-			}
-		}
-
-		static bool HasEnoughSpaceToUniverseEnd(UDMXEntityFixturePatch* Patch, int32 AtAddress)
-		{
-			// + 1 is needed otherwise we're off by one, e.g. when AtAddress=DMX_UNIVERSE_SIZE and ChannelSpan=1
-			return (DMX_UNIVERSE_SIZE + 1) - AtAddress >= Patch->GetChannelSpan();
-		}
-
-		static void AssignPatchTo(UDMXEntityFixturePatch* Patch, int32 ToAddress, int32 UniverseToAssignTo)
-		{
-			Patch->Modify();
-			Patch->SetAutoStartingAddress(ToAddress);
-			Patch->SetUniverseID(UniverseToAssignTo);
-		}
-
-		static bool FillIntoFirstGap(
-			UDMXEntityFixturePatch* ToAssign,
-			int32 UniverseToAssignTo,
-			int32& MinimumAddress,
-			int32& IndexOfFirstPatchWithMinAddress,
-			TArray<UDMXEntityFixturePatch*>& SortedPatchesWithSetAddress)
-		{
-			const int32 NeededSpan = ToAssign->GetChannelSpan();
-			if (NeededSpan < 1)
-			{
-				return false;
-			}
-
-			UDMXEntityFixturePatch* FirstPatchWithMinAddress = SortedPatchesWithSetAddress[IndexOfFirstPatchWithMinAddress];
-			const bool bDoesPatchFitBeforeFirstPatchWithMinAddress = FirstPatchWithMinAddress->GetStartingChannel() >= MinimumAddress
-				&& FirstPatchWithMinAddress->GetStartingChannel() - MinimumAddress >= NeededSpan;
-			if (bDoesPatchFitBeforeFirstPatchWithMinAddress)
-			{
-				AssignPatchTo(ToAssign, MinimumAddress, UniverseToAssignTo);
-				SortedPatchesWithSetAddress.Insert(ToAssign, 0);
-
-				MinimumAddress = ToAssign->GetEndingChannel() + 1;
-				IndexOfFirstPatchWithMinAddress = 0;
-
-				return true;
-			}
-
-			const bool bNoGapsToLookAt = IndexOfFirstPatchWithMinAddress == SortedPatchesWithSetAddress.Num() - 1;
-			if (bNoGapsToLookAt)
-			{
-				return false;
-			}
-
-			int32 PrevAvailableAddress = GetEndAddressOf(SortedPatchesWithSetAddress[IndexOfFirstPatchWithMinAddress]) + 1;
-			for (int32 iPatchInUniverse = IndexOfFirstPatchWithMinAddress + 1; iPatchInUniverse < SortedPatchesWithSetAddress.Num(); ++iPatchInUniverse)
-			{
-				const int32 NextStart = SortedPatchesWithSetAddress[iPatchInUniverse]->GetStartingChannel();
-				const int32 UnoccupiedSpan = NextStart - PrevAvailableAddress;
-
-				if (UnoccupiedSpan >= NeededSpan)
-				{
-					AssignPatchTo(ToAssign, PrevAvailableAddress, UniverseToAssignTo);
-					SortedPatchesWithSetAddress.Insert(ToAssign, iPatchInUniverse);
-					return true;
-				}
-
-				PrevAvailableAddress = GetEndAddressOf(SortedPatchesWithSetAddress[iPatchInUniverse]) + 1;
-			}
-			return false;
-		}
-		static int32 GetEndAddressOf(UDMXEntityFixturePatch* Patch)
-		{
-			return Patch->GetEndingChannel();
-		}
-	};
-
-	if (ChangedFixturePatches.Num() == 0)
-	{
-		return {};
-	}
-
-	// Auto assign Patches from multiple Libraries is not supported
-	UDMXLibrary* Library = ChangedFixturePatches[0]->GetParentLibrary();
-	check(Library);
-	for (UDMXEntityFixturePatch* Patch : ChangedFixturePatches)
-	{
-		check(Patch->GetParentLibrary() == Library);
-	}
-
-	// Only care about those that have auto assign addresses set
-	TArray<UDMXEntityFixturePatch*> PatchesToAutoAssign = ChangedFixturePatches;
-	PatchesToAutoAssign.RemoveAll([](UDMXEntityFixturePatch* Patch) {
-		return !Patch->IsAutoAssignAddress();
-		});
-
-	TArray<UDMXEntityFixturePatch*> AllFixturePatches = Library->GetEntitiesTypeCast<UDMXEntityFixturePatch>();
-	TMap<int32, TArray<UDMXEntityFixturePatch*>> UniverseToAllPatches = FDMXRuntimeUtils::MapToUniverses(AllFixturePatches);
-	TArray<UDMXEntityFixturePatch*> PatchesToAssignInNextUniverse;
-	for (auto UniverseIterator = UniverseToAllPatches.CreateIterator(); UniverseIterator; ++UniverseIterator)
-	{
-		const int32 CurrentUniverse = UniverseIterator->Key;
-		TArray<UDMXEntityFixturePatch*>& SortedPatchesWithSetAddress = [&UniverseIterator, &PatchesToAutoAssign, CurrentUniverse]() -> TArray<UDMXEntityFixturePatch*>&
-		{
-			TArray<UDMXEntityFixturePatch*>& Result = UniverseIterator->Value;
-			Result.RemoveAll([&, PatchesToAutoAssign](UDMXEntityFixturePatch* Patch) {
-				return PatchesToAutoAssign.Contains(Patch);
-				});
-			Result.RemoveAll([&, PatchesToAutoAssign, CurrentUniverse](UDMXEntityFixturePatch* Patch) {
-				return Patch->GetUniverseID() != CurrentUniverse;
-				});
-			Result.Sort([](const UDMXEntityFixturePatch& Patch, const UDMXEntityFixturePatch& Other) {
-				return
-					Patch.GetUniverseID() < Other.GetUniverseID() ||
-					(Patch.GetUniverseID() == Other.GetUniverseID() && Patch.GetStartingChannel() <= Other.GetStartingChannel());
-				}
-			);
-			return Result;
-		}();
-		TArray<UDMXEntityFixturePatch*> PatchesToAssignInThisUniverse = [&UniverseIterator, &PatchesToAutoAssign, CurrentUniverse]()
-		{
-			TArray<UDMXEntityFixturePatch*> Result = PatchesToAutoAssign;
-			Result.RemoveAll([&, UniverseIterator](UDMXEntityFixturePatch* Patch)
-				{
-					return Patch->GetUniverseID() != CurrentUniverse;
-				});
-			return Result;
-		}();
-
-		if (bCanChangePatchUniverses)
-		{
-			// The patches in PatchesToAssignInNextUniverse are left over from last universe: hence they should ignore the MinimumAddress. 
-			TArray<UDMXEntityFixturePatch*> ToAssignNextUniverse;
-			Local::AssignAddressesInUniverse(SortedPatchesWithSetAddress, PatchesToAssignInNextUniverse, 1, CurrentUniverse, ToAssignNextUniverse);
-
-			PatchesToAssignInNextUniverse = ToAssignNextUniverse;
-		}
-		Local::AssignAddressesInUniverse(SortedPatchesWithSetAddress, PatchesToAssignInThisUniverse, MinimumAddress, CurrentUniverse, PatchesToAssignInNextUniverse);
-	}
-
-	const bool bNeedNewUniverse = PatchesToAssignInNextUniverse.Num() > 0;
-	if (bNeedNewUniverse)
-	{
-		if (bCanChangePatchUniverses)
-		{
-			const int32 HighestUniverse = [&UniverseToAllPatches]()
-			{
-				int32 HighestSoFar = -1;
-				for (auto UniverseIterator = UniverseToAllPatches.CreateIterator(); UniverseIterator; ++UniverseIterator)
-				{
-					HighestSoFar = FMath::Max(HighestSoFar, UniverseIterator->Key);
-				}
-				return HighestSoFar;
-			}();
-			for (UDMXEntityFixturePatch* UnassignedPatch : PatchesToAssignInNextUniverse)
-			{
-				UnassignedPatch->Modify();
-				UnassignedPatch->SetUniverseID(HighestUniverse + 1);
-			}
-			return AutoAssignedAddresses(PatchesToAssignInNextUniverse, 1, bCanChangePatchUniverses);
-		}
-
-		return PatchesToAssignInNextUniverse;
-	}
-	PRAGMA_ENABLE_DEPRECATION_WARNINGS
-
-	return {};
-}
-
-void FDMXEditorUtils::AutoAssignedChannels(bool bAllowDecrementUniverse, bool bAllowDecrementChannels, TArray<UDMXEntityFixturePatch*> FixturePatches)
-{
-	if (FixturePatches.IsEmpty())
-	{
-		return;
-	}
-
-	// Ensure patches are of a single DMX Library
-	const UDMXLibrary* DMXLibrary = FixturePatches[0]->GetParentLibrary();
-	for (UDMXEntityFixturePatch* FixturePatch : FixturePatches)
-	{
-		if (!ensureAlwaysMsgf(FixturePatch->GetParentLibrary() == DMXLibrary, TEXT("Trying to auto assign fixture patches from different DMX Libraries at once. This is not supported.")))
-		{
-			return;
-		}
-	}
-
-	TArray<UDMXEntityFixturePatch*> AllFixturePatchesInLibrary = DMXLibrary->GetEntitiesTypeCast<UDMXEntityFixturePatch>();
-	
-	// Sort by Universe and channel
-	FixturePatches.Sort([](const UDMXEntityFixturePatch& FixturePatchA, const UDMXEntityFixturePatch& FixturePatchB)
-		{
-			const bool bUniverseIsSmaller = FixturePatchA.GetUniverseID() < FixturePatchB.GetUniverseID();
-			const bool bUniverseIsEqual = FixturePatchA.GetUniverseID() == FixturePatchB.GetUniverseID();
-			const bool bChannelIsSmallerOrEqual = FixturePatchA.GetStartingChannel() <= FixturePatchB.GetStartingChannel();
-			return bUniverseIsSmaller || (bUniverseIsEqual && bChannelIsSmallerOrEqual);
-		});
-	// Move all to the universe and address of the first patch
-	for (UDMXEntityFixturePatch* FixturePatch : FixturePatches)
-	{
-		FixturePatch->Modify();
-		FixturePatch->SetUniverseID(FixturePatches[0]->GetUniverseID());
-		FixturePatch->SetStartingChannel(FixturePatches[0]->GetStartingChannel());
-	}
-
-	// Using absolute channels (Universe * DMX_UIVERSE_SIZE + Channel)
-	auto FindAbsoluteChannelLambda([bAllowDecrementUniverse, bAllowDecrementChannels, &FixturePatches, &AllFixturePatchesInLibrary](UDMXEntityFixturePatch* FixturePatch, const TArray<UDMXEntityFixturePatch*>& OtherPendingAutoAssignFixturePatches) -> int32
-		{
-			AllFixturePatchesInLibrary.StableSort([](const UDMXEntityFixturePatch& FixturePatchA, const UDMXEntityFixturePatch& FixturePatchB)
-				{
-					const bool bUniverseIsSmaller = FixturePatchA.GetUniverseID() < FixturePatchB.GetUniverseID();
-					const bool bUniverseIsEqual = FixturePatchA.GetUniverseID() == FixturePatchB.GetUniverseID();
-					const bool bChannelIsSmallerOrEqual = FixturePatchA.GetStartingChannel() <= FixturePatchB.GetStartingChannel();
-					return bUniverseIsSmaller || (bUniverseIsEqual && bChannelIsSmallerOrEqual);
-				});
-
-			const int32 OldUniverse = FixturePatch->GetUniverseID();
-			const int32 OldAbsoluteStartingChannel = FixturePatch->GetUniverseID() * DMX_UNIVERSE_SIZE + FixturePatch->GetStartingChannel();
-			const int32 RequiredChannelSpan = FixturePatch->GetChannelSpan();
-			const int32 MinAbsoluteChannel = bAllowDecrementUniverse ? 1 : FixturePatch->GetUniverseID() * DMX_UNIVERSE_SIZE + 1;
-
-			int32 AbsoluteStartOfGap = MinAbsoluteChannel;
-			for (int32 IndexOfFixturePatch = 0; IndexOfFixturePatch < AllFixturePatchesInLibrary.Num(); IndexOfFixturePatch++)
-			{				
-				UDMXEntityFixturePatch* Other = AllFixturePatchesInLibrary[IndexOfFixturePatch];
-				const int32 OtherAbsoluteEndingChannel = Other->GetUniverseID() * DMX_UNIVERSE_SIZE + Other->GetEndingChannel();
-				const int32 GapChannelSpan = Other->GetUniverseID() * DMX_UNIVERSE_SIZE + Other->GetStartingChannel() - AbsoluteStartOfGap;
-
-				// Ignore others that are pending auto-assign
-				if (OtherPendingAutoAssignFixturePatches.Contains(Other))
-				{
-					continue;
-				}
-
-				// Peek ahead when looking at self
-				if (Other == FixturePatch)
-				{
-					UDMXEntityFixturePatch* Next = AllFixturePatchesInLibrary.IsValidIndex(IndexOfFixturePatch + 1) ? AllFixturePatchesInLibrary[IndexOfFixturePatch + 1] : nullptr;
-					if (!Next)
-					{
-						break;
-					}
-
-					const int32 NextAbsoluteStartingChannel = Next->GetUniverseID() * DMX_UNIVERSE_SIZE + Next->GetStartingChannel();
-					const int32 GapToNext = NextAbsoluteStartingChannel - AbsoluteStartOfGap;
-					if (GapToNext >= RequiredChannelSpan)
-					{
-						break;
-					}
-				}
-
-				if (!bAllowDecrementUniverse && OtherAbsoluteEndingChannel + 1 < MinAbsoluteChannel)
-				{
-					AbsoluteStartOfGap = MinAbsoluteChannel;
-					continue;
-				}
-
-				if (!bAllowDecrementChannels && AbsoluteStartOfGap < OldAbsoluteStartingChannel)
-				{
-					AbsoluteStartOfGap = OtherAbsoluteEndingChannel + 1;
-					continue;
-				}
-
-				if (GapChannelSpan >= RequiredChannelSpan)
-				{
-					break;
-				}
-
-				AbsoluteStartOfGap = OtherAbsoluteEndingChannel + 1;
-			}
-
-			int32 Universe = AbsoluteStartOfGap / DMX_UNIVERSE_SIZE;
-			int32 Channel = AbsoluteStartOfGap % DMX_UNIVERSE_SIZE;
-			if (Channel + FixturePatch->GetChannelSpan() - 1 > DMX_UNIVERSE_SIZE)
-			{
-				return (Universe + 1) * DMX_UNIVERSE_SIZE + 1;
-			}
-			else
-			{
-				return AbsoluteStartOfGap;
-			}
-		});
-
-	const FScopedTransaction AutoAssignTransaction(LOCTEXT("AutoAssignTransaction", "Auto Assign Fixture Patch"));
-	TArray<UDMXEntityFixturePatch*> OtherPendingAutoAssignFixturePatches(FixturePatches);
-	for (UDMXEntityFixturePatch* FixturePatch : FixturePatches)
-	{
-		OtherPendingAutoAssignFixturePatches.Remove(FixturePatch);
-		const int32 AutoAssignAbsoluteChannel = FindAbsoluteChannelLambda(FixturePatch, OtherPendingAutoAssignFixturePatches);
-		if (AutoAssignAbsoluteChannel > 0)
-		{
-			FixturePatch->PreEditChange(UDMXEntityFixturePatch::StaticClass()->FindPropertyByName(UDMXEntityFixturePatch::GetStartingChannelPropertyNameChecked()));
-			FixturePatch->SetUniverseID(AutoAssignAbsoluteChannel / DMX_UNIVERSE_SIZE);
-			FixturePatch->SetStartingChannel(AutoAssignAbsoluteChannel % DMX_UNIVERSE_SIZE);
-			FixturePatch->PostEditChange();
-		}
-	}
-}
-
 void FDMXEditorUtils::UpdatePatchColors(UDMXLibrary* Library)
 {
 	check(Library);
diff --git a/Engine/Plugins/VirtualProduction/DMX/DMXEngine/Source/DMXEditor/Public/DMXEditorUtils.h b/Engine/Plugins/VirtualProduction/DMX/DMXEngine/Source/DMXEditor/Public/DMXEditorUtils.h
--- a/Engine/Plugins/VirtualProduction/DMX/DMXEngine/Source/DMXEditor/Public/DMXEditorUtils.h
+++ b/Engine/Plugins/VirtualProduction/DMX/DMXEngine/Source/DMXEditor/Public/DMXEditorUtils.h
@@ -20,20 +20,6 @@ class UDMXEntityFixturePatch;
 class DMXEDITOR_API FDMXEditorUtils
 {
 public:
-	typedef TArray<UDMXEntityFixturePatch*> FUnassignedPatchesArray;
-
-	UE_DEPRECATED(5.0, "Moved to DMXRuntimeUtils to allow using it at runtime.")
-	static FString GenerateUniqueNameFromExisting(const TSet<FString>& InExistingNames, const FString& InBaseName);
-
-	UE_DEPRECATED(5.0, "Deprecated in favor of UDMXRuntimeUtils::FindUniqueEntityName.")
-	static FString FindUniqueEntityName(const UDMXLibrary* InLibrary, TSubclassOf<UDMXEntity> InEntityClass, const FString& InBaseName = TEXT(""));
-
-	UE_DEPRECATED(5.0, "Deprecated to reduce redundant code. Instead use UDMXRuntimeUtils::FindUniqueEntityName.")
-	static void SetNewFixtureFunctionsNames(UDMXEntityFixtureType* InFixtureType);
-
-	UE_DEPRECATED(5.0, "Deprecated in favor of new UDMXEntityFixtureType::CreateFixtureType and UDMXEntityFixturePatch::CreateFixturePatch that support creating patches in blueprints.")
-	static bool AddEntity(UDMXLibrary* InLibrary, const FString& NewEntityName, TSubclassOf<UDMXEntity> NewEntityClass, UDMXEntity** OutNewEntity = nullptr);
-
 	/**
 	 * Validates an Entity name, also checking for uniqueness among others of the same type.
 	 * @param NewEntityName		The name to validate.
@@ -50,20 +36,12 @@ public:
 	/**  Checks if the Entity is being referenced by other objects. */
 	static bool IsEntityUsed(const UDMXLibrary* InLibrary, const UDMXEntity* InEntity);
 
-	/** DEPRECATED 5.0 */
-	UE_DEPRECATED(5.0, "Replaced with UDMXEntityFixtureType::RemoveFixtureTypeFromLibrary and UDMXEntityFixturePatch::RemoveFixturePatchFromLibrary to support blueprint and runtime changes.")
-	static void RemoveEntities(UDMXLibrary* InLibrary, const TArray<UDMXEntity*>& InEntities);
-
 	/**  Copies Entities to the operating system's clipboard. */
 	static void CopyEntities(const TArray<UDMXEntity*>&& EntitiesToCopy);
 
 	/**  Determines whether the current contents of the clipboard contain paste-able DMX Entity information */
 	static bool CanPasteEntities(UDMXLibrary* ParentLibrary);
 
-	/** DEPRECATED 5.0 */
-	UE_DEPRECATED(5.0, "Replaced with FDMXEditorUtils::CreateEntitiesFromClipboard. NOTE: GetEntitiesFromClipboard no longer returns any entities, as entities can only be created within a DMXLibrary.")
-	static void GetEntitiesFromClipboard(TArray<UDMXEntity*>& OutNewObjects);
-
 	/**
 	 * Creates the copied DMX Entities from the clipboard without attempting to paste/apply them in any way
 	 * 
@@ -82,34 +60,6 @@ public:
 	/**  Returns the Entity class type name (e.g: Fixture Type for UDMXEntityFixtureType) in singular or plural */
 	static FText GetEntityTypeNameText(TSubclassOf<UDMXEntity> EntityClass, bool bPlural = false);
 
-	/** DEPRECATED 5.1 */
-	UE_DEPRECATED(5.1, "Simplyifing AutoAssign, just use FDMXEditorUtils::AutoAssignedAddresses instead.")
-	static bool TryAutoAssignToUniverses(UDMXEntityFixturePatch* Patch, const TSet<int32>& AllowedUniverses);
-
-	/** DEPRECATED 5.1 */
-	UE_DEPRECATED(5.1, "The bAutoAssign property got deprecated in UDMXEntityFixturePatch, so there's no meaning to ever auto assign on Fixture Type changes. So this function is removed without replacement.")
-	static void AutoAssignedAddresses(UDMXEntityFixtureType* ChangedParentFixtureType);
-
-	/** DEPRECATED 5.1 */
-	UE_DEPRECATED(5.1, "Simplyifing AutoAssign, just use FDMXEditorUtils::AutoAssignedChannels instead.")
-	static FUnassignedPatchesArray AutoAssignedAddresses(const TArray<UDMXEntityFixturePatch*>& ChangedFixturePatches, int32 MinimumAddress = 1, bool bCanChangePatchUniverses = true);
-
-	/**
-	 * Updates starting addresses for fixture patches that have bAutoAssignAddress set, ignores others.
-	 * Note, patches all have to reside in the same library.
-	 *
-	 * The caller is responsible to call Modify on the patches and register undo/redo.
-	 *
-	 * If bCanChangePatchUniverses = true, this function will assign patches that do not fit into the existing universes to
-	 * the next universe. MinimumAddress is ignored for the new universe, i.e. we start placing remaining patches at position 1.
-	 *
-	 * @param bAllowDecrementUniverse	If true, the patches can be assigend to a previous universe
-	 * @param bCanChangePatchUniverses	If true, the patches can lower their absoulte channel
-	 * @param FixturePatches			The patches that want their channels to be auto assigned
-	 */
-	UE_DEPRECATED(5.2, "AutoAssignedChannels is deprecated without replacement. Please create your own auto assign implementation. See FDMXFixturePatchAutoAssignUtility for an example.")
-	static void AutoAssignedChannels(bool bAllowDecrementUniverse, bool bAllowDecrementChannels, TArray<UDMXEntityFixturePatch*> FixturePatches);
-	
 	/**
 	 * Creates a unique color for all patches that use the default color FLinearColor(1.0f, 0.0f, 1.0f)
 	 *
diff --git a/Engine/Plugins/VirtualProduction/DMX/DMXEngine/Source/DMXRuntime/Private/DMXSubsystem.cpp b/Engine/Plugins/VirtualProduction/DMX/DMXEngine/Source/DMXRuntime/Private/DMXSubsystem.cpp
--- a/Engine/Plugins/VirtualProduction/DMX/DMXEngine/Source/DMXRuntime/Private/DMXSubsystem.cpp
+++ b/Engine/Plugins/VirtualProduction/DMX/DMXEngine/Source/DMXRuntime/Private/DMXSubsystem.cpp
@@ -439,50 +439,18 @@ UDMXEntityFixturePatch* UDMXSubsystem::GetFixturePatch(FDMXEntityFixturePatchRef
 
 bool UDMXSubsystem::GetFunctionsMap(UDMXEntityFixturePatch* InFixturePatch, TMap<FDMXAttributeName, int32>& OutAttributesMap)
 {
-	OutAttributesMap.Empty();
-
-	if (InFixturePatch == nullptr)
-	{
-		UE_LOG(DMXSubsystemLog, Warning, TEXT("%hs: FixturePatch is null"), __FUNCTION__);
-
-		return false;
-	}
-
-	const FDMXFixtureMode* ModePtr = InFixturePatch->GetActiveMode();
-	if (!ModePtr)
+	if (InFixturePatch)
 	{
-		UE_LOG(DMXSubsystemLog, Warning, TEXT("Cannot get function map, fixture Patch %s has no valid active mode"), *InFixturePatch->Name);
-		return false;
-	}
-
-	const FDMXSignalSharedPtr& Signal = InFixturePatch->GetLastReceivedDMXSignal();
-
-	if(Signal.IsValid())
-	{ 
-		const TArray<uint8>& ChannelData = Signal->ChannelData;
-		
-		const int32 PatchStartingIndex = InFixturePatch->GetStartingChannel() - 1;
-
-		for (const FDMXFixtureFunction& Function : ModePtr->Functions)
-		{
-			const int32 FunctionStartIndex = Function.Channel - 1 + PatchStartingIndex;
-			const int32 FunctionLastIndex = FunctionStartIndex + FDMXConversions::GetSizeOfSignalFormat(Function.DataType) - 1;
-			if (FunctionLastIndex >= ChannelData.Num())
-			{
-				break;
-			}
-
-			const uint32 ChannelValue = UDMXEntityFixtureType::BytesToFunctionValue(Function, ChannelData.GetData() + FunctionStartIndex);
-			OutAttributesMap.Add(Function.Attribute, ChannelValue);
-		}
+		InFixturePatch->GetAttributeValues(OutAttributesMap);
+		return true;
 	}
-
-	return true;
+	
+	return false;
 }
 
 bool UDMXSubsystem::GetFunctionsMapForPatch(UDMXEntityFixturePatch* InFixturePatch, TMap<FDMXAttributeName, int32>& OutAttributesMap)
 {
-	// TODO: This is a duplicate..
+	// DEPRECATED 5.5, duplicate of GetFunctionsMap
 	return GetFunctionsMap(InFixturePatch, OutAttributesMap);
 }
 
@@ -533,6 +501,7 @@ FName UDMXSubsystem::GetAttributeLabel(FDMXAttributeName AttributeName)
 
 /*static*/ UDMXSubsystem* UDMXSubsystem::GetDMXSubsystem_Pure()
 {
+	check(GEngine);
 	return GEngine->GetEngineSubsystem<UDMXSubsystem>();
 }
 
diff --git a/Engine/Plugins/VirtualProduction/DMX/DMXEngine/Source/DMXRuntime/Private/Game/DMXComponent.cpp b/Engine/Plugins/VirtualProduction/DMX/DMXEngine/Source/DMXRuntime/Private/Game/DMXComponent.cpp
--- a/Engine/Plugins/VirtualProduction/DMX/DMXEngine/Source/DMXRuntime/Private/Game/DMXComponent.cpp
+++ b/Engine/Plugins/VirtualProduction/DMX/DMXEngine/Source/DMXRuntime/Private/Game/DMXComponent.cpp
@@ -100,7 +100,7 @@ void UDMXComponent::SetupReceiveDMXBinding()
 			FixturePatch->OnFixturePatchReceivedDMX.AddDynamic(this, &UDMXComponent::OnFixturePatchReceivedDMX);
 
 			FDMXNormalizedAttributeValueMap NormalizeAttributeValues;
-			FixturePatch->GetNormalizedAttributesValues(NormalizeAttributeValues);
+			FixturePatch->GetNormalizedAttributeValues(NormalizeAttributeValues);
 
 			if (NormalizeAttributeValues.Map.Num() > 0)
 			{
@@ -166,7 +166,7 @@ void UDMXComponent::BeginPlay()
 	if (UDMXEntityFixturePatch* FixturePatch = GetFixturePatch())
 	{
 		FDMXNormalizedAttributeValueMap NormalizeAttributeValues;
-		FixturePatch->GetNormalizedAttributesValues(NormalizeAttributeValues);
+		FixturePatch->GetNormalizedAttributeValues(NormalizeAttributeValues);
 
 		if (NormalizeAttributeValues.Map.Num() > 0)
 		{
diff --git a/Engine/Plugins/VirtualProduction/DMX/DMXEngine/Source/DMXRuntime/Private/Library/DMXEntityFixturePatch.cpp b/Engine/Plugins/VirtualProduction/DMX/DMXEngine/Source/DMXRuntime/Private/Library/DMXEntityFixturePatch.cpp
--- a/Engine/Plugins/VirtualProduction/DMX/DMXEngine/Source/DMXRuntime/Private/Library/DMXEntityFixturePatch.cpp
+++ b/Engine/Plugins/VirtualProduction/DMX/DMXEngine/Source/DMXRuntime/Private/Library/DMXEntityFixturePatch.cpp
@@ -66,7 +66,7 @@ UDMXEntityFixturePatch* UDMXEntityFixturePatch::CreateFixturePatchInLibrary(FDMX
 			NewFixturePatch->SetUniverseID(ConstructionParams.UniverseID);
 			NewFixturePatch->SetStartingChannel(ConstructionParams.StartingAddress);
 			NewFixturePatch->SetActiveModeIndex(ConstructionParams.ActiveMode);
-
+			
 			if (ConstructionParams.MVRFixtureUUID.IsValid())
 			{
 				// Make sure the MVR UUID is truly unique across the DNX Library
@@ -503,14 +503,6 @@ void UDMXEntityFixturePatch::ValidateActiveMode()
 	}
 }
 
-bool UDMXEntityFixturePatch::CanReadActiveMode() const
-{
-	// DEPRECATED 4.27
-	return ParentFixtureTypeTemplate != nullptr
-		&& ParentFixtureTypeTemplate->IsValidLowLevelFast()
-		&& ParentFixtureTypeTemplate->Modes.IsValidIndex(ActiveMode);
-}
-
 const FDMXFixtureMode* UDMXEntityFixturePatch::GetActiveMode() const
 {
 	if (ParentFixtureTypeTemplate && 
@@ -546,27 +538,6 @@ void UDMXEntityFixturePatch::SetUniverseID(int32 NewUniverseID)
 	RebuildCache();
 }
 
-#if WITH_EDITOR
-void UDMXEntityFixturePatch::SetAutoStartingAddress(int32 NewAutoStartingAddress)
-{
-	// DEPRECATED 5.1
-	AutoStartingAddress_DEPRECATED = NewAutoStartingAddress;
-	ManualStartingAddress_DEPRECATED = NewAutoStartingAddress;
-
-	RebuildCache();
-}
-#endif // WITH_EDITOR
-
-#if WITH_EDITOR
-void UDMXEntityFixturePatch::SetManualStartingAddress(int32 NewManualStartingAddress)
-{
-	// DEPRECATED 5.1
-	ManualStartingAddress_DEPRECATED = NewManualStartingAddress;
-
-	RebuildCache();
-}
-#endif // WITH_EDITOR
-
 void UDMXEntityFixturePatch::SetStartingChannel(int32 NewStartingChannel)
 {
 	if (NewStartingChannel == StartingChannel)
@@ -927,7 +898,13 @@ float UDMXEntityFixturePatch::GetNormalizedAttributeValue(FDMXAttributeName Attr
 
 void UDMXEntityFixturePatch::GetAttributesValues(TMap<FDMXAttributeName, int32>& AttributesValues)
 {
-	AttributesValues.Reset();
+	// DEPRECATED 5.5
+	GetAttributeValues(AttributesValues);
+}
+
+void UDMXEntityFixturePatch::GetAttributeValues(TMap<FDMXAttributeName, int32>& AttributeValues)
+{
+	AttributeValues.Reset();
 
 	// Update the cache if it isn't updated on tick
 	if (!IsTickable())
@@ -937,11 +914,17 @@ void UDMXEntityFixturePatch::GetAttributesValues(TMap<FDMXAttributeName, int32>&
 
 	if (const TMap<FDMXAttributeName, int32>* AttributeValuesPtr = Cache.GetAllRawAttributeValues())
 	{
-		AttributesValues = *AttributeValuesPtr;
+		AttributeValues = *AttributeValuesPtr;
 	}
 }
 
 void UDMXEntityFixturePatch::GetNormalizedAttributesValues(FDMXNormalizedAttributeValueMap& NormalizedAttributesValues)
+{
+	// DEPRECATED 5.5
+	GetNormalizedAttributeValues(NormalizedAttributesValues);
+}
+
+void UDMXEntityFixturePatch::GetNormalizedAttributeValues(FDMXNormalizedAttributeValueMap& NormalizedAttributeValues)
 {
 	// Update the cache if it isn't updated on tick
 	if (!IsTickable())
@@ -951,7 +934,7 @@ void UDMXEntityFixturePatch::GetNormalizedAttributesValues(FDMXNormalizedAttribu
 
 	if (const FDMXNormalizedAttributeValueMap* NormalizedAttributeValuesPtr = Cache.GetAllNormalizedAttributeValues())
 	{
-		NormalizedAttributesValues = *NormalizedAttributeValuesPtr;
+		NormalizedAttributeValues = *NormalizedAttributeValuesPtr;
 	}
 }
 
diff --git a/Engine/Plugins/VirtualProduction/DMX/DMXEngine/Source/DMXRuntime/Private/Library/DMXEntityFixtureType.cpp b/Engine/Plugins/VirtualProduction/DMX/DMXEngine/Source/DMXRuntime/Private/Library/DMXEntityFixtureType.cpp
--- a/Engine/Plugins/VirtualProduction/DMX/DMXEngine/Source/DMXRuntime/Private/Library/DMXEntityFixtureType.cpp
+++ b/Engine/Plugins/VirtualProduction/DMX/DMXEngine/Source/DMXRuntime/Private/Library/DMXEntityFixtureType.cpp
@@ -124,72 +124,6 @@ int32 FDMXFixtureMatrix::GetLastChannel() const
 	return FirstCellChannel + NumCells * NumChannelsPerCell - 1;
 }
 
-int32 FDMXFixtureMatrix::GetFixtureMatrixLastChannel() const
-{
-	// DEPRECATED 5.0
-	return GetLastChannel();
-}
-
-bool FDMXFixtureMatrix::GetChannelsFromCell(FIntPoint CellCoordinate, FDMXAttributeName Attribute, TArray<int32>& Channels) const
-{
-	// DEPRECATED 5.0
-	Channels.Reset();
-
-	TArray<int32> AllChannels;
-
-	if (CellCoordinate.X < 0 || CellCoordinate.X >= XCells)
-	{
-		return false;
-	}
-
-	if (CellCoordinate.Y < 0 || CellCoordinate.Y >= YCells)
-	{
-		return false;
-	}
-
-	for (int32 YCell = 0; YCell < YCells; YCell++)
-	{
-		for (int32 XCell = 0; XCell < XCells; XCell++)
-		{
-			AllChannels.Add(XCell + YCell * XCells);
-		}
-	}
-
-	TArray<int32> OrderedChannels;
-	FDMXRuntimeUtils::PixelMappingDistributionSort(PixelMappingDistribution, XCells, YCells, AllChannels, OrderedChannels);
-
-	check(AllChannels.Num() == OrderedChannels.Num());
-
-	int32 CellSize = 0;
-	int32 CellAttributeIndex = -1;
-	int32 CellAttributeSize = 0;
-	for (const FDMXFixtureCellAttribute& CellAttribute : CellAttributes)
-	{
-		int32 CurrentFunctionSize = FDMXConversions::GetSizeOfSignalFormat(CellAttribute.DataType);
-		if (CellAttribute.Attribute == Attribute)
-		{
-			CellAttributeIndex = CellSize;
-			CellAttributeSize = CurrentFunctionSize;
-		}
-		CellSize += CurrentFunctionSize;
-	}
-
-	// no function found
-	if (CellAttributeIndex < 0 || CellAttributeSize == 0)
-	{
-		return false;
-	}
-
-	int32 ChannelBase = FirstCellChannel + (OrderedChannels[CellCoordinate.Y + CellCoordinate.X * YCells] * CellSize) + CellAttributeIndex;
-
-	for (int32 ChannelIndex = 0; ChannelIndex < CellAttributeSize; ChannelIndex++)
-	{
-		Channels.Add(ChannelBase + ChannelIndex);
-	}
-
-	return true;
-}
-
 #if WITH_EDITOR
 int32 FDMXFixtureMode::AddOrInsertFunction(int32 IndexOfFunction, FDMXFixtureFunction InFunction)
 {
diff --git a/Engine/Plugins/VirtualProduction/DMX/DMXEngine/Source/DMXRuntime/Public/DMXSubsystem.h b/Engine/Plugins/VirtualProduction/DMX/DMXEngine/Source/DMXRuntime/Public/DMXSubsystem.h
--- a/Engine/Plugins/VirtualProduction/DMX/DMXEngine/Source/DMXRuntime/Public/DMXSubsystem.h
+++ b/Engine/Plugins/VirtualProduction/DMX/DMXEngine/Source/DMXRuntime/Public/DMXSubsystem.h
@@ -11,12 +11,14 @@
 
 #include "DMXSubsystem.generated.h"
 
-class IDMXProtocol;
 class UDMXEntityFixtureType;
 class UDMXEntityFixturePatch;
 class UDMXLibrary;
-class UDMXModulator;
 
+#if UE_ENABLE_INCLUDE_ORDER_DEPRECATED_IN_5_5
+class IDMXProtocol;
+class UDMXModulator;
+#endif // UE_ENABLE_INCLUDE_ORDER_DEPRECATED_IN_5_5
 
 DECLARE_DYNAMIC_MULTICAST_DELEGATE_ThreeParams(FProtocolReceivedDelegate, FDMXProtocolName, Protocol, int32, RemoteUniverse, const TArray<uint8>&, DMXBuffer);
 
@@ -33,136 +35,148 @@ class DMXRUNTIME_API UDMXSubsystem
 	GENERATED_BODY()
 	
 public:
-	/**  Send DMX using function names and integer values. */
+	/** Send DMX using function names and integer values. */
+	UE_DEPRECATED(4.27, "Use UDMXEntityFixurePatch::SendDMX instead.")
 	UFUNCTION(BlueprintCallable, Category = "DMX", meta = (DeprecatedFunction, DeprecationMessage = "Deprecated 4.27. Use DMXEntityFixurePatch::SendDMX instead"))
-	UE_DEPRECATED(5.5, "Use UDMXEntityFixurePatch::SendDMX instead.")
 	void SendDMX(UDMXEntityFixturePatch* FixturePatch, TMap<FDMXAttributeName, int32> AttributeMap, EDMXSendResult& OutResult);
 
-	/**  DEPRECATED 4.27 */
+	/** DEPRECATED 4.27 */
+	UE_DEPRECATED(4.27, "Use UDMXSubsystem::SendDMXToOutputPort instead.")
 	UFUNCTION(BlueprintCallable, Category = "DMX", meta = (DeprecatedFunction, DeprecationMessage = "Deprecated 4.27. Use SendDMXToOutputPort instead."))
-	UE_DEPRECATED(5.5, "Use UDMXSubsystem::SendDMXToOutputPort instead.")
 	void SendDMXRaw(FDMXProtocolName SelectedProtocol, int32 RemoteUniverse, TMap<int32, uint8> AddressValueMap, EDMXSendResult& OutResult);
 
-	/**  Sends DMX Values over the Output Port */
-	UFUNCTION(BlueprintCallable, Category = "DMX")
+	/** Sends DMX via an Output Port. */
+	UFUNCTION(BlueprintCallable, Category = "DMX", Meta = (DisplayName = "Send DMX To Output Port"))
 	static void SendDMXToOutputPort(FDMXOutputPortReference OutputPortReference, TMap<int32, uint8> ChannelToValueMap, int32 LocalUniverse = 1);
 
-	/**  DEPRECATED 4.27 */
+	/** DEPRECATED 4.27 */
+	UE_DEPRECATED(4.27, "Use UDMXSubsystem::GetDMXDataFromInputPort or GetDMXDataFromOutputPort instead.")
 	UFUNCTION(BlueprintCallable, Category = "DMX", meta = (DeprecatedFunction, DeprecationMessage = "Deprecated 4.27. Use GetDMXDataFromInputPort or GetDMXDataFromOutputPort instead."))
-	UE_DEPRECATED(5.5, "Use UDMXSubsystem::GetDMXDataFromInputPort or GetDMXDataFromOutputPort instead.")
 	void GetRawBuffer(FDMXProtocolName SelectedProtocol, int32 RemoteUniverse, TArray<uint8>& DMXBuffer);
 	
-	/**  Gets accumulated latest DMX Values from the Input Port (all that's been received since Begin Play) */
+	/** Gets latest DMX Values from a DMX Universe of a DMX Input Port. If no DMX was received the resulting array will be empty. */
 	UFUNCTION(BlueprintCallable, Category = "DMX", Meta = (DisplayName = "Get DMX Data From Input Port"))
 	static void GetDMXDataFromInputPort(FDMXInputPortReference InputPortReference, TArray<uint8>& DMXData, int32 LocalUniverse = 1);
 	
-	/**  Gets accumulated latest DMX Values from the Output Port  (all that's been sent since Begin Play)*/
+	/** Gets latest DMX Values from a DMX Universe of a DMX Output Port. If no DMX was received the resulting array will be empty. */
 	UFUNCTION(BlueprintCallable, Category = "DMX", Meta = (DisplayName = "Get DMX Data From Output Port"))
 	static void GetDMXDataFromOutputPort(FDMXOutputPortReference OutputPortReference, TArray<uint8>& DMXData, int32 LocalUniverse = 1);
 
-	/**  Return reference to array of Fixture Patch objects of a given type. */
+	/** Return an array of Fixture Patches that use the provided Fixture Type. */
 	UFUNCTION(BlueprintCallable, Category = "DMX", meta = (AutoCreateRefTerm = "FixtureType"))
 	void GetAllFixturesOfType(const FDMXEntityFixtureTypeRef& FixtureType, TArray<UDMXEntityFixturePatch*>& OutResult);
 
-	/**  Return reference to array of Fixture Patch objects of a given category. */
+	/** Return an array of Fixture Patches that use the provided category. */
 	UFUNCTION(BlueprintCallable, Category = "DMX")
 	void GetAllFixturesOfCategory(const UDMXLibrary* DMXLibrary, FDMXFixtureCategory Category, TArray<UDMXEntityFixturePatch*>& OutResult);
 
-	/**  Return reference to array of Fixture Patch objects in a given universe. */
+	/** Return an array of Fixture Patches that reside in the provided universe. */
 	UFUNCTION(BlueprintCallable, Category = "DMX")
 	void GetAllFixturesInUniverse(const UDMXLibrary* DMXLibrary, int32 UniverseId, TArray<UDMXEntityFixturePatch*>& OutResult);
 
-	/**  Return map with all DMX functions and their associated values given DMX buffer and desired universe. */
-	UFUNCTION(BlueprintCallable, Category = "DMX", meta = (DisplayName = "Get Fixture Attributes"))
+	/** Return a map with all DMX functions and their associated values provided DMX buffer and desired universe. */
+	UE_DEPRECATED(5.5, "Instead please call UDMXEntityFixturePatch::GetAttributeValues to retrieve attribute values safely.")
+	UFUNCTION(BlueprintCallable, Category = "DMX", meta = (DeprecatedFunction, DeprecationMessage = "Deprecated 5.5. Instead please call the Fixture Patch function 'Get Attribute Values' to retrieve attribute values safely."))
 	void GetFixtureAttributes(const UDMXEntityFixturePatch* InFixturePatch, const TArray<uint8>& DMXBuffer, TMap<FDMXAttributeName, int32>& OutResult);
 
-	/**  Return reference to array of Fixture Patch objects with a given tag. */
+	/** Return an array of Fixture Patches that have the custom tag set. */
 	UFUNCTION(BlueprintCallable, Category = "DMX")
 	TArray<UDMXEntityFixturePatch*> GetAllFixturesWithTag(const UDMXLibrary* DMXLibrary, FName CustomTag);
 
-	/**  Return reference to array of Fixture Patch objects in library. */
+	/** Return an array of Fixture Patches in the provided DMX Library. */
 	UFUNCTION(BlueprintCallable, Category = "DMX")
 	TArray<UDMXEntityFixturePatch*> GetAllFixturesInLibrary(const UDMXLibrary* DMXLibrary);
 
-	/**  Return reference to Fixture Patch object with a given name. */
+	/** Return the Fixture Patch with given name or an invalid object if no Fixture Patch matches the name. */
 	UFUNCTION(BlueprintCallable, Category = "DMX", meta = (AutoCreateRefTerm="Name"))
 	UDMXEntityFixturePatch* GetFixtureByName(const UDMXLibrary* DMXLibrary, const FString& Name);
 
-	/**  Return reference to array of Fixture Types objects in library. */
+	/** Returns all Fixture Types in a DMX Library. */
 	UFUNCTION(BlueprintCallable, Category = "DMX")
 	TArray<UDMXEntityFixtureType*> GetAllFixtureTypesInLibrary(const UDMXLibrary* DMXLibrary);
 
-	/**  Return reference to Fixture Type object with a given name. */
+	/** Return the Fixture Type with provided name or an invalid object if no Fixture Type matches the name. */
 	UFUNCTION(BlueprintCallable, Category = "DMX", meta = (AutoCreateRefTerm = "Name"))
 	UDMXEntityFixtureType* GetFixtureTypeByName(const UDMXLibrary* DMXLibrary, const FString& Name);
 
-	/**  DEPRECATED 4.27 */
+	/** DEPRECATED 4.27 */
+	UE_DEPRECATED(4.27, "Controllers are removed in favor of Ports.")
 	UFUNCTION(BlueprintCallable, Category = "DMX", meta = (DeprecatedFunction, DeprecationMessage = "Deprecated 4.27. Controllers are removed in favor of Ports."))
-	UE_DEPRECATED(5.5, "Controllers are removed in favor of Ports.")
 	void GetAllUniversesInController(const UDMXLibrary* DMXLibrary, FString ControllerName, TArray<int32>& OutResult);
 
-	/**  DEPRECATED 4.27 */
+	/** DEPRECATED 4.27 */
+	UE_DEPRECATED(4.27, "Controllers are removed in favor of Ports.")
 	UFUNCTION(BlueprintCallable, Category = "DMX", meta = (DeprecatedFunction, DeprecationMessage = "Deprecated 4.27. Controllers are removed in favor of Ports."))
-	UE_DEPRECATED(5.5, "Controllers are removed in favor of Ports.")
 	TArray<UDMXEntityController*> GetAllControllersInLibrary(const UDMXLibrary* DMXLibrary);
 
-	/**  DEPRECATED 4.27 */
+	/** DEPRECATED 4.27 */
+	UE_DEPRECATED(4.27, "Controllers are removed in favor of Ports.")
 	UFUNCTION(BlueprintCallable, Category = "DMX", meta = (AutoCreateRefTerm = "Name"), meta = (DeprecatedFunction, DeprecationMessage = "Deprecated 4.27. Controllers are removed in favor of Ports."))
-	UE_DEPRECATED(5.5, "Controllers are removed in favor of Ports.")
 	UDMXEntityController* GetControllerByName(const UDMXLibrary* DMXLibrary, const FString& Name);
 
+	/** DEPRECATED 5.5, instead use 'Load All DMX Libraries Synchronous' */
 	UE_DEPRECATED(5.5, "Renamed to UDMXSubsystem::LoadAllDMXLibrariesSynchronous. See also UDMXSubsystem::GetDMXLibraries to get soft object ptrs for all DMX Libraries.")
 	TArray<UDMXLibrary*> GetAllDMXLibraries();
 
-	/**  Loads all DMX Libraries in this project synchronous. */
+	/** Loads all DMX Libraries in this project synchronous, returns an array of DMX Libraries */
 	UFUNCTION(BlueprintCallable, Category = "DMX", Meta = (DisplayName = "Load DMX Libraries Synchronous"))
 	TArray<UDMXLibrary*> LoadDMXLibrariesSynchronous() const;
 
-	/**  Returns a Soft Object Ptr to all DMX Library in the project. */
-	UFUNCTION(BlueprintCallable, Category = "DMX")
+	/** Gets all DMX Libraries in this project, returns an array of Soft Object References to the DMX Libraries without loading them. */
+	UFUNCTION(BlueprintCallable, Category = "DMX", Meta = (DisplayName = "Get DMX Libraries"))
 	TArray<TSoftObjectPtr<UDMXLibrary>> GetDMXLibraries() const;
 
 	/**
-	 * Return integer given an array of bytes. Up to the first 4 bytes in the array will be used for the conversion.
-	 * @param bUseLSB	Least Significant Byte mode makes the individual bytes (channels) of the function be
-	 *					interpreted with the first bytes being the lowest part of the number.
-	 *					Most Fixtures use MSB (Most Significant Byte).
+	 * Converts consecutive DMX channel values to a signed 32bit integer value.
+	 * 
+	 * @param Bytes				The byte array that is converted to a normalized value. Up to 3 bytes (24 bits) are supported.
+	 * @param bUseLSB			When true, the byte array is interpreted in little endian format (least significant byte first) otherwise big endian.
+	 * 
+	 * @return					The signed 32bit integer value.
 	 */
 	UFUNCTION(BlueprintPure, Category = "DMX")
 	int32 BytesToInt(const TArray<uint8>& Bytes, bool bUseLSB = false) const;
 
 	/**
-	 * Return normalized value given an array of bytes. Up to the first 4 bytes in the array will be used for the conversion.
-	 * @param bUseLSB	Least Significant Byte mode makes the individual bytes (channels) of the function be
-	 *					interpreted with the first bytes being the lowest part of the number.
-	 *					Most Fixtures use MSB (Most Significant Byte).
+	 * Converts consecutive DMX channel values to a normalized value.
+	 *
+	 * @param Bytes				The byte array that is converted to a normalized value. Up to 4 bytes (32 bits) are supported.
+	 * @param bUseLSB			When true, the byte array is returned in little endian format (least significant byte first) otherwise big endian.
+	 * 
+	 * @return					The normalized value.
 	 */
 	UFUNCTION(BlueprintPure, Category = "DMX", meta = (Keywords = "float"))
 	float BytesToNormalizedValue(const TArray<uint8>& Bytes, bool bUseLSB = false) const;
 
 	/**
-	 * Return the Bytes format of Value in the desired Signal Format.
-	 * @param bUseLSB	Least Significant Byte mode makes the individual bytes (channels) of the function be
-	 *					interpreted with the first bytes being the lowest part of the number.
-	 *					Most Fixtures use MSB (Most Significant Byte).
+	 * Converts a normalized value to an array of DMX channel values.
+	 *
+	 * @param InValue			The normalized floating point value in the range of 0.0 - 1.0. Other values get clamped.
+	 * @param InSignalFormat	Specifies the resolution of the resulting byte array, hence the precision of resulting data (e.g. 0-65535 for 16bit).
+	 * @param Bytes				The resulting byte array.
+	 * @param bUseLSB			When true, the byte array is returned in little endian format (least significant byte first) otherwise big endian.
 	 */
 	UFUNCTION(BlueprintPure, Category = "DMX", meta = (Keywords = "float"))
 	void NormalizedValueToBytes(float InValue, EDMXFixtureSignalFormat InSignalFormat, TArray<uint8>& Bytes, bool bUseLSB = false) const;
 
 	/**
-	 * Return the Bytes format of Value in the desired Signal Format.
-	 * @param bUseLSB	Least Significant Byte mode makes the individual bytes (channels) of the function be
-	 *					interpreted with the first bytes being the lowest part of the number.
-	 *					Most Fixtures use MSB (Most Significant Byte).
+	 * Converts a signed 32bit integer value to an array of DMX channel values.
+	 * 
+	 * @param InValue			The signed 32bit integer value. The value range depends on the signal format. Excess values get clamped.
+	 * @param InSignalFormat	Specifies the resolution of the resulting byte array, hence the precision of resulting data (e.g. 0-65535 for 16bit).
+	 * @param Bytes				The resulting byte array.
+	 * @param bUseLSB			When true, the byte array is returned in little endian format (least significant byte first) otherwise big endian.
 	 */
-	UFUNCTION(BlueprintPure, Category = "DMX")
+	UFUNCTION(BlueprintPure, Category = "DMX", Meta = (DisplayName = "Int To Bytes"))
 	static void IntValueToBytes(int32 InValue, EDMXFixtureSignalFormat InSignalFormat, TArray<uint8>& Bytes, bool bUseLSB = false);
 
 	/**
-	 * Return the normalized value of an Int value from the specified Signal Format.
-	 * @param bUseLSB	Least Significant Byte mode makes the individual bytes (channels) of the function be
-	 *					interpreted with the first bytes being the lowest part of the number.
-	 *					Most Fixtures use MSB (Most Significant Byte).
+	 * Converts a signed 32bit integer value to a normalized value.
+	 * 
+	 * @param InValue			The signed 32bit integer value. The value range depends on the signal format (e.g. 0-65535 for 16bit). Excess values get clamped.
+	 * @param bUseLSB			When true, the byte array is returned in little endian format (least significant byte first) otherwise big endian.
+	 *  
+	 * @return					The normalized value.
 	 */
 	UFUNCTION(BlueprintPure, Category = "DMX", meta = (Keywords = "float"))
 	float IntToNormalizedValue(int32 InValue, EDMXFixtureSignalFormat InSignalFormat) const;
@@ -172,21 +186,24 @@ public:
 	 * @return	The normalized value of the passed in Int using the Function's signal format.
 	 *			-1.0 if the Function is not found in the Fixture Patch.
 	 */
-	UFUNCTION(BlueprintPure, Category = "DMX", meta = (Keywords = "float", DisplayName = "GetNormalizedAttributeValue"))
+	UE_DEPRECATED(5.5, "Instead please call the optimized UDMXEntityFixturePatch::GetAttributeValue to retrieve the attribute value.")
+	UFUNCTION(BlueprintPure, Category = "DMX", meta = (DeprecatedFunction, DeprecationMessage = "Deprecated 5.5. Instead please call the optimized Fixture Patch function 'Get Attribute Value' to retrieve the attribute value."))
 	float GetNormalizedAttributeValue(UDMXEntityFixturePatch* InFixturePatch, FDMXAttributeName InFunctionAttribute, int32 InValue) const;
 
 	/**
-	 * Creates a literal UDMXEntityFixturePatch reference
-	 * @param	InFixturePatch	pointer to set the UDMXEntityFixturePatch to
-	 * @return	The literal UDMXEntityFixturePatch
+	 * Gets the Fixture Type from a Fixture Type Reference.
+	 * 
+	 * @param	InFixtureType		Can be used to set the Fixture Type of the reference.
+	 * @return						The referenced Fixture Type or nullptr if no valid Fixture Type is referenced.
 	 */
 	UFUNCTION(BlueprintPure, meta = (BlueprintInternalUseOnly = "TRUE"), Category = "DMX")
 	UDMXEntityFixtureType* GetFixtureType(FDMXEntityFixtureTypeRef InFixtureType);
 
 	/**
-	 * Creates a literal UDMXEntityFixturePatch reference
-	 * @param	InFixturePatch	pointer to set the UDMXEntityFixturePatch to
-	 * @return	The literal UDMXEntityFixturePatch
+	 * Gets the Fixture Patch from a Fixture Patch Reference.
+	 *
+	 * @param	InFixturePatch		Can be used to set the referenced Fixture Patch of the reference.
+	 * @return						The referenced Fixture Patch or nullptr if no valid Fixture Patch is referenced.
 	 */
 	UFUNCTION(BlueprintPure, meta = (BlueprintInternalUseOnly = "TRUE"), Category = "DMX")
 	UDMXEntityFixturePatch* GetFixturePatch(FDMXEntityFixturePatchRef InFixturePatch);
@@ -207,7 +224,8 @@ public:
 	 * @param	OutAttributesMap Function and Channel value output map
 	 * @return	True if outputting was successfully
 	 */
-	UFUNCTION(BlueprintCallable, meta = (BlueprintInternalUseOnly = "TRUE", AutoCreateRefTerm = "SelectedProtocol"), Category = "DMX")
+	UE_DEPRECATED(5.5, "Duplicate of GetFunctionsMap. Instead please call UDMXSubsystem::GetFunctionsMap or UDMXEntityFixturePatch::GetAttributeValue.")
+	UFUNCTION(BlueprintCallable, meta = (BlueprintInternalUseOnly = "TRUE", AutoCreateRefTerm = "SelectedProtocol", DeprecatedFunction, DeprecationMessage = "Duplicate of GetFunctionsMap. Instead please call DMXSubsystem::GetFunctionsMap or DMXEntityFixturePatch::GetAttributeValues."), Category = "DMX")
 	bool GetFunctionsMapForPatch(UDMXEntityFixturePatch* InFixturePatch, TMap<FDMXAttributeName, int32>& OutAttributesMap);
 
 	/**
@@ -220,26 +238,28 @@ public:
 	int32 GetFunctionsValue(const FName FunctionAttributeName, const TMap<FDMXAttributeName, int32>& InAttributesMap);
 
 	/**
-	 * Checks if a FixturePatchs is of a given FixtureType
-	 * @param	InFixturePatch fixture patch to check
-	 * @param	RefTypeValue a FixtureTypeRef to check against the fixture patch type
-	 * @return	bool result of checking if the patch is of a given type 
+	 * Returns true if a Fixture Patch is of a given FixtureType
+	 * 
+	 * @param	InFixturePatch Fixture Patch to check
+	 * @param	RefTypeValue a FixtureTypeRef to check against the Fixture Patch.
+	 * @return	Returns true if the Fixture Patch matches the Fixture Type.
 	 */
 	UFUNCTION(BlueprintCallable, meta = (BlueprintInternalUseOnly = "TRUE"), Category = "DMX")
 	bool PatchIsOfSelectedType(UDMXEntityFixturePatch* InFixturePatch, FString RefTypeValue);
 
-	/** Get a DMX Subsystem, pure version */
+	/** Get the DMX Subsystem, pure. */
 	UFUNCTION(BlueprintPure, Category = "DMX Subsystem", meta = (BlueprintInternalUseOnly = "true"))
 	static UDMXSubsystem* GetDMXSubsystem_Pure();
 
-	/** Get a DMX Subsystem, callable version */
+	/** Get the DMX Subsystem, callable. */
 	UFUNCTION(BlueprintCallable, Category = "DMX Subsystem", meta = (BlueprintInternalUseOnly = "true"))
 	static UDMXSubsystem* GetDMXSubsystem_Callable();
 
 	/**
-	 * Gets the FName for a FDMXAttributeName, since structs can't have UFUCNTIONS to create a getter
-	 * @param	AttributeName the struct we want to grab the name from
-	 * @return	FName the name of the AttributeName
+	 * Gets the FName of an Attribute Name.
+	 * 
+	 * @param AttributeName		The Attribute Name struct.
+	 * @return					The name of the Attribute
 	 */
 	UFUNCTION(BlueprintCallable, Category = "DMX")
 	FName GetAttributeLabel(FDMXAttributeName AttributeName);
@@ -248,62 +268,62 @@ public:
 	UPROPERTY(BlueprintAssignable, Category = "DMX", meta = (DeprecatedProperty, DeprecationMessage = "WARNING: This can execute faster than tick leading to possible blueprint performance issues. Use DMXComponent's OnFixturePatchReceived event or GetRawBuffer instead."))
 	FProtocolReceivedDelegate OnProtocolReceived_DEPRECATED;
 
-	/**  Set DMX Cell value using matrix coordinates. */
+	/** Set DMX Cell value using matrix coordinates. */
 	UE_DEPRECATED(4.27, "Use UDMXEntityFixturePatch::SetMatrixCellValue instead.")
 	UFUNCTION(BlueprintCallable, Category = "DMX", meta = (DeprecatedFunction, DeprecationMessage = "Deprecated 4.26. DMXEntityFixurePatch::SendMatrixCellValue instead"))
 	bool SetMatrixCellValue(UDMXEntityFixturePatch* FixturePatch, FIntPoint Coordinate /* Cell coordinate X/Y */, FDMXAttributeName Attribute, int32 Value);
 
-	/**  Get DMX Cell value using matrix coordinates. */
+	/** Get DMX Cell value using matrix coordinates. */
 	UE_DEPRECATED(4.27, "Use UDMXEntityFixturePatch::GetMatrixCellValue instead.")
 	UFUNCTION(BlueprintCallable, Category = "DMX", meta = (DeprecatedFunction, DeprecationMessage = "Deprecated 4.26. DMXEntityFixurePatch::GetMatrixCellValues instead"))
 	bool GetMatrixCellValue(UDMXEntityFixturePatch* FixturePatch, FIntPoint Coordinate /* Cell coordinate X/Y */, TMap<FDMXAttributeName, int32>& AttributeValueMap);
 
-	/**  Gets the starting channel of each cell attribute at given coordinate, relative to the Starting Channel of the patch. */
+	/** Gets the starting channel of each cell attribute at given coordinate, relative to the Starting Channel of the patch. */
 	UE_DEPRECATED(4.27, "Use UDMXEntityFixturePatch::GetMatrixCellChannelsRelative instead.")
 	UFUNCTION(BlueprintCallable, Category = "DMX", meta = (DeprecatedFunction, DeprecationMessage = "Deprecated 4.26. DMXEntityFixurePatch::GetMatrixCellChannelsRelative instead"))
 	bool GetMatrixCellChannelsRelative(UDMXEntityFixturePatch* FixturePatch, FIntPoint Coordinate /* Cell coordinate X/Y */, TMap<FDMXAttributeName, int32>& AttributeChannelMap);
 	
-	/**  Gets the absolute starting channel of each cell attribute at given coordinate */
+	/** Gets the absolute starting channel of each cell attribute at given coordinate */
 	UE_DEPRECATED(4.27, "Use UDMXEntityFixturePatch::GetMatrixCellChannelsAbsolute instead.")
 	UFUNCTION(BlueprintCallable, Category = "DMX", meta = (DeprecatedFunction, DeprecationMessage = "Deprecated 4.26. DMXEntityFixurePatch::GetMatrixCellChannelsAbsolute instead"))
 	bool GetMatrixCellChannelsAbsolute(UDMXEntityFixturePatch* FixturePatch, FIntPoint Coordinate /* Cell coordinate X/Y */, TMap<FDMXAttributeName, int32>& AttributeChannelMap);
 
-	/**  Get Matrix Fixture properties */
+	/** Get Matrix Fixture properties */
 	UE_DEPRECATED(4.27, "Use UDMXEntityFixturePatch::GetMatrixProperties instead.")
 	UFUNCTION(BlueprintPure, Category = "DMX", meta = (DeprecatedFunction, DeprecationMessage = "Deprecated 4.26. DMXEntityFixurePatch::GetMatrixProperties instead"))
 	bool GetMatrixProperties(UDMXEntityFixturePatch* FixturePatch, FDMXFixtureMatrix& MatrixProperties);
 
-	/**  Get all attributes for the fixture patch. */
+	/** Get all attributes for the Fixture Patch. */
 	UE_DEPRECATED(4.27, "Use UDMXEntityFixturePatch::GetCellAttributes instead.")
 	UFUNCTION(BlueprintCallable, Category = "DMX", meta = (DeprecatedFunction, DeprecationMessage = "Deprecated 4.26. DMXEntityFixurePatch::GetCellAttributes instead"))
 	bool GetCellAttributes(UDMXEntityFixturePatch* FixturePatch, TArray<FDMXAttributeName>& CellAttributes);
 
-	/**  Get data for single cell. */
+	/** Get data for single cell. */
 	UE_DEPRECATED(4.27, "Use UDMXEntityFixturePatch::GetMatrixCell instead.")
 	UFUNCTION(BlueprintCallable, Category = "DMX", meta = (DeprecatedFunction, DeprecationMessage = "Deprecated 4.26. DMXEntityFixurePatch::GetMatrixCell instead"))
 	bool GetMatrixCell(UDMXEntityFixturePatch* FixturePatch, FIntPoint Coordinate /* Cell coordinate X/Y */, FDMXCell& Cell);
 
-	/**  Get array of all cells and associated data. */
+	/** Get array of all cells and associated data. */
 	UE_DEPRECATED(4.27, "Use UDMXEntityFixturePatch::GetAllMatrixCells instead.")
 	UFUNCTION(BlueprintCallable, Category = "DMX", meta = (DeprecatedFunction, DeprecationMessage = "Deprecated 4.26. DMXEntityFixurePatch::GetAllMatrixCells instead"))
 	bool GetAllMatrixCells(UDMXEntityFixturePatch* FixturePatch, TArray<FDMXCell>& Cells);
 
-	/**  Sort an array according to the selected distribution pattern. */
+	/** Sort an array according to the selected distribution pattern. */
 	UFUNCTION(BlueprintCallable, Category = "DMX")
 	void PixelMappingDistributionSort(EDMXPixelMappingDistribution InDistribution, int32 InNumXPanels, int32 InNumYPanels, const TArray<int32>& InUnorderedList, TArray<int32>& OutSortedList);
 
-	public:
-		/** Delegate broadcast when all dmx library assets were loaded */
-		UE_DEPRECATED(5.5, "DMX Libraries are no longer loaded by default and this delegate is no longer raised. Instead please use UDMXSubsystem::GetDMXLibraries or UDMXSubsystem::LoadDMXLibrariesSynchronous.")
-		FSimpleMulticastDelegate OnAllDMXLibraryAssetsLoaded;
+public:
+	/** Delegate broadcast when all DMX Library assets were loaded */
+	UE_DEPRECATED(5.5, "DMX Libraries are no longer loaded by default and this delegate is no longer raised. Instead please use UDMXSubsystem::GetDMXLibraries or UDMXSubsystem::LoadDMXLibrariesSynchronous.")
+	FSimpleMulticastDelegate OnAllDMXLibraryAssetsLoaded;
 
 #if WITH_EDITOR
-		/** Delegate broadcast when a dmx library asset was added */
-		UE_DEPRECATED(5.5, "DMX Libraries are no longer loaded by default and this delegate is no longer raised. Instead please refer to the asset subsystem directly. See IAssetRegistry::OnAssetAdded.")
-		FDMXOnDMXLibraryAssetDelegate OnDMXLibraryAssetAdded;
+	/** Delegate broadcast when a DMX Library asset was added */
+	UE_DEPRECATED(5.5, "DMX Libraries are no longer loaded by default and this delegate is no longer raised. Instead please refer to the asset subsystem directly. See IAssetRegistry::OnAssetAdded.")
+	FDMXOnDMXLibraryAssetDelegate OnDMXLibraryAssetAdded;
 
-		/** Delegate broadcast when a dmx library asset was removed */
-		UE_DEPRECATED(5.5, "DMX Libraries are no longer loaded by default and this delegate is no longer raised. Instead please refer to the asset subsystem directly. See IAssetRegistry::OnAssetRemoved.")
-		FDMXOnDMXLibraryAssetDelegate OnDMXLibraryAssetRemoved;
+	/** Delegate broadcast when a DMX Library asset was removed */
+	UE_DEPRECATED(5.5, "DMX Libraries are no longer loaded by default and this delegate is no longer raised. Instead please refer to the asset subsystem directly. See IAssetRegistry::OnAssetRemoved.")
+	FDMXOnDMXLibraryAssetDelegate OnDMXLibraryAssetRemoved;
 #endif // WITH_EDITOR
 };
diff --git a/Engine/Plugins/VirtualProduction/DMX/DMXEngine/Source/DMXRuntime/Public/DMXTypes.h b/Engine/Plugins/VirtualProduction/DMX/DMXEngine/Source/DMXRuntime/Public/DMXTypes.h
--- a/Engine/Plugins/VirtualProduction/DMX/DMXEngine/Source/DMXRuntime/Public/DMXTypes.h
+++ b/Engine/Plugins/VirtualProduction/DMX/DMXEngine/Source/DMXRuntime/Public/DMXTypes.h
@@ -14,7 +14,7 @@ UE_DEPRECATED(5.5, "FDMXColorCIE1931xyY is deprecated. Please use FDMXGDTFColorC
 typedef FDMXGDTFColorCIE1931xyY FDMXColorCIE1931xyY;
 
 /** Holds an array Attribute Names with their normalized Values (expand the property to see the map) */
-USTRUCT(BlueprintType, Category = "DMX")
+USTRUCT(BlueprintType, Category = "DMX", meta = (DisplayName = "DMX Normalized Attribute Value Map"))
 struct DMXRUNTIME_API FDMXNormalizedAttributeValueMap
 {
 	GENERATED_BODY()
diff --git a/Engine/Plugins/VirtualProduction/DMX/DMXEngine/Source/DMXRuntime/Public/Game/DMXComponent.h b/Engine/Plugins/VirtualProduction/DMX/DMXEngine/Source/DMXRuntime/Public/Game/DMXComponent.h
--- a/Engine/Plugins/VirtualProduction/DMX/DMXEngine/Source/DMXRuntime/Public/Game/DMXComponent.h
+++ b/Engine/Plugins/VirtualProduction/DMX/DMXEngine/Source/DMXRuntime/Public/Game/DMXComponent.h
@@ -52,7 +52,7 @@ public:
 	void SetFixturePatch(UDMXEntityFixturePatch* InFixturePatch);
 
 	/** Sets whether the component receives dmx from the patch. Note, this is saved with the component when called in editor. */
-	UFUNCTION(BlueprintCallable, Category = "DMX")
+	UFUNCTION(BlueprintCallable, Category = "DMX", meta = (DisplayName = "Set Receive DMX From Patch"))
 	void SetReceiveDMXFromPatch(bool bReceive);
 
 	UPROPERTY(BlueprintReadOnly, EditAnywhere, Category = "DMX", Meta = (ShowOnlyInnerProperties))
diff --git a/Engine/Plugins/VirtualProduction/DMX/DMXEngine/Source/DMXRuntime/Public/Library/DMXEntityFixturePatch.h b/Engine/Plugins/VirtualProduction/DMX/DMXEngine/Source/DMXRuntime/Public/Library/DMXEntityFixturePatch.h
--- a/Engine/Plugins/VirtualProduction/DMX/DMXEngine/Source/DMXRuntime/Public/Library/DMXEntityFixturePatch.h
+++ b/Engine/Plugins/VirtualProduction/DMX/DMXEngine/Source/DMXRuntime/Public/Library/DMXEntityFixturePatch.h
@@ -25,7 +25,7 @@ struct FPropertyChangedEvent;
 DECLARE_MULTICAST_DELEGATE_OneParam(FDMXOnFixturePatchChangedDelegate, const UDMXEntityFixturePatch* /** ChangedFixturePatch */);
 
 /** Parameters to construct a Fixture Patch. */
-USTRUCT(BlueprintType)
+USTRUCT(BlueprintType, meta = (DisplayName = "DMX Entity Fixture Patch Construction Params"))
 struct DMXRUNTIME_API FDMXEntityFixturePatchConstructionParams
 {
 	GENERATED_BODY()
@@ -43,7 +43,7 @@ struct DMXRUNTIME_API FDMXEntityFixturePatchConstructionParams
 	int32 UniverseID = 1;
 
 	/** Starting channel for when auto-assign address is false */
-	UPROPERTY(EditAnywhere, BlueprintReadWrite, Category = "Fixture Patch", meta = (DisplayName = "Starting Address", UIMin = "1", UIMax = "512", ClampMin = "1", ClampMax = "512"))
+	UPROPERTY(EditAnywhere, BlueprintReadWrite, Category = "Fixture Patch", meta = (DisplayName = "Starting Channel", UIMin = "1", UIMax = "512", ClampMin = "1", ClampMax = "512"))
 	int32 StartingAddress = 1;
 
 	/** 
@@ -107,7 +107,7 @@ public:
 	/** Returns a delegate that is and should be broadcast whenever a Fixture Type changed */
 	static FDMXOnFixturePatchChangedDelegate& GetOnFixturePatchChanged();
 
-	/**  Send DMX using attribute names and integer values. */
+	/** Send DMX using attribute names and integer values. */
 	UFUNCTION(BlueprintCallable, Category = "DMX")
 	void SendDMX(TMap<FDMXAttributeName, int32> AttributeMap);
 
@@ -154,61 +154,30 @@ public:
 	/** Called from Fixture Type to keep ActiveMode in valid range when Modes are removed from the Type */
 	void ValidateActiveMode();
 
-	/** Checks if the current Mode for this Patch is valid for its Fixture Type */
-	UE_DEPRECATED(4.27, "Use GetActiveMode instead.")
-	bool CanReadActiveMode() const;
-
 	/** Returns the active mode, or nullptr if there is no valid active mode */
 	const FDMXFixtureMode* GetActiveMode() const;
 	
 	/** Gets the parent fixture type this was constructed from */
-	FORCEINLINE UDMXEntityFixtureType* GetFixtureType() const { return ParentFixtureTypeTemplate; }
+	UDMXEntityFixtureType* GetFixtureType() const { return ParentFixtureTypeTemplate; }
 
 	/** Sets the fixture type this is using */
 	UFUNCTION(BlueprintCallable, Category = "DMX|Fixture Patch")
 	void SetFixtureType(UDMXEntityFixtureType* NewFixtureType);
 
-	/** Returns the universe ID of the patch*/
-	FORCEINLINE int32 GetUniverseID() const { return UniverseID; }
+	/** Returns the universe ID of the patch */
+	int32 GetUniverseID() const { return UniverseID; }
 
-	/** Sets the Universe ID of the patch*/
+	/** Sets the Universe ID of the patch */
+	UFUNCTION(BlueprintCallable, Category = "DMX|Fixture Patch")
 	void SetUniverseID(int32 NewUniverseID);
 
-#if WITH_EDITOR
-	UE_DEPRECATED(5.1, "bAutoAssignAddress and related members are deprecated. Please use SetStartingChannel instead of SetAutoStartingAddress.")
-	void SetAutoStartingAddress(int32 NewAutoStartingAddress);
-
-	UE_DEPRECATED(5.1, "bAutoAssignAddress and related members are deprecated. Please use GetStartingChannel instead of GetAutoStartingAddress.")
-	FORCEINLINE int32 GetAutoStartingAddress() const { return AutoStartingAddress_DEPRECATED; }
-
-	UE_DEPRECATED(5.1, "bAutoAssignAddress and related members are deprecated. Please use SetStartingChannel instead of SetManualStartingAddress.")
-	void SetManualStartingAddress(int32 NewManualStartingAddress);
-
-	UE_DEPRECATED(5.1, "bAutoAssignAddress and related members are deprecated. Please use GetStartingChannel instead of GetManualStartingAddress.")
-	FORCEINLINE int32 GetManualStartingAddress() const { return ManualStartingAddress_DEPRECATED; }
-#endif // WITH_EDITOR
-
-	/** 
-	 * Sets the starting channel of the Fixture Patch.
-	 * 
-	 * If Auto Assign Address was set to true, turns off Auto Assign Address.
-	 */
+	/** Sets the starting channel of the Fixture Patch. */
 	UFUNCTION(BlueprintCallable, Category = "DMX|Fixture Patch")
 	void SetStartingChannel(int32 NewStartingChannel);
 
 	/** Return the starting channel */
 	UFUNCTION(BlueprintCallable, Category = "DMX|Fixture Patch")
-	FORCEINLINE int32 GetStartingChannel() const { return StartingChannel; }
-
-#if WITH_EDITOR
-	/** Sets bAutoAssignAddress for the patch. Does not update relevant properties. Use UDMXEditor Module's UDMXEditorUtils::AutoAssignAddresses methods instead. */
-	UE_DEPRECATED(5.1, "bAutoAssignAddress and related members are deprecated. Auto assign is now a method in FDMXEditorUtils and applied on demand.")
-	void SetAutoAssignAddressUnsafe(bool bShouldAutoAssignAddress) { bAutoAssignAddress_DEPRECATED = bShouldAutoAssignAddress; }
-
-	/** Returns true if the patch is set to auto assign address */
-	UE_DEPRECATED(5.1, "bAutoAssignAddress and related members are deprecated. Auto assign is now a method in FDMXEditorUtils and applied on demand.")
-	FORCEINLINE bool IsAutoAssignAddress() const { return bAutoAssignAddress_DEPRECATED; }
-#endif // WITH_EDITOR
+	int32 GetStartingChannel() const { return StartingChannel; }
 
 	/** Returns the number of channels this Patch occupies with the Fixture functions from its Active Mode or 0 if the patch has no valid Active Mode */
 	UFUNCTION(BlueprintPure, Category = "DMX|Fixture Patch")
@@ -254,54 +223,45 @@ public:
 	static FName GetMVRFixtureUUIDPropertyNameChecked() { return GET_MEMBER_NAME_CHECKED(UDMXEntityFixturePatch, MVRFixtureUUID); }
 	static FName GetFixtureIDPropertyNameChecked() { return GET_MEMBER_NAME_CHECKED(UDMXEntityFixturePatch, FixtureID); }
 	static FName GetStartingChannelPropertyNameChecked() { return GET_MEMBER_NAME_CHECKED(UDMXEntityFixturePatch, StartingChannel); }
-
-	UE_DEPRECATED(5.1, "bAutoAssignAddress and related members are deprecated. Auto assign is now a method in FDMXEditorUtils and applied on demand.")
-	static FName GetManualStartingAddressPropertyNameChecked() { return GET_MEMBER_NAME_CHECKED(UDMXEntityFixturePatch, ManualStartingAddress_DEPRECATED); }
-
-	UE_DEPRECATED(5.1, "bAutoAssignAddress and related members are deprecated. Auto assign is now a method in FDMXEditorUtils and applied on demand.")
-	static FName GetAutoStartingAddressPropertyNameChecked() { return GET_MEMBER_NAME_CHECKED(UDMXEntityFixturePatch, AutoStartingAddress_DEPRECATED); }
-
-	UE_DEPRECATED(5.1, "bAutoAssignAddress and related members are deprecated. Auto assign is now a method in FDMXEditorUtils and applied on demand.")
-	static FName GetAutoAssignAddressPropertyNameChecked() { return GET_MEMBER_NAME_CHECKED(UDMXEntityFixturePatch, bAutoAssignAddress_DEPRECATED); }
 #endif // WITH_EDITOR
 
 protected:
 	/** The local universe of the patch */
-	UPROPERTY(EditAnywhere, BlueprintReadOnly, Category = "Fixture Patch", meta = (ClampMin = 0, DisplayName = "Universe"))
+	UPROPERTY(EditAnywhere, BlueprintReadOnly, Category = "Fixture Patch", meta = (ClampMin = 0))
 	int32 UniverseID;
 
 #if WITH_EDITORONLY_DATA
 	/** DEPRECATED 5.1 */
-	UPROPERTY(Meta = (DeprecatedProperty, DeprecationMessage = "bAutoAssignAddress and related members are deprecated. Auto assign is now only a method in FDMXEditorUtils and should be applied on demand."))
+	UPROPERTY(meta = (DeprecatedProperty, DeprecationMessage = "bAutoAssignAddress and related members are deprecated. Auto assign is now only a method in FDMXEditorUtils and should be applied on demand."))
 	bool bAutoAssignAddress_DEPRECATED = true;
 
 	/** DEPRECATED 5.1 */
-	UPROPERTY(Meta = (DeprecatedProperty, DeprecationMessage = "bAutoAssignAddress and related members are deprecated. Auto assign is now only a method in FDMXEditorUtils and should be applied on demand."))
+	UPROPERTY(meta = (DeprecatedProperty, DeprecationMessage = "bAutoAssignAddress and related members are deprecated. Auto assign is now only a method in FDMXEditorUtils and should be applied on demand."))
 	int32 ManualStartingAddress_DEPRECATED = 1;
 
 	/** DEPRECATED 5.1 */
-	UPROPERTY(Meta = (DeprecatedProperty, DeprecationMessage = "bAutoAssignAddress and related members are deprecated. Auto assign is now only a method in FDMXEditorUtils and should be applied on demand."))
+	UPROPERTY(meta = (DeprecatedProperty, DeprecationMessage = "bAutoAssignAddress and related members are deprecated. Auto assign is now only a method in FDMXEditorUtils and should be applied on demand."))
 	int32 AutoStartingAddress_DEPRECATED = 1;
 #endif // WITH_EDITORONLY_DATA
 
 	/** Starting Channel of the Patch */
-	UPROPERTY(EditAnywhere, BlueprintReadOnly, Category = "Fixture Patch", Meta = (UIMin = "1", UIMax = "512", ClampMin = "1", ClampMax = "512"))
+	UPROPERTY(EditAnywhere, BlueprintReadOnly, Category = "Fixture Patch", meta = (UIMin = "1", UIMax = "512", ClampMin = "1", ClampMax = "512"))
 	int32 StartingChannel = 0;
 
-	/** Property to point to the template parent fixture for details panel purposes */
+	/** The Fixture Type that defines the DMX channel layout of this Fixture Patch */
 	UPROPERTY(VisibleDefaultsOnly, BlueprintReadWrite, BlueprintSetter = SetFixtureType, Category = "Fixture Patch", meta = (DisplayName = "Fixture Type"))
 	TObjectPtr<UDMXEntityFixtureType> ParentFixtureTypeTemplate;
 
 	/** The Index of the Mode in the Fixture Type the Patch uses */
-	UPROPERTY(EditAnywhere, BlueprintReadOnly, Category = "Fixture Patch")
+	UPROPERTY(EditAnywhere, BlueprintReadOnly, meta = (DisplayName = "Active Mode Index"), Category = "DMX|Fixture Patch")
 	int32 ActiveMode;
 
 	/** The MVR Fixture UUID */
-	UPROPERTY(BlueprintReadOnly, VisibleAnywhere, Category = "Fixture Patch")
+	UPROPERTY(BlueprintReadOnly, VisibleAnywhere, meta = (DisplayName = "MVR Fixture UUID"), Category = "DMX|Fixture Patch")
 	FGuid MVRFixtureUUID;
 
 	/** The Fixture ID. Note, fixture patch much like some lighting consoles only supports numerical fixture IDs, and not a string as per MVR specs. */
-	UPROPERTY(BlueprintReadOnly, VisibleAnywhere, Category = "Fixture Patch")
+	UPROPERTY(BlueprintReadOnly, VisibleAnywhere, Category = "DMX|Fixture Patch")
 	int32 FixtureID = 0;
 
 	/** Delegate broadcast when a Fixture Patch changed */
@@ -309,19 +269,19 @@ protected:
 
 public:
 	/** Custom tags for filtering patches  */
-	UPROPERTY(EditAnywhere, BlueprintReadOnly, Category = "Fixture Patch")
+	UPROPERTY(EditAnywhere, BlueprintReadOnly, Category = "DMX|Fixture Patch")
 	TArray<FName> CustomTags;
 
 #if WITH_EDITORONLY_DATA
 	/** Color when displayed in the fixture patch editor */
-	UPROPERTY(EditAnywhere, BlueprintReadOnly, Category = "Fixture Patch")
+	UPROPERTY(EditAnywhere, BlueprintReadOnly, Category = "DMX|Fixture Patch")
 	FLinearColor EditorColor;
 
 	/** 
 	 * If true, the patch receives dmx and raises the OnFixturePatchReceivedDMX event in editor. 
 	 * NOTE: If 'All Fixture Patches receive DMX in editor' is set to true in Project Settings -> Plugins -> DMX, this setting here is ignored.
 	 */
-	UPROPERTY(EditAnywhere, BlueprintReadOnly, Category = "Fixture Patch")
+	UPROPERTY(EditAnywhere, BlueprintReadOnly, Category = "DMX|Fixture Patch")
 	bool bReceiveDMXInEditor;
 #endif
 
@@ -334,7 +294,7 @@ public:
 #endif // WITH_EDITOR
 
 	/**
-	 * Returns an array of valid attributes for the currently active mode.
+	 * Returns an array of attributes for the currently active mode.
 	 * Attributes outside the Active Mode's channel span range are ignored.
 	 */
 	UFUNCTION(BlueprintPure, Category = "DMX|Fixture Patch")
@@ -358,34 +318,28 @@ public:
 	UFUNCTION(BlueprintPure, Category = "DMX|Fixture Patch")
 	TMap<FDMXAttributeName, int32> GetAttributeChannelAssignments() const;
 
-	/**
-	 * Returns a map of function names and their Data Types.
-	 * Functions outside the Active Mode's channel span range are ignored.
-	 */
+	/** Returns a map of function names and their Data Types. */
 	UFUNCTION(BlueprintPure, Category = "DMX|Fixture Patch")
 	TMap<FDMXAttributeName, EDMXFixtureSignalFormat> GetAttributeSignalFormats() const;
 
 	/** DEPRECATED 4.27 */
 	UE_DEPRECATED(4.27, "Deprecated since it's unclear how to use this function correctly. Use UDMXSubsystem::BytesToInt instead.")
-	UFUNCTION(BlueprintPure, Category = "DMX|Fixture Patch", Meta = (DeprecatedFunction, DeprecationMessage = "Deprecated 4.27. Please use DMXSubsystem's Bytes to Int instead, then create a map from that with Attribute Names ."))
+	UFUNCTION(BlueprintPure, Category = "DMX|Fixture Patch", meta = (DeprecatedFunction, DeprecationMessage = "Deprecated 4.27. Please use DMXSubsystem's Bytes to Int instead, then create a map from that with Attribute Names ."))
 	TMap<FDMXAttributeName, int32> ConvertRawMapToAttributeMap(const TMap<int32, uint8>& RawMap) const;
 
-	/**
-	 * Returns a map of function channels and their values.
-	 * Functions outside the Active Mode's channel span range are ignored.
-	 */
+	/** Converts a map of Attribute Names with their DMX values to a map of DMX channels and with their DMX Values. */
 	UFUNCTION(BlueprintPure, Category = "DMX|Fixture Patch")
 	TMap<int32, uint8> ConvertAttributeMapToRawMap(const TMap<FDMXAttributeName, int32>& FunctionMap) const;
 
-	/**  Returns true if given function map is valid for this fixture. */
+	/** Returns true if the Fixture Patch contains all Attributes in an Attribute Name to DMX Value Map. */
 	UFUNCTION(BlueprintPure, Category = "DMX|Fixture Patch")
 	bool IsMapValid(const TMap<FDMXAttributeName, int32>& FunctionMap) const;
 
-	/**  Returns true if the fixture patch contains the attribute. */
+	/** Returns true if the Fixture Patch contains the specified attribute and can use it to send and receive DMX */
 	UFUNCTION(BlueprintPure, Category = "DMX|Fixture Patch")
 	bool ContainsAttribute(FDMXAttributeName FunctionAttribute) const;
 
-	/**  Returns a map that only contains attributes used this patch. */
+	/** Removes any Attribute Name that can not be sent or received by a Fixture Patch from an Attribute Name to DMX Value Map */
 	UFUNCTION(BlueprintPure, Category = "DMX|Fixture Patch")
 	TMap<FDMXAttributeName, int32> ConvertToValidMap(const TMap<FDMXAttributeName, int32>& FunctionMap) const;
 	
@@ -434,13 +388,22 @@ public:
 	UFUNCTION(BlueprintCallable, Category = "DMX|Fixture Patch", meta = (DeprecatedFunction, DeprecationMessage = "Deprecated 4.26. Use GetNormalizedAttributeValue instead. Note, new method returns normalized values!"))
 	float GetNormalizedAttributeValue(FDMXAttributeName Attribute, bool& bSuccess);
 
+
+	/** DEPRECATED 5.5, renamed to GetAttributeValues. */
+	UE_DEPRECATED(5.5, "Renamed to UDMXEntityFixturePatch::GetAttributeValues for consistency with similar methods.")
+	void GetAttributesValues(TMap<FDMXAttributeName, int32>& AttributesValues);
+
 	/**
 	 * Returns the value of each attribute, or zero if no value was ever received.
 	 *
 	 * @param AttributesValues	Out: Resulting map of Attributes with their values
 	 */
 	UFUNCTION(BlueprintCallable, Category = "DMX|Fixture Patch")
-	void GetAttributesValues(TMap<FDMXAttributeName, int32>& AttributesValues);
+	void GetAttributeValues(TMap<FDMXAttributeName, int32>& AttributeValues);
+
+	/** DEPRECATED 5.5, renamed to GetNormalizedAttributeValues. */
+	UE_DEPRECATED(5.5, "Renamed to UDMXEntityFixturePatch::GetNormalizedAttributeValues for consistency with similar methods.")
+	void GetNormalizedAttributesValues(FDMXNormalizedAttributeValueMap& NormalizedAttributeValues);
 
 	/**
 	 * Returns the normalized value of each attribute, or zero if no value was ever received.
@@ -448,7 +411,7 @@ public:
 	 * @param AttributesValues	Out: Resulting map of Attributes with their normalized values
 	 */
 	UFUNCTION(BlueprintCallable, Category = "DMX|Fixture Patch")
-	void GetNormalizedAttributesValues(FDMXNormalizedAttributeValueMap& NormalizedAttributesValues);
+	void GetNormalizedAttributeValues(FDMXNormalizedAttributeValueMap& NormalizedAttributesValues);
 
 	/** Sends the DMX value of the Attribute to specified matrix coordinates */
 	UFUNCTION(BlueprintCallable, Category = "DMX|Fixture Patch")
@@ -463,39 +426,39 @@ public:
 	UFUNCTION(BlueprintCallable, Category = "DMX|Fixture Patch")
 	bool SendNormalizedMatrixCellValue(const FIntPoint& CellCoordinate /* Cell coordinate X/Y */, const FDMXAttributeName& Attribute, float RelativeValue);
 
-	/**  Gets the DMX Cell value using matrix coordinates */
+	/** Gets the DMX Cell value using matrix coordinates */
 	UFUNCTION(BlueprintCallable, Category = "DMX|Fixture Patch")
 	bool GetMatrixCellValues(const FIntPoint& CellCoordinate /* Cell coordinate X/Y */, TMap<FDMXAttributeName, int32>& ValuePerAttribute);
 	
-	/**  Gets the DMX Cell value using matrix coordinates. */
+	/** Gets the DMX Cell value using matrix coordinates. */
 	UFUNCTION(BlueprintCallable, Category = "DMX|Fixture Patch")
 	bool GetNormalizedMatrixCellValues(const FIntPoint& CellCoordinate /* Cell coordinate X/Y */, TMap<FDMXAttributeName, float>& NormalizedValuePerAttribute);
 
-	/**  Gets the starting channel of each cell attribute at given coordinate, relative to the Starting Channel of the patch. */
+	/** Gets the starting channel of each cell attribute at given coordinate, relative to the Starting Channel of the patch. */
 	UFUNCTION(BlueprintCallable, Category = "DMX|Fixture Patch")
 	bool GetMatrixCellChannelsRelative(const FIntPoint& CellCoordinate /* Cell coordinate X/Y */, TMap<FDMXAttributeName, int32>& AttributeChannelMap);
 	
-	/**  Gets the absolute starting channel of each cell attribute at given coordinate */
+	/** Gets the absolute starting channel of each cell attribute at given coordinate */
 	UFUNCTION(BlueprintCallable, Category = "DMX|Fixture Patch")
 	bool GetMatrixCellChannelsAbsolute(const FIntPoint& CellCoordinate /* Cell coordinate X/Y */, TMap<FDMXAttributeName, int32>& AttributeChannelMap);
 
-	/**  Validates and gets the absolute starting channel of each cell attribute at given coordinate */
+	/** Validates and gets the absolute starting channel of each cell attribute at given coordinate */
 	UFUNCTION(BlueprintCallable, Category = "DMX|Fixture Patch")
 	bool GetMatrixCellChannelsAbsoluteWithValidation(const FIntPoint& InCellCoordinate /* Cell coordinate X/Y */, TMap<FDMXAttributeName, int32>& OutAttributeChannelMap);
 
-	/**  Gets the Matrix Fixture properties, returns false if the patch is not using a matrix fixture */
+	/** Gets the Matrix Fixture properties, returns false if the patch is not using a matrix fixture */
 	UFUNCTION(BlueprintPure, Category = "DMX|Fixture Patch")
 	bool GetMatrixProperties(FDMXFixtureMatrix& MatrixProperties) const;
 
-	/**  Gets all attributes names of a cell */
+	/** Gets all attributes names of a cell */
 	UFUNCTION(BlueprintCallable, Category = "DMX|Fixture Patch")
 	bool GetCellAttributes(TArray<FDMXAttributeName>& CellAttributes);
 
-	/**  Gets the cell corresponding to the passed in coordinate */
+	/** Gets the cell corresponding to the passed in coordinate */
 	UFUNCTION(BlueprintCallable, Category = "DMX|Fixture Patch")
 	bool GetMatrixCell(const FIntPoint& CellCoordinate /* Cell coordinate X/Y */, FDMXCell& Cell);
 
-	/**  Gets all matrix cells */
+	/** Gets all matrix cells */
 	UFUNCTION(BlueprintCallable, Category = "DMX|Fixture Patch")
 	bool GetAllMatrixCells(TArray<FDMXCell>& Cells);
 
diff --git a/Engine/Plugins/VirtualProduction/DMX/DMXEngine/Source/DMXRuntime/Public/Library/DMXEntityFixtureType.h b/Engine/Plugins/VirtualProduction/DMX/DMXEngine/Source/DMXRuntime/Public/Library/DMXEntityFixtureType.h
--- a/Engine/Plugins/VirtualProduction/DMX/DMXEngine/Source/DMXRuntime/Public/Library/DMXEntityFixtureType.h
+++ b/Engine/Plugins/VirtualProduction/DMX/DMXEngine/Source/DMXRuntime/Public/Library/DMXEntityFixtureType.h
@@ -16,7 +16,7 @@ class UDMXImport;
 class UDMXImportGDTF;
 
 
-UENUM(BlueprintType)
+UENUM(BlueprintType, DisplayName = "DMX Pixel Mapping Distribution")
 enum class EDMXPixelMappingDistribution : uint8
 {
 	TopLeftToRight,
@@ -40,7 +40,7 @@ enum class EDMXPixelMappingDistribution : uint8
 	BottomRightToAntiClockwise
 };
 
-USTRUCT(BlueprintType)
+USTRUCT(BlueprintType, meta = (DisplayName = "DMX Fixture Function"))
 struct DMXRUNTIME_API FDMXFixtureFunction
 {
 	GENERATED_BODY()
@@ -115,10 +115,6 @@ struct DMXRUNTIME_API FDMXFixtureFunction
 	UPROPERTY(VisibleAnywhere, BlueprintReadWrite, meta = (DisplayName = "Channel Assignment", DisplayPriority = "2"), Category = "Function Settings")
 	int32 Channel = 1;
 
-	/** DEPRECATED 5.0. Instead the 'Channel' property is EditAnywhere so any function can be assigned freely */
-	UPROPERTY(meta = (DeprecatedProperty, DeprecationMessage = "Deprecated since the Channel property can be set in the DMX Library Editor."))
-	int32 ChannelOffset_DEPRECATED = 0;
-
 	/** This function's data type. Defines the used number of channels (bytes) */
 	UPROPERTY(EditAnywhere, BlueprintReadWrite, meta = (DisplayPriority = "5"), Category = "Function Settings")
 	EDMXFixtureSignalFormat DataType = EDMXFixtureSignalFormat::E8Bit;
@@ -156,6 +152,27 @@ private:
 	UPROPERTY(EditAnywhere, Category = "Physical Properties")
 	double PhysicalTo = 1.0;
 #endif // WITH_EDITORONLY_DATA
+
+public:
+	// Workaround for clang deprecation warnings for any deprecated members in implicitly-defined special member functions
+	PRAGMA_DISABLE_DEPRECATION_WARNINGS
+	FDMXFixtureFunction(const FDMXFixtureFunction&) = default;
+	FDMXFixtureFunction(FDMXFixtureFunction&&) = default;
+	FDMXFixtureFunction& operator=(const FDMXFixtureFunction&) = default;
+	FDMXFixtureFunction& operator=(FDMXFixtureFunction&&) = default;
+	PRAGMA_ENABLE_DEPRECATION_WARNINGS
+
+
+	//////////////////////////////////////////////////
+	// Deprecated Members
+
+#if WITH_EDITORONLY_DATA
+	/** DEPRECATED 5.0 */
+	UE_DEPRECATED(5.0, "Instead please refer to the Channel property")
+	UPROPERTY(meta = (DeprecatedProperty, DeprecationMessage = "Deprecated, instead please refer to the Channel property."))
+	int32 ChannelOffset_DEPRECATED = 0;
+#endif // WITH_EDITORONLY_DATA 
+
 };
 
 template<>
@@ -168,7 +185,7 @@ struct TStructOpsTypeTraits<FDMXFixtureFunction> : public TStructOpsTypeTraitsBa
 };
 
 
-USTRUCT(BlueprintType)
+USTRUCT(BlueprintType, meta = (DisplayName = "DMX Fixture Cell Attribute"))
 struct DMXRUNTIME_API FDMXFixtureCellAttribute
 {
 	GENERATED_BODY()
@@ -223,7 +240,7 @@ struct DMXRUNTIME_API FDMXFixtureCellAttribute
 	bool bUseLSBMode;
 };
 
-USTRUCT(BlueprintType)
+USTRUCT(BlueprintType, meta = (DisplayName = "DMX Fixture Matrix"))
 struct DMXRUNTIME_API FDMXFixtureMatrix
 {
 	GENERATED_BODY()
@@ -237,12 +254,6 @@ struct DMXRUNTIME_API FDMXFixtureMatrix
 	/** Returns the last channel of the Matrix */
 	int32 GetLastChannel() const;
 
-	UE_DEPRECATED(5.0, "Deprecated for more consistent naming. Instead use FDMXFixtureMatrix::GetLastChannel")
-	int32 GetFixtureMatrixLastChannel() const;
-
-	UE_DEPRECATED(5.0, "Deprecated to reduce redundant code. Instead use UDMXFixturePatch::GetMatrixCellChannelsRelative")
-	bool GetChannelsFromCell(FIntPoint CellCoordinate, FDMXAttributeName Attribute, TArray<int32>& Channels) const;
-
 	UPROPERTY(EditAnywhere, BlueprintReadWrite, meta = (DisplayPriority = "60", DisplayName = "Cell Attributes"), Category = "Mode Settings")
 	TArray<FDMXFixtureCellAttribute> CellAttributes;
 
@@ -259,7 +270,7 @@ struct DMXRUNTIME_API FDMXFixtureMatrix
 	EDMXPixelMappingDistribution PixelMappingDistribution = EDMXPixelMappingDistribution::TopLeftToRight;
 };
 
-USTRUCT(BlueprintType)
+USTRUCT(BlueprintType, meta = (DisplayName = "DMX Cell"))
 struct DMXRUNTIME_API FDMXCell
 {
 	GENERATED_BODY()
@@ -278,7 +289,7 @@ struct DMXRUNTIME_API FDMXCell
 	{}
 };
 
-USTRUCT(BlueprintType)
+USTRUCT(BlueprintType, meta = (DisplayName = "DMX Fixture Mode"))
 struct DMXRUNTIME_API FDMXFixtureMode
 {
 	GENERATED_BODY()
@@ -316,17 +327,17 @@ struct DMXRUNTIME_API FDMXFixtureMode
 
 
 /** Parameters to construct a Fixture Type. */
-USTRUCT(BlueprintType)
+USTRUCT(BlueprintType, meta = (DisplayName = "DMX Entity Fixture Type Construction Params"))
 struct DMXRUNTIME_API FDMXEntityFixtureTypeConstructionParams
 {
 	GENERATED_BODY()
 
 	/** The DMX Library in which the Fixture Type will be constructed */
-	UPROPERTY(EditAnywhere, BlueprintReadWrite, Category = "Fixture Type")
+	UPROPERTY(EditAnywhere, BlueprintReadWrite, Category = "Fixture Type", meta = (DisplayName = "Parent DMX Library"))
 	TObjectPtr<UDMXLibrary> ParentDMXLibrary = nullptr;
 
 	/** The Category of the Fixture, useful for Filtering */
-	UPROPERTY(EditAnywhere, BlueprintReadWrite, Category = "Fixture Type")
+	UPROPERTY(EditAnywhere, BlueprintReadWrite, Category = "Fixture Type", meta = (DisplayName = "DMX Category"))
 	FDMXFixtureCategory DMXCategory;
 
 	/** The Modes of the Fixture Type */
@@ -355,11 +366,11 @@ public:
 	UDMXEntityFixtureType();
 
 	/** Creates a new Fixture Type in the DMX Library */
-	UFUNCTION(BlueprintCallable, Category = "DMX")
+	UFUNCTION(BlueprintCallable, Category = "DMX|Fixture Type")
 	static UDMXEntityFixtureType* CreateFixtureTypeInLibrary(FDMXEntityFixtureTypeConstructionParams ConstructionParams, const FString& DesiredName = TEXT(""), bool bMarkDMXLibraryDirty = true);
 
 	/** Removes a Fixture Type from a DMX Library */
-	UFUNCTION(BlueprintCallable, Category = "DMX")
+	UFUNCTION(BlueprintCallable, Category = "DMX|Fixture Type")
 	static void RemoveFixtureTypeFromLibrary(FDMXEntityFixtureTypeRef FixtureTypeRef);
 
 	//~ Begin UObject interface
@@ -382,7 +393,7 @@ public:
 	FString GetCleanGDTFFileNameSynchronous(bool bWithExtension) const;
 
 	UE_DEPRECATED(5.5, "Setting GDTFs this way is not supported. Instead set the GDTFSource and generate the modes via UDMXGDTF.")
-	UFUNCTION(BlueprintCallable, Category = "Fixture Settings", Meta = (DeprecatedFunction, DeprecationMessage = "Setting GDTFs from blueprints was never fully supported and now deprecated. Instead please refer to the Create Fixture Type In DMX Library function."))
+	UFUNCTION(BlueprintCallable, Category = "Fixture Settings", meta = (DeprecatedFunction, DeprecationMessage = "Setting GDTFs from blueprints was never fully supported and now deprecated. Instead please refer to the Create Fixture Type In DMX Library function."))
 	void SetModesFromDMXImport(UDMXImport* DMXImportAsset);
 #endif // WITH_EDITOR
 
@@ -390,21 +401,21 @@ public:
 	static FDMXOnFixtureTypeChangedDelegate& GetOnFixtureTypeChanged();
 
 	/** The Category of the Fixture, useful for Filtering */
-	UPROPERTY(EditAnywhere, BlueprintReadOnly, Category = "Fixture Settings", meta = (DisplayName = "DMX Category"))
+	UPROPERTY(EditAnywhere, BlueprintReadOnly, Category = "DMX|Fixture Type", meta = (DisplayName = "DMX Category"))
 	FDMXFixtureCategory DMXCategory;
 
-	UPROPERTY(EditAnywhere, BlueprintReadOnly, Category = "Fixture Settings")
+	UPROPERTY(EditAnywhere, BlueprintReadOnly, Category = "DMX|Fixture Type")
 	TArray<FDMXFixtureMode> Modes;
 
 	/** 
 	 * Modulators applied right before a patch of this type is received. 
 	 * NOTE: Modulators only affect the patch's normalized values! Untouched values are still available when accesing raw values. 
 	 */
-	UPROPERTY(EditAnywhere, Instanced, Category = "Mode Settings", meta = (DisplayPriority = "50"))
+	UPROPERTY(EditAnywhere, Instanced, Category = "DMX|Fixture Type", meta = (DisplayPriority = "50"))
 	TArray<TObjectPtr<UDMXModulator>> InputModulators;
 
 private:
-	/** Delegate that should be broadcast whenever a fixture type changed */
+	/** Delegate that should be broadcast whenever a Fixture Type changed */
 	static FDMXOnFixtureTypeChangedDelegate OnFixtureTypeChangedDelegate;
 
 
@@ -578,13 +589,13 @@ public:
 	static float BytesToFunctionNormalizedValue(const FDMXFixtureFunction& InFunction, const uint8* InBytes);
 	static float BytesToNormalizedValue(EDMXFixtureSignalFormat InSignalFormat, bool bUseLSB, const uint8* InBytes);
 
-	/** The GDTF that initializes this fixture type. When changed, reinitializes with data from the GDTF. */
-	UPROPERTY(EditAnywhere, BlueprintReadOnly, Category = "Fixture Settings", Meta = (DisplayName = "GDTF Source"))
+	/** The GDTF that initializes this Fixture Type. When changed, reinitializes with data from the GDTF. */
+	UPROPERTY(EditAnywhere, BlueprintReadOnly, Category = "DMX|Fixture Type", meta = (DisplayName = "GDTF Source"))
 	TSoftObjectPtr<UDMXImportGDTF> GDTFSource;
 
 #if WITH_EDITORONLY_DATA
 	/** If checked, generates a new GDTF instead of exporting the imported GDTF. This adopts changes in editor but in most cases will result in data loss and is not recommended. */
-	UPROPERTY(EditAnywhere, AdvancedDisplay, Category = "Fixture Settings", Meta = (DisplayName = "Generate GDTF (not recommended)"))
+	UPROPERTY(EditAnywhere, AdvancedDisplay, Category = "DMX|Fixture Type", meta = (DisplayName = "Generate GDTF (not recommended)"))
 	bool bExportGeneratedGDTF = false;
 
 	/** If true only shows latest GDTF mode revisions in editor */
@@ -598,11 +609,12 @@ public:
 public:
 #if WITH_EDITORONLY_DATA
 	/** DEPRECATED 5.5 - The GDTF from which this Fixture Type was setup. */
-	UE_DEPRECATED(5.5, "Changed to a soft object pointer to reduce the memory footprint of fixture types. Please refer to GDTFSource instead.")
-	UPROPERTY(BlueprintReadOnly, Category = "Fixture Settings", Meta = (DeprecatedProperty, DeprecationMessage = "Changed to a soft object pointer to reduce the memory footprint of fixture types. Please use GDTF Source instead."))
+	UE_DEPRECATED(5.5, "Changed to a soft object pointer to reduce the memory footprint of Fixture Types. Please refer to GDTFSource instead.")
+	UPROPERTY(BlueprintReadOnly, Category = "Fixture Settings", meta = (DeprecatedProperty, DeprecationMessage = "Changed to a soft object pointer to reduce the memory footprint of Fixture Types. Please use GDTF Source instead."))
 	TObjectPtr<UDMXImport> DMXImport;
-#endif
 
-	UPROPERTY(Meta = (DeprecatedProperty, DeprecationMessage = "FixtureMatrixEnabled is deprecated. Instead now each Mode has a FixtureMatrixEnabled property."))
+	UE_DEPRECATED(5.5, "bFixtureMatrixEnabled is deprecated. Instead now each Mode has a FixtureMatrixEnabled property.")
+	UPROPERTY(meta = (DeprecatedProperty, DeprecationMessage = "FixtureMatrixEnabled is deprecated. Instead now each Mode has a FixtureMatrixEnabled property."))
 	bool bFixtureMatrixEnabled_DEPRECATED = false;
+#endif
 };
