DMX: Rework Blueprint Nodes and Runtime API to reset and clear DMX buffers

* DMX Subsystem now has a Blueprint Node to clear DMX Buffers
* Fixture Patch now offers Blueprint Nodes to send default and zero values.
* Monitors now update when DMX Buffers were cleared by other systems.
* Changing port settings in the DMX Library now clears DMX Buffers.

diff --git a/Engine/Plugins/VirtualProduction/DMX/DMXEngine/Source/DMXEditor/Private/Customizations/DMXLibraryDetails.cpp b/Engine/Plugins/VirtualProduction/DMX/DMXEngine/Source/DMXEditor/Private/Customizations/DMXLibraryDetails.cpp
--- a/Engine/Plugins/VirtualProduction/DMX/DMXEngine/Source/DMXEditor/Private/Customizations/DMXLibraryDetails.cpp
+++ b/Engine/Plugins/VirtualProduction/DMX/DMXEngine/Source/DMXEditor/Private/Customizations/DMXLibraryDetails.cpp
@@ -3,19 +3,17 @@
 #include "Customizations/DMXLibraryDetails.h"
 
 #include "DMXEditorStyle.h"
+#include "DMXSubsystem.h"
 #include "Library/DMXLibrary.h"
-
 #include "DetailCategoryBuilder.h"
 #include "DetailLayoutBuilder.h"
 #include "DetailWidgetRow.h"
-#include "EditorFontGlyphs.h"
 #include "ISettingsModule.h"
 #include "Modules/ModuleManager.h"
 #include "Widgets/Images/SImage.h"
 #include "Widgets/Input/SButton.h"
 #include "Widgets/Text/STextBlock.h"
 
-
 #define LOCTEXT_NAMESPACE "DMXLibraryDetails"
 
 TSharedRef<IDetailCustomization> FDMXLibraryDetails::MakeInstance()
@@ -71,6 +69,13 @@ void FDMXLibraryDetails::CustomizeDetails(IDetailLayoutBuilder& DetailBuilder)
 	DetailCategoryBuilder
 		.AddProperty(PortReferencesHandle.ToSharedRef())
 		.ShouldAutoExpand(true);
+
+	PortReferencesHandle->SetOnChildPropertyValueChanged(FSimpleDelegate::CreateSP(this, &FDMXLibraryDetails::OnPortReferencesChanged));
+}
+
+void FDMXLibraryDetails::OnPortReferencesChanged()
+{
+	UDMXSubsystem::ClearDMXBuffers();
 }
 
 #undef LOCTEXT_NAMESPACE
diff --git a/Engine/Plugins/VirtualProduction/DMX/DMXEngine/Source/DMXEditor/Private/Customizations/DMXLibraryDetails.h b/Engine/Plugins/VirtualProduction/DMX/DMXEngine/Source/DMXEditor/Private/Customizations/DMXLibraryDetails.h
--- a/Engine/Plugins/VirtualProduction/DMX/DMXEngine/Source/DMXEditor/Private/Customizations/DMXLibraryDetails.h
+++ b/Engine/Plugins/VirtualProduction/DMX/DMXEngine/Source/DMXEditor/Private/Customizations/DMXLibraryDetails.h
@@ -21,4 +21,8 @@ protected:
 	//~ Begin IDetailCustomization interface
 	virtual void CustomizeDetails(IDetailLayoutBuilder& DetailBuilder) override;
 	//~ End IDetailCustomization interface
+
+private:
+	/** Called when port references changed */
+	void OnPortReferencesChanged();
 };
diff --git a/Engine/Plugins/VirtualProduction/DMX/DMXEngine/Source/DMXEditor/Private/DMXEditorUtils.cpp b/Engine/Plugins/VirtualProduction/DMX/DMXEngine/Source/DMXEditor/Private/DMXEditorUtils.cpp
--- a/Engine/Plugins/VirtualProduction/DMX/DMXEngine/Source/DMXEditor/Private/DMXEditorUtils.cpp
+++ b/Engine/Plugins/VirtualProduction/DMX/DMXEngine/Source/DMXEditor/Private/DMXEditorUtils.cpp
@@ -446,20 +446,13 @@ bool FDMXEditorUtils::DoesLibraryHaveUniverseConflicts(UDMXLibrary* Library, FTe
 
 void FDMXEditorUtils::ClearAllDMXPortBuffers()
 {
-	for (const FDMXInputPortSharedRef& InputPort : FDMXPortManager::Get().GetInputPorts())
-	{
-		InputPort->ClearBuffers();
-	}
-
-	for (const FDMXOutputPortSharedRef& OutputPort : FDMXPortManager::Get().GetOutputPorts())
-	{
-		OutputPort->ClearBuffers();
-	}
+	// DEPRECATED 5.5
+	FDMXPortManager::Get().ClearBuffers();
 }
 
 void FDMXEditorUtils::ClearFixturePatchCachedData()
-{
-	// Clear patch buffers
+{	
+	// DEPRECATED 5.5
 	UDMXSubsystem* Subsystem = UDMXSubsystem::GetDMXSubsystem_Callable();
 	if (Subsystem && Subsystem->IsValidLowLevel())
 	{
diff --git a/Engine/Plugins/VirtualProduction/DMX/DMXEngine/Source/DMXEditor/Private/DMXPIEManager.cpp b/Engine/Plugins/VirtualProduction/DMX/DMXEngine/Source/DMXEditor/Private/DMXPIEManager.cpp
--- a/Engine/Plugins/VirtualProduction/DMX/DMXEngine/Source/DMXEditor/Private/DMXPIEManager.cpp
+++ b/Engine/Plugins/VirtualProduction/DMX/DMXEngine/Source/DMXEditor/Private/DMXPIEManager.cpp
@@ -2,8 +2,7 @@
 
 #include "DMXPIEManager.h"
 
-#include "DMXEditorUtils.h"
-
+#include "DMXSubsystem.h"
 #include "Editor.h"
 
 
@@ -21,12 +20,10 @@ FDMXPIEManager::~FDMXPIEManager()
 
 void FDMXPIEManager::OnBeginPIE(const bool bIsSimulating)
 {
-	FDMXEditorUtils::ClearAllDMXPortBuffers();
-	FDMXEditorUtils::ClearFixturePatchCachedData();
+	UDMXSubsystem::ClearDMXBuffers();
 }
 
 void FDMXPIEManager::OnEndPIE(const bool bIsSimulating)
 {
-	FDMXEditorUtils::ClearAllDMXPortBuffers();
-	FDMXEditorUtils::ClearFixturePatchCachedData();
+	UDMXSubsystem::ClearDMXBuffers();
 }
diff --git a/Engine/Plugins/VirtualProduction/DMX/DMXEngine/Source/DMXEditor/Private/Widgets/Monitors/SDMXActivityMonitor.cpp b/Engine/Plugins/VirtualProduction/DMX/DMXEngine/Source/DMXEditor/Private/Widgets/Monitors/SDMXActivityMonitor.cpp
--- a/Engine/Plugins/VirtualProduction/DMX/DMXEngine/Source/DMXEditor/Private/Widgets/Monitors/SDMXActivityMonitor.cpp
+++ b/Engine/Plugins/VirtualProduction/DMX/DMXEngine/Source/DMXEditor/Private/Widgets/Monitors/SDMXActivityMonitor.cpp
@@ -6,26 +6,25 @@
 #include "Containers/UnrealString.h"
 #include "DMXEditorSettings.h"
 #include "DMXEditorStyle.h"
-#include "DMXEditorUtils.h"
 #include "DMXProtocolTypes.h"
+#include "DMXSubsystem.h"
 #include "IO/DMXInputPort.h"
 #include "IO/DMXOutputPort.h"
 #include "IO/DMXPortManager.h"
 #include "IO/DMXRawListener.h"
-#include "Styling/AppStyle.h"
 #include "SlateOptMacros.h"
+#include "Styling/AppStyle.h"
 #include "Widgets/Input/SButton.h"
 #include "Widgets/Input/SEditableTextBox.h"
+#include "Widgets/Input/SSpinBox.h"
 #include "Widgets/Layout/SScrollBar.h"
 #include "Widgets/Layout/SScrollBox.h"
 #include "Widgets/Layout/SSeparator.h"
 #include "Widgets/Layout/SWrapBox.h"
-#include "Widgets/Input/SSpinBox.h"
 #include "Widgets/Monitors/SDMXActivityInUniverse.h"
 #include "Widgets/Monitors/SDMXMonitorSourceSelector.h"
 #include "Widgets/Views/SListView.h"
 
-
 #define LOCTEXT_NAMESPACE "SDMXActivityMonitor"
 
 SDMXActivityMonitor::SDMXActivityMonitor()
@@ -231,21 +230,11 @@ void SDMXActivityMonitor::Construct(const FArguments& InArgs)
 	UpdateListenerRegistration();
 
 	FDMXPortManager::Get().OnPortsChanged.AddSP(this, &SDMXActivityMonitor::OnPortsChanged);
+	FDMXPortManager::Get().GetOnBuffersCleared().AddSP(this, &SDMXActivityMonitor::OnPortBuffersCleared);
 }
 END_SLATE_FUNCTION_BUILD_OPTIMIZATION
 
 
-void SDMXActivityMonitor::ClearAllDMXBuffers()
-{
-	FDMXEditorUtils::ClearAllDMXPortBuffers();
-	UniverseToDataMap.Reset();
-
-	for (const TSharedRef<FDMXRawListener>& Input : DMXListeners)
-	{
-		Input->ClearBuffer();
-	}
-}
-
 void SDMXActivityMonitor::ResizeDataMapToUniverseRange()
 {
 	const TMap<int32, TArray<uint8, TFixedAllocator<DMX_UNIVERSE_SIZE>>> CachedUniverseToDataMap = UniverseToDataMap;
@@ -455,12 +444,14 @@ void SDMXActivityMonitor::OnPortsChanged()
 	UpdateListenerRegistration();
 }
 
-FReply SDMXActivityMonitor::OnClearButtonClicked()
+void SDMXActivityMonitor::OnPortBuffersCleared()
 {
-	FDMXEditorUtils::ClearAllDMXPortBuffers();
-	ClearAllDMXBuffers();
-
 	ClearDisplay();
+}
+
+FReply SDMXActivityMonitor::OnClearButtonClicked()
+{
+	UDMXSubsystem::ClearDMXBuffers();
 
 	return FReply::Handled();
 }
@@ -469,6 +460,8 @@ void SDMXActivityMonitor::ClearDisplay()
 {
 	check(UniverseList.IsValid());
 
+	UniverseToDataMap.Reset();
+
 	UniverseListSource.Reset();
 	UniverseList->RequestListRefresh();
 }
diff --git a/Engine/Plugins/VirtualProduction/DMX/DMXEngine/Source/DMXEditor/Private/Widgets/Monitors/SDMXActivityMonitor.h b/Engine/Plugins/VirtualProduction/DMX/DMXEngine/Source/DMXEditor/Private/Widgets/Monitors/SDMXActivityMonitor.h
--- a/Engine/Plugins/VirtualProduction/DMX/DMXEngine/Source/DMXEditor/Private/Widgets/Monitors/SDMXActivityMonitor.h
+++ b/Engine/Plugins/VirtualProduction/DMX/DMXEngine/Source/DMXEditor/Private/Widgets/Monitors/SDMXActivityMonitor.h
@@ -74,6 +74,9 @@ private:
 	/** Called when ports changed */
 	void OnPortsChanged();
 
+	/** Called when port buffers were cleared */
+	void OnPortBuffersCleared();
+
 	/** Called when the clear ui values button was clicked */
 	FReply OnClearButtonClicked();
 
diff --git a/Engine/Plugins/VirtualProduction/DMX/DMXEngine/Source/DMXEditor/Private/Widgets/Monitors/SDMXChannelsMonitor.cpp b/Engine/Plugins/VirtualProduction/DMX/DMXEngine/Source/DMXEditor/Private/Widgets/Monitors/SDMXChannelsMonitor.cpp
--- a/Engine/Plugins/VirtualProduction/DMX/DMXEngine/Source/DMXEditor/Private/Widgets/Monitors/SDMXChannelsMonitor.cpp
+++ b/Engine/Plugins/VirtualProduction/DMX/DMXEngine/Source/DMXEditor/Private/Widgets/Monitors/SDMXChannelsMonitor.cpp
@@ -2,25 +2,22 @@
 
 #include "Widgets/Monitors/SDMXChannelsMonitor.h"
 
+#include "Containers/UnrealString.h"
 #include "DMXEditorSettings.h"
-#include "DMXEditorUtils.h"
 #include "DMXProtocolCommon.h"
-#include "DMXProtocolSettings.h"
+#include "DMXSubsystem.h"
 #include "IO/DMXInputPort.h"
 #include "IO/DMXOutputPort.h"
-#include "Widgets/SDMXChannel.h"
-#include "Widgets/Monitors/SDMXMonitorSourceSelector.h"
-
-#include "Styling/AppStyle.h"
+#include "IO/DMXPortManager.h"
 #include "SlateOptMacros.h"
-#include "Containers/UnrealString.h"
+#include "Styling/AppStyle.h"
 #include "Widgets/Input/SButton.h"
-#include "Widgets/Input/SCheckBox.h" 
 #include "Widgets/Input/SEditableTextBox.h" 
 #include "Widgets/Layout/SScrollBox.h"
 #include "Widgets/Layout/SSeparator.h"
 #include "Widgets/Layout/SWrapBox.h"
-
+#include "Widgets/Monitors/SDMXMonitorSourceSelector.h"
+#include "Widgets/SDMXChannel.h"
 
 #define LOCTEXT_NAMESPACE "SDMXChannelsMonitor"
 
@@ -133,6 +130,8 @@ void SDMXChannelsMonitor::Construct(const FArguments& InArgs)
 	// Init UI
 	CreateChannelValueWidgets();
 	ZeroChannelValues();
+
+	FDMXPortManager::Get().GetOnBuffersCleared().AddSP(this, &SDMXChannelsMonitor::OnPortBuffersCleared);
 }
 END_SLATE_FUNCTION_BUILD_OPTIMIZATION
 
@@ -143,7 +142,7 @@ void SDMXChannelsMonitor::Tick(const FGeometry& AllottedGeometry, const double I
 	if (SourceSelector->IsMonitorInputPorts())
 	{
 		// The widget returns the right arrays, also output ports that loopback when input ports are selected.
-		for (const FDMXInputPortSharedRef& InputPort : SourceSelector->GetSelectedInputPorts())
+		for (const TSharedRef<FDMXInputPort>& InputPort : SourceSelector->GetSelectedInputPorts())
 		{
 			if (InputPort->GameThreadGetDMXSignal(UniverseID, Signal))
 			{
@@ -153,7 +152,7 @@ void SDMXChannelsMonitor::Tick(const FGeometry& AllottedGeometry, const double I
 	}
 	else
 	{
-		for (const FDMXOutputPortSharedRef& OutputPort : SourceSelector->GetSelectedOutputPorts())
+		for (const TSharedRef<FDMXOutputPort>& OutputPort : SourceSelector->GetSelectedOutputPorts())
 		{
 			constexpr bool bEvenIfNotLoopbackToEngine = true;
 			if (OutputPort->GameThreadGetDMXSignal(UniverseID, Signal, bEvenIfNotLoopbackToEngine))
@@ -255,16 +254,6 @@ void SDMXChannelsMonitor::SetChannelValues(const TArray<uint8>& Buffer)
 	}
 }
 
-FReply SDMXChannelsMonitor::OnClearButtonClicked()
-{
-	FDMXEditorUtils::ClearAllDMXPortBuffers();
-	FDMXEditorUtils::ClearFixturePatchCachedData();
-
-	ZeroChannelValues();
-
-	return FReply::Handled();
-}
-
 void SDMXChannelsMonitor::OnUniverseIDValueCommitted(const FText& InNewText, ETextCommit::Type CommitType)
 {
 	// If the entered text isn't numeric, restore the previous Value
@@ -296,4 +285,16 @@ void SDMXChannelsMonitor::OnUniverseIDValueCommitted(const FText& InNewText, ETe
 	}
 }
 
+FReply SDMXChannelsMonitor::OnClearButtonClicked()
+{
+	UDMXSubsystem::ClearDMXBuffers();
+
+	return FReply::Handled();
+}
+
+void SDMXChannelsMonitor::OnPortBuffersCleared()
+{
+	ZeroChannelValues();
+}
+
 #undef LOCTEXT_NAMESPACE
diff --git a/Engine/Plugins/VirtualProduction/DMX/DMXEngine/Source/DMXEditor/Private/Widgets/Monitors/SDMXChannelsMonitor.h b/Engine/Plugins/VirtualProduction/DMX/DMXEngine/Source/DMXEditor/Private/Widgets/Monitors/SDMXChannelsMonitor.h
--- a/Engine/Plugins/VirtualProduction/DMX/DMXEngine/Source/DMXEditor/Private/Widgets/Monitors/SDMXChannelsMonitor.h
+++ b/Engine/Plugins/VirtualProduction/DMX/DMXEngine/Source/DMXEditor/Private/Widgets/Monitors/SDMXChannelsMonitor.h
@@ -62,6 +62,9 @@ private:
 	/** Called when the clear bButton was clicked */
 	FReply OnClearButtonClicked();
 
+	/** Called when port buffers were cleared */
+	void OnPortBuffersCleared();
+
 	/** Horizontal box that contains the Source Selector and its Label */
 	TSharedPtr<SHorizontalBox> SourceSelectorBox;
 
diff --git a/Engine/Plugins/VirtualProduction/DMX/DMXEngine/Source/DMXEditor/Public/DMXEditorUtils.h b/Engine/Plugins/VirtualProduction/DMX/DMXEngine/Source/DMXEditor/Public/DMXEditorUtils.h
--- a/Engine/Plugins/VirtualProduction/DMX/DMXEngine/Source/DMXEditor/Public/DMXEditorUtils.h
+++ b/Engine/Plugins/VirtualProduction/DMX/DMXEngine/Source/DMXEditor/Public/DMXEditorUtils.h
@@ -86,9 +86,11 @@ public:
 	static bool DoesLibraryHaveUniverseConflicts(UDMXLibrary* Library, FText& OutInputPortConflictMessage, FText& OutOutputPortConflictMessage);
 
 	/** Zeros memory in all active DMX buffers of all protocols */
+	UE_DEPRECATED(5.5, "Instead use FDMXPortManager::ClearPortBuffers")
 	static void ClearAllDMXPortBuffers();
 
 	/** Clears cached data fixture patches received */
+	UE_DEPRECATED(5.5, "Instead use UDMXSubsystem::ClearDMXBuffers.")
 	static void ClearFixturePatchCachedData();
 
 	/** Gets the package or creates a new one if it doesn't exist */
diff --git a/Engine/Plugins/VirtualProduction/DMX/DMXEngine/Source/DMXRuntime/Private/DMXSubsystem.cpp b/Engine/Plugins/VirtualProduction/DMX/DMXEngine/Source/DMXRuntime/Private/DMXSubsystem.cpp
--- a/Engine/Plugins/VirtualProduction/DMX/DMXEngine/Source/DMXRuntime/Private/DMXSubsystem.cpp
+++ b/Engine/Plugins/VirtualProduction/DMX/DMXEngine/Source/DMXRuntime/Private/DMXSubsystem.cpp
@@ -120,6 +120,28 @@ namespace
 #endif // WITH_EDITOR
 }
 
+void UDMXSubsystem::ClearDMXBuffers()
+{
+	// Clear port buffers
+	FDMXPortManager::Get().ClearBuffers();
+
+	// Rebuild fixture patch caches from cleared buffers, effectively clearing them as well.
+	UDMXSubsystem* Subsystem = UDMXSubsystem::GetDMXSubsystem_Callable();
+	if (Subsystem && Subsystem->IsValidLowLevel())
+	{
+		TArray<TSoftObjectPtr<UDMXLibrary>> DMXLibraries = Subsystem->GetDMXLibraries();
+		for (const TSoftObjectPtr<UDMXLibrary>& Library : DMXLibraries)
+		{
+			if (Library.IsValid())
+			{
+				Library.Get()->ForEachEntityOfType<UDMXEntityFixturePatch>([](UDMXEntityFixturePatch* Patch) {
+					Patch->RebuildCache();
+					});
+			}
+		}
+	}
+}
+
 void UDMXSubsystem::SendDMX(UDMXEntityFixturePatch* FixturePatch, TMap<FDMXAttributeName, int32> AttributeMap, EDMXSendResult& OutResult)
 {
 	OutResult = EDMXSendResult::Success;
diff --git a/Engine/Plugins/VirtualProduction/DMX/DMXEngine/Source/DMXRuntime/Private/Library/DMXLibrary.cpp b/Engine/Plugins/VirtualProduction/DMX/DMXEngine/Source/DMXRuntime/Private/Library/DMXLibrary.cpp
--- a/Engine/Plugins/VirtualProduction/DMX/DMXEngine/Source/DMXRuntime/Private/Library/DMXLibrary.cpp
+++ b/Engine/Plugins/VirtualProduction/DMX/DMXEngine/Source/DMXRuntime/Private/Library/DMXLibrary.cpp
@@ -222,7 +222,7 @@ UDMXEntity* UDMXLibrary::FindEntity(const FGuid& Id)
 {
 	for (UDMXEntity* Entity : Entities)
 	{
-		if (Entity && Entity->GetID() == Id)
+		if (IsValid(Entity) && Entity->GetID() == Id)
 		{
 			return Entity;
 		}
diff --git a/Engine/Plugins/VirtualProduction/DMX/DMXEngine/Source/DMXRuntime/Public/DMXSubsystem.h b/Engine/Plugins/VirtualProduction/DMX/DMXEngine/Source/DMXRuntime/Public/DMXSubsystem.h
--- a/Engine/Plugins/VirtualProduction/DMX/DMXEngine/Source/DMXRuntime/Public/DMXSubsystem.h
+++ b/Engine/Plugins/VirtualProduction/DMX/DMXEngine/Source/DMXRuntime/Public/DMXSubsystem.h
@@ -35,6 +35,14 @@ class DMXRUNTIME_API UDMXSubsystem
 	GENERATED_BODY()
 	
 public:
+	/** 
+	 * Clears all buffered DMX data of Ports and Fixture Patches. 
+	 * Note, this function clears the buffers, it does not zero them out.
+	 * To reset to default or zero, see Fixture Patch members Send Default Values and Send Zero Values.
+	 */
+	UFUNCTION(BlueprintCallable, Category = "DMX", meta = (DisplayName = "Clear DMX Buffers"))
+	static void ClearDMXBuffers();
+
 	/** Send DMX using function names and integer values. */
 	UE_DEPRECATED(4.27, "Use UDMXEntityFixurePatch::SendDMX instead.")
 	UFUNCTION(BlueprintCallable, Category = "DMX", meta = (DeprecatedFunction, DeprecationMessage = "Deprecated 4.27. Use DMXEntityFixurePatch::SendDMX instead"))
diff --git a/Engine/Plugins/VirtualProduction/DMX/DMXEngine/Source/DMXRuntime/Public/Library/DMXEntityFixturePatch.h b/Engine/Plugins/VirtualProduction/DMX/DMXEngine/Source/DMXRuntime/Public/Library/DMXEntityFixturePatch.h
--- a/Engine/Plugins/VirtualProduction/DMX/DMXEngine/Source/DMXRuntime/Public/Library/DMXEntityFixturePatch.h
+++ b/Engine/Plugins/VirtualProduction/DMX/DMXEngine/Source/DMXRuntime/Public/Library/DMXEntityFixturePatch.h
@@ -113,14 +113,16 @@ public:
 
 	/** 
 	 * Sends the default value for all attributes, including matrix attributes. 
-	 * Note, this call will not raise send dmx traces.
+	 * Note, calls will not be considered by the DMX Conflict Monitor.
 	 */
+	UFUNCTION(BlueprintCallable, Category = "DMX")
 	void SendDefaultValues();
 
 	/** 
 	 * Sends zeroes for all attributes, including matrix attributes.
-	 * Note, this call will not raise send dmx traces.
+	 * Note, calls will not be considered by the DMX Conflict Monitor.
 	 */
+	UFUNCTION(BlueprintCallable, Category = "DMX")
 	void SendZeroValues();
 
 	/** 
diff --git a/Engine/Plugins/VirtualProduction/DMX/DMXProtocol/Source/DMXProtocol/Private/IO/DMXPortManager.cpp b/Engine/Plugins/VirtualProduction/DMX/DMXProtocol/Source/DMXProtocol/Private/IO/DMXPortManager.cpp
--- a/Engine/Plugins/VirtualProduction/DMX/DMXProtocol/Source/DMXProtocol/Private/IO/DMXPortManager.cpp
+++ b/Engine/Plugins/VirtualProduction/DMX/DMXProtocol/Source/DMXProtocol/Private/IO/DMXPortManager.cpp
@@ -41,6 +41,21 @@ FDMXPortManager& FDMXPortManager::Get()
 	return *CurrentManager;
 }
 
+void FDMXPortManager::ClearBuffers()
+{
+	for (const FDMXInputPortSharedRef& InputPort : FDMXPortManager::Get().GetInputPorts())
+	{
+		InputPort->ClearBuffers();
+	}
+
+	for (const FDMXOutputPortSharedRef& OutputPort : FDMXPortManager::Get().GetOutputPorts())
+	{
+		OutputPort->ClearBuffers();
+	}
+
+	OnBuffersClearedDelegate.Broadcast();
+}
+
 FDMXInputPortSharedRef FDMXPortManager::GetInputPortFromConfigChecked(const FDMXInputPortConfig& InputPortConfig)
 {
 	// The config needs a valid guid
diff --git a/Engine/Plugins/VirtualProduction/DMX/DMXProtocol/Source/DMXProtocol/Public/IO/DMXPortManager.h b/Engine/Plugins/VirtualProduction/DMX/DMXProtocol/Source/DMXProtocol/Public/IO/DMXPortManager.h
--- a/Engine/Plugins/VirtualProduction/DMX/DMXProtocol/Source/DMXProtocol/Public/IO/DMXPortManager.h
+++ b/Engine/Plugins/VirtualProduction/DMX/DMXProtocol/Source/DMXProtocol/Public/IO/DMXPortManager.h
@@ -69,10 +69,18 @@ public:
 
 	static FDMXPortManager& Get();
 
+	/** Returns the current input ports */
 	const TArray<FDMXInputPortSharedRef>& GetInputPorts() const { return InputPorts; }
 
+	/** Returns the current output ports */
 	const TArray<FDMXOutputPortSharedRef>& GetOutputPorts() const { return OutputPorts; }
 
+	/** Clears all buffered data of Input Ports, Output Ports and Raw Listeners. Note, this clears the data, it does not zero it out.*/
+	void ClearBuffers();
+
+	/** Returns a delegate broadcast when port buffers were cleared */
+	FSimpleMulticastDelegate& GetOnBuffersCleared() { return OnBuffersClearedDelegate; }
+
 	/** Gets the input port that corresponds to the input port config. Checks the config is in the DMXProtocolSetting's InputPortConfigs array. */
 	FDMXInputPortSharedRef GetInputPortFromConfigChecked(const FDMXInputPortConfig& InputPortConfig);
 
@@ -142,6 +150,9 @@ private:
 	/** True when protocols are suspended */
 	bool bProtocolsSuspended = false;
 
+	/** Delegate raised when buffers were cleared */
+	FSimpleMulticastDelegate OnBuffersClearedDelegate;
+
 	////////////////////////////////////////////////////////////
 	// Initialization 
 public:
