DMX: Remove now obsolete Analytics Provider code

#rnx

diff --git a/Engine/Plugins/VirtualProduction/DMX/DMXEngine/Source/DMXRuntime/Private/DMXSubsystem.cpp b/Engine/Plugins/VirtualProduction/DMX/DMXEngine/Source/DMXRuntime/Private/DMXSubsystem.cpp
--- a/Engine/Plugins/VirtualProduction/DMX/DMXEngine/Source/DMXRuntime/Private/DMXSubsystem.cpp
+++ b/Engine/Plugins/VirtualProduction/DMX/DMXEngine/Source/DMXRuntime/Private/DMXSubsystem.cpp
@@ -34,92 +34,6 @@
 
 DECLARE_LOG_CATEGORY_CLASS(DMXSubsystemLog, Log, All);
 
-namespace
-{
-	const FName InvalidUniverseError = FName("InvalidUniverseError");
-
-#if WITH_EDITOR
-	/** Helper to create analytics for dmx libraries in use */
-	void CreateEngineAnalytics(const TArray<UDMXLibrary*>& DMXLibraries)
-	{
-		if (FEngineAnalytics::IsAvailable())
-		{
-			// DMX Library usage statistics
-			{
-				int32 CountLibraries = 0;
-				int32 CountPatches = 0;
-				int32 CountChannels = 0;
-
-				for (UDMXLibrary* Library : DMXLibraries)
-				{
-					CountLibraries++;
-
-					for (UDMXEntity* Entity : Library->GetEntities())
-					{
-						if (UDMXEntityFixturePatch* Patch = Cast<UDMXEntityFixturePatch>(Entity))
-						{
-							CountPatches++;
-							CountChannels += Patch->GetChannelSpan();
-						}
-					}
-				}
-
-				TArray<FAnalyticsEventAttribute> LibraryEventAttributes;
-				LibraryEventAttributes.Add(FAnalyticsEventAttribute(TEXT("NumDMXLibraries"), CountLibraries));
-				LibraryEventAttributes.Add(FAnalyticsEventAttribute(TEXT("NumDMXPatches"), CountPatches));
-				LibraryEventAttributes.Add(FAnalyticsEventAttribute(TEXT("NumDMXChannels"), CountChannels));
-				FEngineAnalytics::GetProvider().RecordEvent(TEXT("Usage.DMX.DMXLibraries"), LibraryEventAttributes);
-			}
-
-			// DMX Port usage statistics
-			{
-				int32 CountArtNetPorts = 0;
-				int32 CountSACNPorts = 0;
-				int32 CountOtherPorts = 0;
-				const UDMXProtocolSettings* ProtocolSettings = GetDefault<UDMXProtocolSettings>();
-				for (const FDMXInputPortConfig& Config : ProtocolSettings->InputPortConfigs)
-				{
-					if (Config.GetProtocolName() == "Art-Net")
-					{
-						CountArtNetPorts++;
-					}
-					else if (Config.GetProtocolName() == "sACN")
-					{
-						CountSACNPorts++;
-					}
-					else
-					{
-						CountOtherPorts++;
-					}
-				}
-
-				for (const FDMXOutputPortConfig& Config : ProtocolSettings->OutputPortConfigs)
-				{
-					if (Config.GetProtocolName() == "Art-Net")
-					{
-						CountArtNetPorts++;
-					}
-					else if (Config.GetProtocolName() == "sACN")
-					{
-						CountSACNPorts++;
-					}
-					else
-					{
-						CountOtherPorts++;
-					}
-				}
-
-				TArray<FAnalyticsEventAttribute> PortEventAttributes;
-				PortEventAttributes.Add(FAnalyticsEventAttribute(TEXT("NumArtNetPorts"), CountArtNetPorts));
-				PortEventAttributes.Add(FAnalyticsEventAttribute(TEXT("NumSACNPorts"), CountSACNPorts));
-				PortEventAttributes.Add(FAnalyticsEventAttribute(TEXT("NumOtherPorts"), CountOtherPorts));
-				FEngineAnalytics::GetProvider().RecordEvent(TEXT("Usage.DMX.DMXPorts"), PortEventAttributes);
-			}
-		}
-	}
-#endif // WITH_EDITOR
-}
-
 void UDMXSubsystem::ClearDMXBuffers()
 {
 	// Clear port buffers
