DMX: Change DMX Analytics Provider's event name to Usage.DMX.ToolEvent

#rnx

diff --git a/Engine/Plugins/VirtualProduction/DMX/DMXEngine/Source/DMXEditor/Internal/Analytics/DMXEditorToolAnalyticsProvider.h b/Engine/Plugins/VirtualProduction/DMX/DMXEngine/Source/DMXEditor/Internal/Analytics/DMXEditorToolAnalyticsProvider.h
--- a/Engine/Plugins/VirtualProduction/DMX/DMXEngine/Source/DMXEditor/Internal/Analytics/DMXEditorToolAnalyticsProvider.h
+++ b/Engine/Plugins/VirtualProduction/DMX/DMXEngine/Source/DMXEditor/Internal/Analytics/DMXEditorToolAnalyticsProvider.h
@@ -22,7 +22,7 @@ namespace UE::DMX
 		 * @param Name						The name of the specific event. Will be appended to the tool specific event name.
 		 * @param Attributes				The attributes recorded with this event.
 		 */
-		void RecordEvent(const FName& Name, const TArray<FAnalyticsEventAttribute>& Attributes);
+		void RecordEvent(const FName& Name, const TArray<FAnalyticsEventAttribute>& InAttributes);
 
 	private:
 		/** Records analytics for when the tool started. */
@@ -40,7 +40,7 @@ namespace UE::DMX
 		/** Name of the tool we are providing analytics for */
 		const FName ToolName;
 
-		/** Prefix for DMX telemetry events */
-		static const FString DMXEventPrefix;
+		/** Event name for DMX telemetry events */
+		static const FString DMXToolEventName;
 	};
 }
diff --git a/Engine/Plugins/VirtualProduction/DMX/DMXEngine/Source/DMXEditor/Private/Analytics/DMXEditorToolAnalyticsProvider.cpp b/Engine/Plugins/VirtualProduction/DMX/DMXEngine/Source/DMXEditor/Private/Analytics/DMXEditorToolAnalyticsProvider.cpp
--- a/Engine/Plugins/VirtualProduction/DMX/DMXEngine/Source/DMXEditor/Private/Analytics/DMXEditorToolAnalyticsProvider.cpp
+++ b/Engine/Plugins/VirtualProduction/DMX/DMXEngine/Source/DMXEditor/Private/Analytics/DMXEditorToolAnalyticsProvider.cpp
@@ -9,7 +9,7 @@
 
 namespace UE::DMX
 {
-	const FString FDMXEditorToolAnalyticsProvider::DMXEventPrefix = TEXT("Usage.DMX.");
+	const FString FDMXEditorToolAnalyticsProvider::DMXToolEventName = TEXT("Usage.DMX.ToolEvent");
 
 	FDMXEditorToolAnalyticsProvider::FDMXEditorToolAnalyticsProvider(const FName& InToolName)
 		: ToolName(InToolName)
@@ -26,15 +26,18 @@ namespace UE::DMX
 		FCoreDelegates::OnEnginePreExit.RemoveAll(this);
 	}
 
-	void FDMXEditorToolAnalyticsProvider::RecordEvent(const FName& Name, const TArray<FAnalyticsEventAttribute>& Attributes)
+	void FDMXEditorToolAnalyticsProvider::RecordEvent(const FName& Name, const TArray<FAnalyticsEventAttribute>& InAttributes)
 	{
 		if (!FEngineAnalytics::IsAvailable())
 		{
 			return;
 		}
 
-		const FString EventName = DMXEventPrefix + ToolName.ToString() + TEXT(".") + Name.ToString();
-		FEngineAnalytics::GetProvider().RecordEvent(EventName, Attributes);
+		TArray<FAnalyticsEventAttribute> Attributes;
+		Attributes.Add(FAnalyticsEventAttribute(TEXT("ToolName"), ToolName));
+		Attributes.Append(InAttributes);
+
+		FEngineAnalytics::GetProvider().RecordEvent(DMXToolEventName, Attributes);
 	}
 
 	void FDMXEditorToolAnalyticsProvider::RecordToolStarted()
@@ -46,11 +49,12 @@ namespace UE::DMX
 
 		ToolStartTimestamp = FDateTime::UtcNow();
 
-		TArray<FAnalyticsEventAttribute> Attributes;
-		Attributes.Add(FAnalyticsEventAttribute(TEXT("ToolName"), ToolName));
+		const TArray<FAnalyticsEventAttribute> Attributes =
+		{
+			FAnalyticsEventAttribute(TEXT("ToolName"), ToolName)
+		};
 
-		const FString EventName = DMXEventPrefix + ToolName.ToString() + TEXT(".ToolStarted");
-		FEngineAnalytics::GetProvider().RecordEvent(EventName, Attributes);
+		FEngineAnalytics::GetProvider().RecordEvent(DMXToolEventName, Attributes);
 	}
 
 	void FDMXEditorToolAnalyticsProvider::RecordToolEnded()
@@ -63,12 +67,13 @@ namespace UE::DMX
 		const FDateTime Now = FDateTime::UtcNow();
 		const FTimespan ToolUsageDuration = Now - ToolStartTimestamp;
 
-		TArray<FAnalyticsEventAttribute> Attributes;
-		Attributes.Add(FAnalyticsEventAttribute(TEXT("ToolName"), ToolName));
-		Attributes.Add(FAnalyticsEventAttribute(TEXT("DurationSeconds"), static_cast<float>(ToolUsageDuration.GetTotalSeconds())));
+		const TArray<FAnalyticsEventAttribute> Attributes =
+		{
+			FAnalyticsEventAttribute(TEXT("ToolName"), ToolName),
+			FAnalyticsEventAttribute(TEXT("DurationSeconds"), static_cast<float>(ToolUsageDuration.GetTotalSeconds()))
+		};
 
-		const FString EventName = DMXEventPrefix + ToolName.ToString() + TEXT(".ToolEnded");
-		FEngineAnalytics::GetProvider().RecordEvent(EventName, Attributes);
+		FEngineAnalytics::GetProvider().RecordEvent(DMXToolEventName, Attributes);
 
 		bEnded = true;
 	}
