DMX: Fix possible crash when accessing invalid Signal Format Property handles in the DMX Fixture Signal Format Customization

#rnx

diff --git a/Engine/Plugins/VirtualProduction/DMX/DMXEngine/Source/DMXEditor/Private/Customizations/DMXFixtureSignalFormatCustomization.cpp b/Engine/Plugins/VirtualProduction/DMX/DMXEngine/Source/DMXEditor/Private/Customizations/DMXFixtureSignalFormatCustomization.cpp
--- a/Engine/Plugins/VirtualProduction/DMX/DMXEngine/Source/DMXEditor/Private/Customizations/DMXFixtureSignalFormatCustomization.cpp
+++ b/Engine/Plugins/VirtualProduction/DMX/DMXEngine/Source/DMXEditor/Private/Customizations/DMXFixtureSignalFormatCustomization.cpp
@@ -51,20 +51,25 @@ void FDMXFixtureSignalFormatCustomization::CustomizeHeader(TSharedRef<IPropertyH
 
 TArray<EDMXFixtureSignalFormat> FDMXFixtureSignalFormatCustomization::GetSignalFormats() const
 {
-	TArray<void*> RawDatas;
-	PropertyHandle->AccessRawData(RawDatas);
-
-	TArray<EDMXFixtureSignalFormat> SignalFormats;
-	for (void* RawData : RawDatas)
+	if (PropertyHandle.IsValid() && PropertyHandle->IsValidHandle())
 	{
-		const EDMXFixtureSignalFormat* SignalFormatPtr = reinterpret_cast<EDMXFixtureSignalFormat*>(RawData);
-		if (SignalFormatPtr)
+		TArray<void*> RawDatas;
+		PropertyHandle->AccessRawData(RawDatas);
+
+		TArray<EDMXFixtureSignalFormat> SignalFormats;
+		for (void* RawData : RawDatas)
 		{
-			SignalFormats.AddUnique(*SignalFormatPtr);
+			const EDMXFixtureSignalFormat* SignalFormatPtr = reinterpret_cast<EDMXFixtureSignalFormat*>(RawData);
+			if (SignalFormatPtr)
+			{
+				SignalFormats.AddUnique(*SignalFormatPtr);
+			}
 		}
+
+		return SignalFormats;
 	}
 
-	return SignalFormats;
+	return {};
 }
 
 void FDMXFixtureSignalFormatCustomization::SetSignalFormats(EDMXFixtureSignalFormat NewSignalFormat) const
