DMX: Fix a crash when the Activity Monitor is refreshed without data given the new methods to reset DMX buffers

#rnx

diff --git a/Engine/Plugins/VirtualProduction/DMX/DMXEngine/Source/DMXEditor/Private/Widgets/Monitors/SDMXActivityInUniverse.cpp b/Engine/Plugins/VirtualProduction/DMX/DMXEngine/Source/DMXEditor/Private/Widgets/Monitors/SDMXActivityInUniverse.cpp
--- a/Engine/Plugins/VirtualProduction/DMX/DMXEngine/Source/DMXEditor/Private/Widgets/Monitors/SDMXActivityInUniverse.cpp
+++ b/Engine/Plugins/VirtualProduction/DMX/DMXEngine/Source/DMXEditor/Private/Widgets/Monitors/SDMXActivityInUniverse.cpp
@@ -141,7 +141,10 @@ void SDMXActivityInUniverse::UpdateChannels()
 	for (const TSharedPtr<FDMXUniverseMonitorChannelItem>& ChannelItem : Channels)
 	{
 		int32 ChannelIndex = ChannelItem->GetChannelID() - 1;
-		check(Buffer.IsValidIndex(ChannelIndex));
+		if (!Buffer.IsValidIndex(ChannelIndex))
+		{
+			return;
+		}
 
 		// Update channels that aren't of zero value or changed to zero to zero value
 		if (ChannelItem->GetValue() != Buffer[ChannelIndex])
