DMX: Fix Add Patch Menu allows to select deleted Fixture Types and related crash, small related fixes

* Fix DMX Library editor does not always update after Fixture Types and Fixture Patches changed when updating from the C++ API.
* Move Fixture Type and Fixture Patch member variables back to their old category as the nessted categories were leading to a cluttered UI

diff --git a/Engine/Plugins/VirtualProduction/DMX/DMXEngine/Source/DMXEditor/Private/DMXEditor.cpp b/Engine/Plugins/VirtualProduction/DMX/DMXEngine/Source/DMXEditor/Private/DMXEditor.cpp
--- a/Engine/Plugins/VirtualProduction/DMX/DMXEngine/Source/DMXEditor/Private/DMXEditor.cpp
+++ b/Engine/Plugins/VirtualProduction/DMX/DMXEngine/Source/DMXEditor/Private/DMXEditor.cpp
@@ -20,6 +20,7 @@
 #include "Library/DMXLibrary.h"
 #include "Misc/MessageDialog.h"
 #include "Modes/DMXEditorApplicationMode.h"
+#include "ScopedTransaction.h"
 #include "Toolbars/DMXEditorToolbar.h"
 #include "Utils.h"
 #include "Widgets/Docking/SDockTab.h"
@@ -244,9 +245,14 @@ void FDMXEditor::OnAddNewEntity(TSubclassOf<UDMXEntity> InEntityClass)
 				FDMXEntityFixturePatchConstructionParams FixturePatchConstructionParams;
 				FixturePatchConstructionParams.FixtureTypeRef = FDMXEntityFixtureTypeRef(LastAddedFixtureType);
 
+				const FScopedTransaction CreateFixtureTypeTransaction(LOCTEXT("CreateFixtureTypeTransaction", "Create DMX Fixture Type"));
+				DMXLibrary->PreEditChange(nullptr);
+
 				UDMXEntityFixturePatch* NewFixturePatch = UDMXEntityFixturePatch::CreateFixturePatchInLibrary(FixturePatchConstructionParams);
 				FixturePatchSharedData->SelectFixturePatch(NewFixturePatch);
 
+				DMXLibrary->PostEditChange();
+
 				return;
 			}
 		}
@@ -258,8 +264,13 @@ void FDMXEditor::OnAddNewEntity(TSubclassOf<UDMXEntity> InEntityClass)
 				FDMXEntityFixturePatchConstructionParams FixturePatchConstructionParams;
 				FixturePatchConstructionParams.FixtureTypeRef = FDMXEntityFixtureTypeRef(FixtureTypesInLibrary[0]);
 
+				const FScopedTransaction CreateFixtureTypeTransaction(LOCTEXT("CreateFixturePatchTransaction", "Create DMX Fixture Patch"));
+				DMXLibrary->PreEditChange(nullptr);
+
 				UDMXEntityFixturePatch* NewFixturePatch = UDMXEntityFixturePatch::CreateFixturePatchInLibrary(FixturePatchConstructionParams);
 				FixturePatchSharedData->SelectFixturePatch(NewFixturePatch);
+
+				DMXLibrary->PostEditChange();
 			}
 			else
 			{
diff --git a/Engine/Plugins/VirtualProduction/DMX/DMXEngine/Source/DMXEditor/Private/Widgets/FixturePatch/SAddFixturePatchMenu.cpp b/Engine/Plugins/VirtualProduction/DMX/DMXEngine/Source/DMXEditor/Private/Widgets/FixturePatch/SAddFixturePatchMenu.cpp
--- a/Engine/Plugins/VirtualProduction/DMX/DMXEngine/Source/DMXEditor/Private/Widgets/FixturePatch/SAddFixturePatchMenu.cpp
+++ b/Engine/Plugins/VirtualProduction/DMX/DMXEngine/Source/DMXEditor/Private/Widgets/FixturePatch/SAddFixturePatchMenu.cpp
@@ -39,14 +39,55 @@ namespace UE::DMXEditor::FixturePatchEditor
 		WeakDMXEditor = InWeakDMXEditor;
 
 		SharedData = WeakDMXEditor.IsValid() ? WeakDMXEditor.Pin()->GetFixturePatchSharedData() : nullptr;
+		if (SharedData.IsValid())
+		{
+			UDMXLibrary::GetOnEntitiesAdded().AddSP(this, &SAddFixturePatchMenu::OnEntityAddedOrRemoved);
+			UDMXLibrary::GetOnEntitiesRemoved().AddSP(this, &SAddFixturePatchMenu::OnEntityAddedOrRemoved);
+
+			Refresh();
+		}
+	}
+
+	void SAddFixturePatchMenu::RequestRefresh()
+	{
+		RequestRefreshModeComboBoxTimerHandle.Invalidate();
+
+		if (!RequestRefreshModeComboBoxTimerHandle.IsValid())
+		{
+			RequestRefreshModeComboBoxTimerHandle = GEditor->GetTimerManager()->SetTimerForNextTick(FTimerDelegate::CreateSP(this, &SAddFixturePatchMenu::Refresh));
+		}
+	}
+
+	void SAddFixturePatchMenu::Refresh()
+	{
+		RequestRefreshModeComboBoxTimerHandle.Invalidate();
+
 		UDMXLibrary* DMXLibrary = WeakDMXEditor.IsValid() ? WeakDMXEditor.Pin()->GetDMXLibrary() : nullptr;
-		if (!SharedData.IsValid() || !DMXLibrary)
+		if (!DMXLibrary)
 		{
 			return;
 		}
 
-		const TArray<UDMXEntityFixtureType*> FixtureTypes = DMXLibrary->GetEntitiesTypeCast<UDMXEntityFixtureType>();
+
+		// Mend the fixture type
+		if (!WeakFixtureType.IsValid() && WeakDMXEditor.IsValid())
+		{
+			const TArray<UDMXEntityFixtureType*> FixtureTypes = DMXLibrary->GetEntitiesTypeCast<UDMXEntityFixtureType>();
+			WeakFixtureType = FixtureTypes.IsEmpty() ? nullptr : FixtureTypes[0];
+		}
+
+
+		// Mend the active mode index
 		UDMXAddFixturePatchMenuData* MenuData = GetMutableDefault<UDMXAddFixturePatchMenuData>();
+		if (WeakFixtureType.IsValid() && !WeakFixtureType->Modes.IsValidIndex(MenuData->ActiveModeIndex))
+		{
+			MenuData->ActiveModeIndex = 0;
+			MenuData->SaveConfig();
+		}
+
+
+		// Create the combo box source and an intial selection
+		const TArray<UDMXEntityFixtureType*> FixtureTypes = DMXLibrary->GetEntitiesTypeCast<UDMXEntityFixtureType>();
 		UDMXEntityFixtureType* const* SelectedFixtureTypePtr = Algo::FindByPredicate(FixtureTypes, [MenuData](const UDMXEntityFixtureType* FixtureType)
 			{
 				return FixtureType == MenuData->SoftFixtureType;
@@ -57,6 +98,8 @@ namespace UE::DMXEditor::FixturePatchEditor
 		}
 		WeakFixtureType = SelectedFixtureTypePtr ? *SelectedFixtureTypePtr : nullptr;
 
+
+		// Rebuild the widget
 		ChildSlot
 		[
 			SNew(SBorder)
@@ -156,21 +199,6 @@ namespace UE::DMXEditor::FixturePatchEditor
 				]
 			]
 		];
-
-		RequestRefresh();
-
-		UDMXLibrary::GetOnEntitiesAdded().AddSP(this, &SAddFixturePatchMenu::OnEntityAddedOrRemoved);
-		UDMXLibrary::GetOnEntitiesRemoved().AddSP(this, &SAddFixturePatchMenu::OnEntityAddedOrRemoved);
-	}
-
-	void SAddFixturePatchMenu::RequestRefresh()
-	{
-		RequestRefreshModeComboBoxTimerHandle.Invalidate();
-
-		if (!RequestRefreshModeComboBoxTimerHandle.IsValid())
-		{
-			RequestRefreshModeComboBoxTimerHandle = GEditor->GetTimerManager()->SetTimerForNextTick(FTimerDelegate::CreateSP(this, &SAddFixturePatchMenu::ForceRefresh));
-		}
 	}
 
 	TSharedRef<SWidget> SAddFixturePatchMenu::MakeFixtureTypeSelectWidget()
@@ -333,40 +361,6 @@ namespace UE::DMXEditor::FixturePatchEditor
 			.Font(FAppStyle::GetFontStyle(TEXT("PropertyWindow.NormalFont")));
 	}
 
-	void SAddFixturePatchMenu::ForceRefresh()
-	{
-		RequestRefreshModeComboBoxTimerHandle.Invalidate();
-
-		// Mend the fixture type
-		if (!WeakFixtureType.IsValid() && WeakDMXEditor.IsValid())
-		{
-			if (UDMXLibrary* DMXLibrary = WeakDMXEditor.Pin()->GetDMXLibrary())
-			{
-				const TArray<UDMXEntityFixtureType*> FixtureTypes = DMXLibrary->GetEntitiesTypeCast<UDMXEntityFixtureType>();
-				WeakFixtureType = FixtureTypes.IsEmpty() ? nullptr : FixtureTypes[0];
-			}
-		}
-
-		// Mend the active mode index
-		UDMXAddFixturePatchMenuData* MenuData = GetMutableDefault<UDMXAddFixturePatchMenuData>();
-		if (WeakFixtureType.IsValid() && !WeakFixtureType->Modes.IsValidIndex(MenuData->ActiveModeIndex))
-		{
-			MenuData->ActiveModeIndex = 0;
-			MenuData->SaveConfig();
-		}
-
-		ModeSources.Reset();
-		if (UDMXEntityFixtureType* FixtureType = WeakFixtureType.Get())
-		{
-			for (int32 ModeIndex = 0; ModeIndex < FixtureType->Modes.Num(); ModeIndex++)
-			{
-				ModeSources.Add(MakeShared<uint32>(ModeIndex));
-			}
-		}
-
-		ModeComboBox->RefreshOptions();
-	}
-
 	void SAddFixturePatchMenu::OnEntityAddedOrRemoved(UDMXLibrary* DMXLibrary, TArray<UDMXEntity*> Entities)
 	{
 		RequestRefresh();
@@ -554,7 +548,11 @@ namespace UE::DMXEditor::FixturePatchEditor
 
 			constexpr bool bMarkLibraryDirty = false;
 			UDMXEntityFixturePatch* NewFixturePatch = UDMXEntityFixturePatch::CreateFixturePatchInLibrary(FixturePatchConstructionParams, FixtureType->Name, bMarkLibraryDirty);
-			NewFixturePatches.Add(NewFixturePatch);
+			
+			if (NewFixturePatch)
+			{
+				NewFixturePatches.Add(NewFixturePatch);
+			}
 		}
 
 		// Align
diff --git a/Engine/Plugins/VirtualProduction/DMX/DMXEngine/Source/DMXEditor/Private/Widgets/FixturePatch/SAddFixturePatchMenu.h b/Engine/Plugins/VirtualProduction/DMX/DMXEngine/Source/DMXEditor/Private/Widgets/FixturePatch/SAddFixturePatchMenu.h
--- a/Engine/Plugins/VirtualProduction/DMX/DMXEngine/Source/DMXEditor/Private/Widgets/FixturePatch/SAddFixturePatchMenu.h
+++ b/Engine/Plugins/VirtualProduction/DMX/DMXEngine/Source/DMXEditor/Private/Widgets/FixturePatch/SAddFixturePatchMenu.h
@@ -37,6 +37,9 @@ namespace UE::DMXEditor::FixturePatchEditor
 		void RequestRefresh();
 
 	private:
+		/** Refreshes the widget */
+		void Refresh();
+
 		/** Creates a widget to select a fixture type */
 		TSharedRef<SWidget> MakeFixtureTypeSelectWidget();
 
@@ -58,9 +61,6 @@ namespace UE::DMXEditor::FixturePatchEditor
 		/** Generates an entry in the mode combo box */
 		TSharedRef<SWidget> GenerateModeComboBoxEntry(const TSharedPtr<uint32> InMode) const;
 
-		/** Refreshes the widget */
-		void ForceRefresh();
-
 		/** Called when an entity was added or removed from the DMX Library */
 		void OnEntityAddedOrRemoved(UDMXLibrary* DMXLibrary, TArray<UDMXEntity*> Entities);
 
diff --git a/Engine/Plugins/VirtualProduction/DMX/DMXEngine/Source/DMXEditor/Private/Widgets/FixturePatch/SDMXFixturePatchEditor.cpp b/Engine/Plugins/VirtualProduction/DMX/DMXEngine/Source/DMXEditor/Private/Widgets/FixturePatch/SDMXFixturePatchEditor.cpp
--- a/Engine/Plugins/VirtualProduction/DMX/DMXEngine/Source/DMXEditor/Private/Widgets/FixturePatch/SDMXFixturePatchEditor.cpp
+++ b/Engine/Plugins/VirtualProduction/DMX/DMXEngine/Source/DMXEditor/Private/Widgets/FixturePatch/SDMXFixturePatchEditor.cpp
@@ -7,6 +7,7 @@
 #include "DMXEditorSettings.h"
 #include "DMXFixturePatchSharedData.h"
 #include "DMXSubsystem.h"
+#include "Engine/TimerHandle.h"
 #include "IO/DMXOutputPort.h"
 #include "IO/DMXOutputPortReference.h"
 #include "IO/DMXPortManager.h"
@@ -33,12 +34,12 @@ void SDMXFixturePatchEditor::Construct(const FArguments& InArgs)
 {
 	SDMXEntityEditor::Construct(SDMXEntityEditor::FArguments());
 
-	DMXEditorPtr = InArgs._DMXEditor;
-	if (!DMXEditorPtr.IsValid())
+	WeakDMXEditor = InArgs._DMXEditor;
+	if (!WeakDMXEditor.IsValid())
 	{
 		return;
 	}
-	FixturePatchSharedData = DMXEditorPtr.Pin()->GetFixturePatchSharedData();
+	FixturePatchSharedData = WeakDMXEditor.Pin()->GetFixturePatchSharedData();
 
 	SetCanTick(false);
 
@@ -60,7 +61,7 @@ void SDMXFixturePatchEditor::Construct(const FArguments& InArgs)
 		+ SSplitter::Slot()	
 		.Value(LeftSideWidth)
 		[
-			SAssignNew(FixturePatchList, SDMXFixturePatchList, DMXEditorPtr)
+			SAssignNew(FixturePatchList, SDMXFixturePatchList, WeakDMXEditor)
 		]
 
 		// Right, Fixture Patcher and Details
@@ -75,7 +76,7 @@ void SDMXFixturePatchEditor::Construct(const FArguments& InArgs)
 			.Value(.618f)
 			[
 				SAssignNew(FixturePatcher, SDMXFixturePatcher)
-				.DMXEditor(DMXEditorPtr)
+				.DMXEditor(WeakDMXEditor)
 			]
 	
 			+SSplitter::Slot()
@@ -89,8 +90,12 @@ void SDMXFixturePatchEditor::Construct(const FArguments& InArgs)
 	// Adopt the selection
 	OnFixturePatchesSelected();
 
-	// Bind to selection changes
+	// Listen to selection changes
 	FixturePatchSharedData->OnFixturePatchSelectionChanged.AddSP(this, &SDMXFixturePatchEditor::OnFixturePatchesSelected);
+
+	// Listen to Fixture Patch and Fixture Type changes
+	UDMXEntityFixtureType::GetOnFixtureTypeChanged().AddSP(this, &SDMXFixturePatchEditor::OnFixtureTypeChanged);
+	UDMXEntityFixturePatch::GetOnFixturePatchChanged().AddSP(this, &SDMXFixturePatchEditor::OnFixturePatchChanged);
 }
 
 FReply SDMXFixturePatchEditor::OnKeyDown(const FGeometry& MyGeometry, const FKeyEvent& InKeyEvent)
@@ -142,6 +147,21 @@ TArray<UDMXEntity*> SDMXFixturePatchEditor::GetSelectedEntities() const
 	return SelectedEntities;
 }
 
+TSharedRef<IDetailsView> SDMXFixturePatchEditor::GenerateFixturePatchDetailsView() const
+{
+	FPropertyEditorModule& PropertyEditorModule = FModuleManager::GetModuleChecked<FPropertyEditorModule>("PropertyEditor");
+
+	FDetailsViewArgs DetailsViewArgs;
+	DetailsViewArgs.bAllowSearch = true;
+	DetailsViewArgs.NameAreaSettings = FDetailsViewArgs::HideNameArea;
+	DetailsViewArgs.bHideSelectionTip = true;
+
+	TSharedRef<IDetailsView> DetailsView = PropertyEditorModule.CreateDetailView(DetailsViewArgs);
+	DetailsView->RegisterInstancedCustomPropertyLayout(UDMXEntityFixturePatch::StaticClass(), FOnGetDetailCustomizationInstance::CreateStatic(&FDMXEntityFixturePatchDetails::MakeInstance, WeakDMXEditor));
+
+	return DetailsView;
+}
+
 void SDMXFixturePatchEditor::SelectUniverse(int32 UniverseID)
 {
 	if (!ensureMsgf(UniverseID >= 0 && UniverseID <= DMX_MAX_UNIVERSE, TEXT("Invalid Universe when trying to select Universe %i."), UniverseID))
@@ -154,31 +174,61 @@ void SDMXFixturePatchEditor::SelectUniverse(int32 UniverseID)
 
 void SDMXFixturePatchEditor::OnFixturePatchesSelected()
 {
-	TArray<TWeakObjectPtr<UDMXEntityFixturePatch>> SelectedFixturePatches = FixturePatchSharedData->GetSelectedFixturePatches();
-	TArray<UObject*> SelectedObjects;
-	for (TWeakObjectPtr<UDMXEntityFixturePatch> WeakSelectedFixturePatch : SelectedFixturePatches)
+	RequestRefreshFixturePatchDetailsView();
+}
+
+void SDMXFixturePatchEditor::OnFixtureTypeChanged(const UDMXEntityFixtureType* ChangedFixtureType)
+{
+	RequestRefreshFixturePatchDetailsView();
+}
+
+void SDMXFixturePatchEditor::OnFixturePatchChanged(const UDMXEntityFixturePatch* ChangedFixturePatch)
+{
+	RequestRefreshFixturePatchDetailsView();
+}
+
+void SDMXFixturePatchEditor::RequestRefreshFixturePatchDetailsView()
+{
+	if (!RefreshFixturePatchDetailsViewTimerHandle.IsValid())
 	{
-		if (UDMXEntity* SelectedObject = WeakSelectedFixturePatch.Get())
-		{
-			SelectedObjects.Add(SelectedObject);
-		}
+		RefreshFixturePatchDetailsViewTimerHandle = GEditor->GetTimerManager()->SetTimerForNextTick(FTimerDelegate::CreateSP(this, &SDMXFixturePatchEditor::RefreshFixturePatchDetailsView));
 	}
-	FixturePatchDetailsView->SetObjects(SelectedObjects);
 }
 
-TSharedRef<IDetailsView> SDMXFixturePatchEditor::GenerateFixturePatchDetailsView() const
+void SDMXFixturePatchEditor::RefreshFixturePatchDetailsView()
 {
-	FPropertyEditorModule& PropertyEditorModule = FModuleManager::GetModuleChecked<FPropertyEditorModule>("PropertyEditor");
-
-	FDetailsViewArgs DetailsViewArgs;
-	DetailsViewArgs.bAllowSearch = true;
-	DetailsViewArgs.NameAreaSettings = FDetailsViewArgs::HideNameArea;
-	DetailsViewArgs.bHideSelectionTip = true;
+	const TSharedPtr<FDMXEditor> DMXEditor = WeakDMXEditor.Pin();
+	UDMXLibrary* DMXLibrary = DMXEditor.IsValid() ? DMXEditor->GetDMXLibrary() : nullptr;
+	
+	if (DMXLibrary && FixturePatchDetailsView.IsValid())
+	{
+		RefreshFixturePatchDetailsViewTimerHandle.Invalidate();
 
-	TSharedRef<IDetailsView> DetailsView = PropertyEditorModule.CreateDetailView(DetailsViewArgs);
-	DetailsView->RegisterInstancedCustomPropertyLayout(UDMXEntityFixturePatch::StaticClass(), FOnGetDetailCustomizationInstance::CreateStatic(&FDMXEntityFixturePatchDetails::MakeInstance, DMXEditorPtr));
+		const TArray<TWeakObjectPtr<UDMXEntityFixturePatch>> SelectedFixturePatches = FixturePatchSharedData->GetSelectedFixturePatches();
+		
+		// Try to make a valid selection if nothing is selected
+		if (SelectedFixturePatches.IsEmpty())
+		{		
+			const TArray<UDMXEntityFixturePatch*> FixturePatches = DMXLibrary->GetEntitiesTypeCast<UDMXEntityFixturePatch>();
+			if (!FixturePatches.IsEmpty())
+			{
+				FixturePatchSharedData->SelectFixturePatch(FixturePatches[0]);
+			}
+		}
 
-	return DetailsView;
+		TArray<UObject*> SelectedObjects;
+		Algo::TransformIf(SelectedFixturePatches, SelectedObjects,
+			[](const TWeakObjectPtr<UDMXEntityFixturePatch>& WeakFixturePatch)
+			{
+				return WeakFixturePatch.IsValid();
+			},
+			[](const TWeakObjectPtr<UDMXEntityFixturePatch>& WeakFixturePatch)
+			{
+				return WeakFixturePatch.Get();
+			});
+
+		FixturePatchDetailsView->SetObjects(SelectedObjects);
+	}
 }
 
 #undef LOCTEXT_NAMESPACE
diff --git a/Engine/Plugins/VirtualProduction/DMX/DMXEngine/Source/DMXEditor/Private/Widgets/FixturePatch/SDMXFixturePatchEditor.h b/Engine/Plugins/VirtualProduction/DMX/DMXEngine/Source/DMXEditor/Private/Widgets/FixturePatch/SDMXFixturePatchEditor.h
--- a/Engine/Plugins/VirtualProduction/DMX/DMXEngine/Source/DMXEditor/Private/Widgets/FixturePatch/SDMXFixturePatchEditor.h
+++ b/Engine/Plugins/VirtualProduction/DMX/DMXEngine/Source/DMXEditor/Private/Widgets/FixturePatch/SDMXFixturePatchEditor.h
@@ -2,21 +2,19 @@
 
 #pragma once
 
+#include "Engine/TimerHandle.h"
 #include "Widgets/SDMXEntityEditor.h"
 
-#include "CoreMinimal.h"
-
 class FDMXEditor;
 class FDMXFixturePatchSharedData;
+class IDetailsView;
 class SDMXFixturePatcher;
 class SDMXFixturePatchTree;
 class SDMXFixturePatchList;
+class SSplitter;
 class UDMXEntityFixturePatch;
-
+class UDMXEntityFixtureType;
 struct FPropertyChangedEvent;
-class IDetailsView;
-class SSplitter;
-
 
 /** Editor for Fixture Patches */
 class SDMXFixturePatchEditor final
@@ -52,14 +50,26 @@ public:
 	// ~End SDMXEntityEditorTab interface 
 
 private:
+	/** Generates a Detail View for the edited Fixture Patch */
+	TSharedRef<IDetailsView> GenerateFixturePatchDetailsView() const;
+	
 	/** Selects the patch */
 	void SelectUniverse(int32 UniverseID);
 
 	/** Called whewn Fixture Patches were selected in Fixture Patch Shared Data */
 	void OnFixturePatchesSelected();
 
-	/** Generates a Detail View for the edited Fixture Patch */
-	TSharedRef<IDetailsView> GenerateFixturePatchDetailsView() const;
+	/** Called when a Fixture Type changed */
+	void OnFixtureTypeChanged(const UDMXEntityFixtureType* ChangedFixtureType);
+
+	/** Called when a Fixture Patch changed */
+	void OnFixturePatchChanged(const UDMXEntityFixturePatch* ChangedFixturePatch);
+
+	/** Refreshes the Fixture Patch Details View on the next tick */
+	void RequestRefreshFixturePatchDetailsView();
+
+	/** Refreshes the Fixture Patch Details View */
+	void RefreshFixturePatchDetailsView();
 
 	/** List of Fixture Patches as MVR Fixtures */
 	TSharedPtr<SDMXFixturePatchList> FixturePatchList;
@@ -77,5 +87,8 @@ private:
 	TSharedPtr<FDMXFixturePatchSharedData> FixturePatchSharedData;
 
 	/** Pointer back to the DMXEditor tool that owns us */
-	TWeakPtr<FDMXEditor> DMXEditorPtr;
+	TWeakPtr<FDMXEditor> WeakDMXEditor;
+
+	/** Timer handle to refresh the fixture patch details view */
+	FTimerHandle RefreshFixturePatchDetailsViewTimerHandle;
 }; 
diff --git a/Engine/Plugins/VirtualProduction/DMX/DMXEngine/Source/DMXEditor/Public/Widgets/SDMXEntityDropdownMenu.h b/Engine/Plugins/VirtualProduction/DMX/DMXEngine/Source/DMXEditor/Public/Widgets/SDMXEntityDropdownMenu.h
--- a/Engine/Plugins/VirtualProduction/DMX/DMXEngine/Source/DMXEditor/Public/Widgets/SDMXEntityDropdownMenu.h
+++ b/Engine/Plugins/VirtualProduction/DMX/DMXEngine/Source/DMXEditor/Public/Widgets/SDMXEntityDropdownMenu.h
@@ -2,31 +2,34 @@
 
 #pragma once
 
-#include "CoreMinimal.h"
-
-#include "Library/DMXLibrary.h"
+#include "Commands/DMXEditorCommands.h"
+#include "DMXEditor.h"
+#include "DMXEditorUtils.h"
 #include "Library/DMXEntity.h"
 #include "Library/DMXEntityFixturePatch.h"
 #include "Library/DMXEntityFixtureType.h"
-#include "Commands/DMXEditorCommands.h"
-#include "DMXEditorUtils.h"
-#include "DMXEditor.h"
-
+#include "Library/DMXLibrary.h"
+#include "PropertyCustomizationHelpers.h"
+#include "ScopedTransaction.h"
 #include "SListViewSelectorDropdownMenu.h"
-#include "Widgets/Views/SListView.h"
+#include "Styling/AppStyle.h"
+#include "Templates/SubclassOf.h"
 #include "Widgets/Input/SComboBox.h"
 #include "Widgets/Input/SComboButton.h"
 #include "Widgets/Input/SSearchBox.h"
 #include "Widgets/Layout/SSpacer.h"
 #include "Widgets/SToolTip.h"
-#include "Styling/AppStyle.h"
-#include "PropertyCustomizationHelpers.h"
-#include "Templates/SubclassOf.h"
+#include "Widgets/Views/SListView.h"
 
-class UDMXLibrary;
+#if UE_ENABLE_INCLUDE_ORDER_DEPRECATED_IN_5_5
+#include "CoreMinimal.h"
+#endif // UE_ENABLE_INCLUDE_ORDER_DEPRECATED_IN_5_5
 
 class SSearchBox;
 class SToolTip;
+class UDMXLibrary;
+
+#define LOCTEXT_NAMESPACE "SDMXEntitySelectorDropdownMenu"
 
 enum class EEntityEntryType : uint8
 {
@@ -95,8 +98,6 @@ private:
 
 //////////////////////////////////////////////////////////////////////////
 
-#define LOCTEXT_NAMESPACE "SDMXEntitySelectorDropdownMenu"
-
 DECLARE_DELEGATE_OneParam(FOnEntitySelected, UDMXEntity*);
 DECLARE_DELEGATE(FOnCreateNewEntitySelected);
 
@@ -121,6 +122,7 @@ public:
 
 		/** If valid, enables the option to create a new entity. */
 		SLATE_ATTRIBUTE(TWeakPtr<FDMXEditor>, DMXEditor)
+
 		/** Needed to find existing Entities. Can be inferred from DMXEditor, if that's valid. */
 		SLATE_ATTRIBUTE(TWeakObjectPtr<UDMXLibrary>, DMXLibrary)
 
@@ -358,28 +360,31 @@ private:
 			}
 			else if (InItem->IsCreateNew())
 			{
-				if (OnCreateNewEntityDelegate.IsBound())
-				{
-					OnCreateNewEntityDelegate.ExecuteIfBound();
-				}
-
 				if (TSharedPtr<FDMXEditor> PinnedEditor = DMXEditor.Get().Pin())
 				{
 					const TSubclassOf<UDMXEntity>&& FilterClass = GetEntityFilterClass();
 
+					const FText TransactionText = FilterClass->IsChildOf(UDMXEntityFixtureType::StaticClass()) ?
+						LOCTEXT("AddFixtureTypeTransaction", "Add Fixture Type") :
+						LOCTEXT("AddFixturePatchTransaction", "Add Fixture Patch");
+
+					const FScopedTransaction AddEntityTransaction(TransactionText);
+
 					if (FilterClass->IsChildOf(UDMXEntityFixtureType::StaticClass()))
 					{
 						PinnedEditor->GetToolkitCommands()->ExecuteAction(FDMXEditorCommands::Get().AddNewEntityFixtureType.ToSharedRef());
-						PinnedEditor->InvokeEditorTabFromEntityType(UDMXEntityFixtureType::StaticClass());
-
 
+						PinnedEditor->InvokeEditorTabFromEntityType(UDMXEntityFixtureType::StaticClass());
 					}
 					else if (FilterClass->IsChildOf(UDMXEntityFixturePatch::StaticClass()))
 					{
 						PinnedEditor->GetToolkitCommands()->ExecuteAction(FDMXEditorCommands::Get().AddNewEntityFixturePatch.ToSharedRef());
+
 						PinnedEditor->InvokeEditorTabFromEntityType(UDMXEntityFixturePatch::StaticClass());
 					}
 				}
+
+				OnCreateNewEntityDelegate.ExecuteIfBound();
 			}
 		}
 		else if (InItem.IsValid() && SelectInfo != ESelectInfo::OnMouseClick)
diff --git a/Engine/Plugins/VirtualProduction/DMX/DMXEngine/Source/DMXRuntime/Public/Library/DMXEntityFixturePatch.h b/Engine/Plugins/VirtualProduction/DMX/DMXEngine/Source/DMXRuntime/Public/Library/DMXEntityFixturePatch.h
--- a/Engine/Plugins/VirtualProduction/DMX/DMXEngine/Source/DMXRuntime/Public/Library/DMXEntityFixturePatch.h
+++ b/Engine/Plugins/VirtualProduction/DMX/DMXEngine/Source/DMXRuntime/Public/Library/DMXEntityFixturePatch.h
@@ -255,15 +255,15 @@ protected:
 	TObjectPtr<UDMXEntityFixtureType> ParentFixtureTypeTemplate;
 
 	/** The Index of the Mode in the Fixture Type the Patch uses */
-	UPROPERTY(EditAnywhere, BlueprintReadOnly, meta = (DisplayName = "Active Mode Index"), Category = "DMX|Fixture Patch")
+	UPROPERTY(EditAnywhere, BlueprintReadOnly, meta = (DisplayName = "Active Mode Index"), Category = "Fixture Patch")
 	int32 ActiveMode;
 
 	/** The MVR Fixture UUID */
-	UPROPERTY(BlueprintReadOnly, VisibleAnywhere, meta = (DisplayName = "MVR Fixture UUID"), Category = "DMX|Fixture Patch")
+	UPROPERTY(BlueprintReadOnly, VisibleAnywhere, meta = (DisplayName = "MVR Fixture UUID"), Category = "Fixture Patch")
 	FGuid MVRFixtureUUID;
 
 	/** The Fixture ID. Note, fixture patch much like some lighting consoles only supports numerical fixture IDs, and not a string as per MVR specs. */
-	UPROPERTY(BlueprintReadOnly, VisibleAnywhere, Category = "DMX|Fixture Patch")
+	UPROPERTY(BlueprintReadOnly, VisibleAnywhere, Category = "Fixture Patch")
 	int32 FixtureID = 0;
 
 	/** Delegate broadcast when a Fixture Patch changed */
@@ -271,19 +271,19 @@ protected:
 
 public:
 	/** Custom tags for filtering patches  */
-	UPROPERTY(EditAnywhere, BlueprintReadOnly, Category = "DMX|Fixture Patch")
+	UPROPERTY(EditAnywhere, BlueprintReadOnly, Category = "Fixture Patch")
 	TArray<FName> CustomTags;
 
 #if WITH_EDITORONLY_DATA
 	/** Color when displayed in the fixture patch editor */
-	UPROPERTY(EditAnywhere, BlueprintReadOnly, Category = "DMX|Fixture Patch")
+	UPROPERTY(EditAnywhere, BlueprintReadOnly, Category = "Fixture Patch")
 	FLinearColor EditorColor;
 
 	/** 
 	 * If true, the patch receives dmx and raises the OnFixturePatchReceivedDMX event in editor. 
 	 * NOTE: If 'All Fixture Patches receive DMX in editor' is set to true in Project Settings -> Plugins -> DMX, this setting here is ignored.
 	 */
-	UPROPERTY(EditAnywhere, BlueprintReadOnly, Category = "DMX|Fixture Patch")
+	UPROPERTY(EditAnywhere, BlueprintReadOnly, Category = "Fixture Patch")
 	bool bReceiveDMXInEditor;
 #endif
 
diff --git a/Engine/Plugins/VirtualProduction/DMX/DMXEngine/Source/DMXRuntime/Public/Library/DMXEntityFixtureType.h b/Engine/Plugins/VirtualProduction/DMX/DMXEngine/Source/DMXRuntime/Public/Library/DMXEntityFixtureType.h
--- a/Engine/Plugins/VirtualProduction/DMX/DMXEngine/Source/DMXRuntime/Public/Library/DMXEntityFixtureType.h
+++ b/Engine/Plugins/VirtualProduction/DMX/DMXEngine/Source/DMXRuntime/Public/Library/DMXEntityFixtureType.h
@@ -401,17 +401,17 @@ public:
 	static FDMXOnFixtureTypeChangedDelegate& GetOnFixtureTypeChanged();
 
 	/** The Category of the Fixture, useful for Filtering */
-	UPROPERTY(EditAnywhere, BlueprintReadOnly, Category = "DMX|Fixture Type", meta = (DisplayName = "DMX Category"))
+	UPROPERTY(EditAnywhere, BlueprintReadOnly, Category = "Fixture Type", meta = (DisplayName = "DMX Category"))
 	FDMXFixtureCategory DMXCategory;
 
-	UPROPERTY(EditAnywhere, BlueprintReadOnly, Category = "DMX|Fixture Type")
+	UPROPERTY(EditAnywhere, BlueprintReadOnly, Category = "Fixture Type")
 	TArray<FDMXFixtureMode> Modes;
 
 	/** 
 	 * Modulators applied right before a patch of this type is received. 
 	 * NOTE: Modulators only affect the patch's normalized values! Untouched values are still available when accesing raw values. 
 	 */
-	UPROPERTY(EditAnywhere, Instanced, Category = "DMX|Fixture Type", meta = (DisplayPriority = "50"))
+	UPROPERTY(EditAnywhere, Instanced, Category = "Fixture Type", meta = (DisplayPriority = "50"))
 	TArray<TObjectPtr<UDMXModulator>> InputModulators;
 
 private:
@@ -590,12 +590,12 @@ public:
 	static float BytesToNormalizedValue(EDMXFixtureSignalFormat InSignalFormat, bool bUseLSB, const uint8* InBytes);
 
 	/** The GDTF that initializes this Fixture Type. When changed, reinitializes with data from the GDTF. */
-	UPROPERTY(EditAnywhere, BlueprintReadOnly, Category = "DMX|Fixture Type", meta = (DisplayName = "GDTF Source"))
+	UPROPERTY(EditAnywhere, BlueprintReadOnly, Category = "Fixture Type", meta = (DisplayName = "GDTF Source"))
 	TSoftObjectPtr<UDMXImportGDTF> GDTFSource;
 
 #if WITH_EDITORONLY_DATA
 	/** If checked, generates a new GDTF instead of exporting the imported GDTF. This adopts changes in editor but in most cases will result in data loss and is not recommended. */
-	UPROPERTY(EditAnywhere, AdvancedDisplay, Category = "DMX|Fixture Type", meta = (DisplayName = "Generate GDTF (not recommended)"))
+	UPROPERTY(EditAnywhere, AdvancedDisplay, Category = "Fixture Type", meta = (DisplayName = "Generate GDTF (not recommended)"))
 	bool bExportGeneratedGDTF = false;
 
 	/** If true only shows latest GDTF mode revisions in editor */
