DMX: Deprecate obsolete and unused DMXEngine types

#rnx

diff --git a/Engine/Plugins/VirtualProduction/DMX/DMXEngine/Source/DMXRuntime/Public/DMXTypes.h b/Engine/Plugins/VirtualProduction/DMX/DMXEngine/Source/DMXRuntime/Public/DMXTypes.h
--- a/Engine/Plugins/VirtualProduction/DMX/DMXEngine/Source/DMXRuntime/Public/DMXTypes.h
+++ b/Engine/Plugins/VirtualProduction/DMX/DMXEngine/Source/DMXRuntime/Public/DMXTypes.h
@@ -76,8 +76,8 @@ struct TStructOpsTypeTraits<FDMXByteArray64>
 	};
 };
 
-
-USTRUCT(BlueprintType)
+struct UE_DEPRECATED(5.5, "This stuct had no specific use and is now deprecated.") FDMXRequestBase;
+USTRUCT(BlueprintType, meta = (Deprecated = "Deprecated 5.5. This stuct had no specific use and is now deprecated."))
 struct FDMXRequestBase
 {
 	GENERATED_BODY()
@@ -89,7 +89,9 @@ public:
 
 };
 
-USTRUCT(BlueprintType)
+struct UE_DEPRECATED(5.5, "This stuct had no specific use and is now deprecated.") FDMXRequest;
+PRAGMA_DISABLE_DEPRECATION_WARNINGS
+USTRUCT(BlueprintType, meta = (Deprecated = "Deprecated 5.5. This stuct had no specific use and is now deprecated."))
 struct FDMXRequest : public FDMXRequestBase
 {
 	GENERATED_BODY()
@@ -100,8 +102,11 @@ public:
 	TSubclassOf<UDMXLibrary> DMXLibrary;
 
 };
+PRAGMA_ENABLE_DEPRECATION_WARNINGS
 
-USTRUCT(BlueprintType)
+struct UE_DEPRECATED(5.5, "This stuct had no specific use and is now deprecated.") FDMXRawArtNetRequest;
+PRAGMA_DISABLE_DEPRECATION_WARNINGS
+USTRUCT(BlueprintType, meta = (Deprecated = "Deprecated 5.5. This stuct had no specific use and is now deprecated."))
 struct FDMXRawArtNetRequest : public FDMXRequestBase
 {
 	GENERATED_BODY()
@@ -121,8 +126,11 @@ public:
 	int32 Address = 1;
 
 };
+PRAGMA_ENABLE_DEPRECATION_WARNINGS
 
-USTRUCT(BlueprintType)
+struct UE_DEPRECATED(5.5, "This stuct had no specific use and is now deprecated.") FDMXRawSACN;
+PRAGMA_DISABLE_DEPRECATION_WARNINGS
+USTRUCT(BlueprintType, meta = (Deprecated = "Deprecated 5.5. This stuct had no specific use and is now deprecated."))
 struct FDMXRawSACN : public FDMXRequestBase
 {
 	GENERATED_BODY()
@@ -135,3 +143,4 @@ public:
 	UPROPERTY(EditAnywhere, BlueprintReadWrite, Category = "DMX|RawRequest", meta = (ClampMin = 0, ClampMax = 512, UIMin = 0, UIMax = 512))
 	int32 Address = 0;
 };
+PRAGMA_ENABLE_DEPRECATION_WARNINGS
diff --git a/Engine/Plugins/VirtualProduction/DMX/DMXEngine/Source/DMXRuntime/Public/Library/DMXEntityReference.h b/Engine/Plugins/VirtualProduction/DMX/DMXEngine/Source/DMXRuntime/Public/Library/DMXEntityReference.h
--- a/Engine/Plugins/VirtualProduction/DMX/DMXEngine/Source/DMXRuntime/Public/Library/DMXEntityReference.h
+++ b/Engine/Plugins/VirtualProduction/DMX/DMXEngine/Source/DMXRuntime/Public/Library/DMXEntityReference.h
@@ -90,7 +90,8 @@ private:
  * Represents a Controller from a DMX Library.
  * Used to store a reference to a Controller outside the DMX Library
  */
-USTRUCT(BlueprintType, meta = (DisplayName = "DMX Controller Ref"))
+struct 	UE_DEPRECATED(5.5, "Controllers are no longer used and deprecated in favor of the DMX Port System.") FDMXEntityControllerRef;
+USTRUCT(BlueprintType, meta = (Deprecated = "Controllers are no longer used and deprecated in favor of the DMX Port System."))
 struct DMXRUNTIME_API FDMXEntityControllerRef
 	: public FDMXEntityReference
 {
@@ -148,8 +149,10 @@ class DMXRUNTIME_API UDMXEntityReferenceConversions
 	GENERATED_BODY()
 
 public:
-	UFUNCTION(BlueprintPure, meta = (DisplayName = "ToController (ControllerReference)", CompactNodeTitle = "->", BlueprintAutocast), Category = "Utilities|DMX")
+	PRAGMA_DISABLE_DEPRECATION_WARNINGS // FDMXEntityControllerRef used in args is deprecated. Leave it to that to raise deprecated warnings in C++.
+	UFUNCTION(BlueprintPure, meta = (DeprecatedFunction, DeprecationMessage = "Controllers are no longer used and deprecated in favor of the DMX Port System.", DisplayName = "ToController (ControllerReference)", CompactNodeTitle = "->", BlueprintAutocast), Category = "Utilities|DMX")
 	static UDMXEntityController* Conv_ControllerRefToObj(const FDMXEntityControllerRef& InControllerRef);
+	PRAGMA_DISABLE_DEPRECATION_WARNINGS
 
 	UFUNCTION(BlueprintPure, meta = (DisplayName = "ToFixtureType (FixtureTypeReference)", CompactNodeTitle = "->", BlueprintAutocast), Category = "Utilities|DMX")
 	static UDMXEntityFixtureType* Conv_FixtureTypeRefToObj(const FDMXEntityFixtureTypeRef& InFixtureTypeRef);
@@ -165,4 +168,4 @@ public:
 
 	UFUNCTION(BlueprintPure, meta = (DisplayName = "ToFixturePatchRef (FixturePatch)", CompactNodeTitle = "->", BlueprintAutocast), Category = "Utilities|DMX")
 	static FDMXEntityFixturePatchRef Conv_FixturePatchObjToRef(UDMXEntityFixturePatch* InFixturePatch);
-};
\ No newline at end of file
+};
