DMX: Fix PRAGMA_DISABLE_DEPRECATION_WARNINGS macro appears a second time without matching PRAGMA_ENABLE_DEPRECATION_WARNINGS macro

#rnx

diff --git a/Engine/Plugins/VirtualProduction/DMX/DMXEngine/Source/DMXRuntime/Private/Library/DMXEntityReference.cpp b/Engine/Plugins/VirtualProduction/DMX/DMXEngine/Source/DMXRuntime/Private/Library/DMXEntityReference.cpp
--- a/Engine/Plugins/VirtualProduction/DMX/DMXEngine/Source/DMXRuntime/Private/Library/DMXEntityReference.cpp
+++ b/Engine/Plugins/VirtualProduction/DMX/DMXEngine/Source/DMXRuntime/Private/Library/DMXEntityReference.cpp
@@ -78,10 +78,12 @@ FDMXEntityControllerRef::FDMXEntityControllerRef(UDMXEntityController* InControl
 	: FDMXEntityReference(InController)
 {}
 
+PRAGMA_DISABLE_DEPRECATION_WARNINGS
 UDMXEntityController* FDMXEntityControllerRef::GetController() const
 {
 	return Cast<UDMXEntityController>(GetEntity());
 }
+PRAGMA_ENABLE_DEPRECATION_WARNINGS
 
 FDMXEntityFixtureTypeRef::FDMXEntityFixtureTypeRef()
 {
@@ -112,11 +114,12 @@ UDMXEntityFixturePatch* FDMXEntityFixturePatchRef::GetFixturePatch() const
 }
 
 //~ Type conversions extension for Entity Reference structs
-
+PRAGMA_DISABLE_DEPRECATION_WARNINGS
 UDMXEntityController* UDMXEntityReferenceConversions::Conv_ControllerRefToObj(const FDMXEntityControllerRef& InControllerRef)
 {
 	return InControllerRef.GetController();
 }
+PRAGMA_ENABLE_DEPRECATION_WARNINGS
 
 UDMXEntityFixtureType* UDMXEntityReferenceConversions::Conv_FixtureTypeRefToObj(const FDMXEntityFixtureTypeRef& InFixtureTypeRef)
 {
@@ -128,10 +131,12 @@ UDMXEntityFixturePatch* UDMXEntityReferenceConversions::Conv_FixturePatchRefToOb
 	return InFixturePatchRef.GetFixturePatch();
 }
 
+PRAGMA_DISABLE_DEPRECATION_WARNINGS
 FDMXEntityControllerRef UDMXEntityReferenceConversions::Conv_ControllerObjToRef(UDMXEntityController* InController)
 {
 	return FDMXEntityControllerRef(InController);
 }
+PRAGMA_ENABLE_DEPRECATION_WARNINGS
 
 FDMXEntityFixtureTypeRef UDMXEntityReferenceConversions::Conv_FixtureTypeObjToRef(UDMXEntityFixtureType* InFixtureType)
 {
diff --git a/Engine/Plugins/VirtualProduction/DMX/DMXEngine/Source/DMXRuntime/Public/Library/DMXEntityReference.h b/Engine/Plugins/VirtualProduction/DMX/DMXEngine/Source/DMXRuntime/Public/Library/DMXEntityReference.h
--- a/Engine/Plugins/VirtualProduction/DMX/DMXEngine/Source/DMXRuntime/Public/Library/DMXEntityReference.h
+++ b/Engine/Plugins/VirtualProduction/DMX/DMXEngine/Source/DMXRuntime/Public/Library/DMXEntityReference.h
@@ -152,7 +152,7 @@ public:
 	PRAGMA_DISABLE_DEPRECATION_WARNINGS // FDMXEntityControllerRef used in args is deprecated. Leave it to that to raise deprecated warnings in C++.
 	UFUNCTION(BlueprintPure, meta = (DeprecatedFunction, DeprecationMessage = "Controllers are no longer used and deprecated in favor of the DMX Port System.", DisplayName = "ToController (ControllerReference)", CompactNodeTitle = "->", BlueprintAutocast), Category = "Utilities|DMX")
 	static UDMXEntityController* Conv_ControllerRefToObj(const FDMXEntityControllerRef& InControllerRef);
-	PRAGMA_DISABLE_DEPRECATION_WARNINGS
+	PRAGMA_ENABLE_DEPRECATION_WARNINGS
 
 	UFUNCTION(BlueprintPure, meta = (DisplayName = "ToFixtureType (FixtureTypeReference)", CompactNodeTitle = "->", BlueprintAutocast), Category = "Utilities|DMX")
 	static UDMXEntityFixtureType* Conv_FixtureTypeRefToObj(const FDMXEntityFixtureTypeRef& InFixtureTypeRef);
@@ -160,8 +160,10 @@ public:
 	UFUNCTION(BlueprintPure, meta = (DisplayName = "ToFixturePatch (FixturePatchReference)", CompactNodeTitle = "->", BlueprintAutocast), Category = "Utilities|DMX")
 	static UDMXEntityFixturePatch* Conv_FixturePatchRefToObj(const FDMXEntityFixturePatchRef& InFixturePatchRef);
 
-	UFUNCTION(BlueprintPure, meta = (DisplayName = "ToControllerRef (Controller)", CompactNodeTitle = "->", BlueprintAutocast), Category = "Utilities|DMX")
+	PRAGMA_DISABLE_DEPRECATION_WARNINGS // FDMXEntityControllerRef used as ret val is deprecated. Leave it to that to raise deprecated warnings in C++.
+	UFUNCTION(BlueprintPure, meta = (DeprecatedFunction, DeprecationMessage = "Controllers are no longer used and deprecated in favor of the DMX Port System.", DisplayName = "ToControllerRef (Controller)", CompactNodeTitle = "->", BlueprintAutocast), Category = "Utilities|DMX")
 	static FDMXEntityControllerRef Conv_ControllerObjToRef(UDMXEntityController* InController);
+	PRAGMA_ENABLE_DEPRECATION_WARNINGS
 
 	UFUNCTION(BlueprintPure, meta = (DisplayName = "ToFixtureTypeRef (FixtureType)", CompactNodeTitle = "->", BlueprintAutocast), Category = "Utilities|DMX")
 	static FDMXEntityFixtureTypeRef Conv_FixtureTypeObjToRef(UDMXEntityFixtureType* InFixtureType);
