DMX: Fully deprecate UDMXEntityController and UDMXEntityUniverseManaged

#rnx

diff --git a/Engine/Plugins/VirtualProduction/DMX/DMXEngine/Source/DMXRuntime/Public/Library/DMXEntityController.h b/Engine/Plugins/VirtualProduction/DMX/DMXEngine/Source/DMXRuntime/Public/Library/DMXEntityController.h
--- a/Engine/Plugins/VirtualProduction/DMX/DMXEngine/Source/DMXRuntime/Public/Library/DMXEntityController.h
+++ b/Engine/Plugins/VirtualProduction/DMX/DMXEngine/Source/DMXRuntime/Public/Library/DMXEntityController.h
@@ -10,8 +10,7 @@
 struct FDMXBuffer;
 
 
-
-// DEPRECATED 4.27, can't be flagged as such to retain upgrade path, some nodes could not be compiled anymore. All members are deprecated.
+class UE_DEPRECATED(4.27, "DMXEntityUniverseManaged and DMXEntityController are deprecated in favor of the DMX Port System. Please refer to FDMXPortManager instead.") UDMXEntityUniverseManaged;
 UCLASS()
 class DMXRUNTIME_API UDMXEntityUniverseManaged
 	: public UDMXEntity
@@ -24,14 +23,15 @@ public:
 	FDMXProtocolName DeviceProtocol;
 };
 
-
-// DEPRECATED 4.27, can't be flagged as such to retain upgrade path, some nodes could not be compiled anymore. All members are deprecated.
+PRAGMA_DISABLE_DEPRECATION_WARNINGS // Leave it to the base class (UDMXEntityUniverseManaged) to raise deprecation warnings
 UCLASS()
 class DMXRUNTIME_API UDMXEntityController
 	: public UDMXEntityUniverseManaged
 {
 	GENERATED_BODY()
 
+PRAGMA_ENABLE_DEPRECATION_WARNINGS
+
 public:
 	UE_DEPRECATED(4.27, "UDMXEntityController is deprecated. Use Ports instead.")
 	UPROPERTY(Meta = (DeprecatedProperty, DeprecationMessage = "Controllers are no longer in use. Use Ports instead."))
