DMX: Fix an issue where Fixture Types could no longer be selected in the DMX Library's Add Fixture Patch Menu

#rnx

diff --git a/Engine/Plugins/VirtualProduction/DMX/DMXEngine/Source/DMXEditor/Private/Widgets/FixturePatch/SAddFixturePatchMenu.cpp b/Engine/Plugins/VirtualProduction/DMX/DMXEngine/Source/DMXEditor/Private/Widgets/FixturePatch/SAddFixturePatchMenu.cpp
--- a/Engine/Plugins/VirtualProduction/DMX/DMXEngine/Source/DMXEditor/Private/Widgets/FixturePatch/SAddFixturePatchMenu.cpp
+++ b/Engine/Plugins/VirtualProduction/DMX/DMXEngine/Source/DMXEditor/Private/Widgets/FixturePatch/SAddFixturePatchMenu.cpp
@@ -70,7 +70,7 @@ namespace UE::DMXEditor::FixturePatchEditor
 
 
 		// Mend the fixture type
-		if (!WeakFixtureType.IsValid() && WeakDMXEditor.IsValid())
+		if (!WeakFixtureType.IsValid())
 		{
 			const TArray<UDMXEntityFixtureType*> FixtureTypes = DMXLibrary->GetEntitiesTypeCast<UDMXEntityFixtureType>();
 			WeakFixtureType = FixtureTypes.IsEmpty() ? nullptr : FixtureTypes[0];
@@ -368,7 +368,12 @@ namespace UE::DMXEditor::FixturePatchEditor
 
 	void SAddFixturePatchMenu::OnFixtureTypeSelected(UDMXEntity* InSelectedFixtureType)
 	{
-		WeakFixtureType = Cast<UDMXEntityFixtureType>(InSelectedFixtureType);
+		UDMXEntityFixtureType* SelectedFixtureType = Cast<UDMXEntityFixtureType>(InSelectedFixtureType);
+		WeakFixtureType = SelectedFixtureType;
+
+		UDMXAddFixturePatchMenuData* MenuData = GetMutableDefault<UDMXAddFixturePatchMenuData>();
+		MenuData->SoftFixtureType = SelectedFixtureType;
+		MenuData->SaveConfig();
 
 		RequestRefresh();
 	}
