DMX - Fix properties that could but don't ShowInnerPropertiesOnly

diff --git a/Engine/Plugins/VirtualProduction/DMX/DMXEngine/Source/DMXEditor/Private/Customizations/DMXAutoExpandedStructCustomization.cpp b/Engine/Plugins/VirtualProduction/DMX/DMXEngine/Source/DMXEditor/Private/Customizations/DMXAutoExpandedStructCustomization.cpp
new file mode 100644
--- /dev/null
+++ b/Engine/Plugins/VirtualProduction/DMX/DMXEngine/Source/DMXEditor/Private/Customizations/DMXAutoExpandedStructCustomization.cpp
@@ -0,0 +1,47 @@
+// Copyright Epic Games, Inc. All Rights Reserved.
+
+#include "DMXAutoExpandedStructCustomization.h"
+
+#include "DetailWidgetRow.h"
+#include "IDetailPropertyRow.h"
+#include "IDetailChildrenBuilder.h"
+
+
+TSharedRef<IPropertyTypeCustomization> FDMXAutoExpandedStructCustomization::MakeInstance()
+{
+	return MakeShared<FDMXAutoExpandedStructCustomization>();
+}
+
+void FDMXAutoExpandedStructCustomization::CustomizeHeader(TSharedRef<IPropertyHandle> InPropertyHandle, FDetailWidgetRow& InHeaderRow, IPropertyTypeCustomizationUtils& CustomizationUtils)
+{
+	constexpr TCHAR ShowOnlyInnerPropertiesMetaDataName[] = TEXT("ShowOnlyInnerProperties");
+	bool bShowHeader = !InPropertyHandle->HasMetaData(ShowOnlyInnerPropertiesMetaDataName);
+	if (bShowHeader)
+	{
+		InHeaderRow
+			.ShouldAutoExpand(true)
+			.NameContent()
+			[
+				InPropertyHandle->CreatePropertyNameWidget()
+			]
+			.ValueContent()
+			[
+				InPropertyHandle->CreatePropertyValueWidget()
+			];
+	}
+}
+
+void FDMXAutoExpandedStructCustomization::CustomizeChildren(TSharedRef<IPropertyHandle> InPropertyHandle, IDetailChildrenBuilder& InChildBuilder, IPropertyTypeCustomizationUtils& CustomizationUtils)
+{
+	uint32 NumChildren;
+	if (InPropertyHandle->GetNumChildren(NumChildren) != FPropertyAccess::Success)
+	{
+		return;
+	}
+
+	for (uint32 ChildIndex = 0; ChildIndex < NumChildren; ++ChildIndex)
+	{
+		const TSharedRef<IPropertyHandle> ChildHandle = InPropertyHandle->GetChildHandle(ChildIndex).ToSharedRef();
+		InChildBuilder.AddProperty(ChildHandle).ShouldAutoExpand(true);
+	}
+}
diff --git a/Engine/Plugins/VirtualProduction/DMX/DMXEngine/Source/DMXEditor/Private/Customizations/DMXAutoExpandedStructCustomization.h b/Engine/Plugins/VirtualProduction/DMX/DMXEngine/Source/DMXEditor/Private/Customizations/DMXAutoExpandedStructCustomization.h
new file mode 100644
--- /dev/null
+++ b/Engine/Plugins/VirtualProduction/DMX/DMXEngine/Source/DMXEditor/Private/Customizations/DMXAutoExpandedStructCustomization.h
@@ -0,0 +1,21 @@
+// Copyright Epic Games, Inc. All Rights Reserved.
+
+#pragma once
+
+#include "IPropertyTypeCustomization.h"
+
+/**
+ * Customization for DMX auto expanded structs
+ */
+class FDMXAutoExpandedStructCustomization
+	: public IPropertyTypeCustomization
+{
+public:
+	/** Creates an instance of the property type customization */
+	static TSharedRef<IPropertyTypeCustomization> MakeInstance();
+
+	//~ IPropertyTypeCustomization interface begin
+	virtual void CustomizeHeader(TSharedRef<IPropertyHandle> InPropertyHandle, FDetailWidgetRow& InHeaderRow, IPropertyTypeCustomizationUtils& CustomizationUtils) override;
+	virtual void CustomizeChildren(TSharedRef<IPropertyHandle> InPropertyHandle, IDetailChildrenBuilder& InChildBuilder, IPropertyTypeCustomizationUtils& CustomizationUtils) override;
+	//~ IPropertyTypeCustomization interface end
+};
diff --git a/Engine/Plugins/VirtualProduction/DMX/DMXEngine/Source/DMXEditor/Private/Customizations/DMXEntityReferenceCustomization.cpp b/Engine/Plugins/VirtualProduction/DMX/DMXEngine/Source/DMXEditor/Private/Customizations/DMXEntityReferenceCustomization.cpp
--- a/Engine/Plugins/VirtualProduction/DMX/DMXEngine/Source/DMXEditor/Private/Customizations/DMXEntityReferenceCustomization.cpp
+++ b/Engine/Plugins/VirtualProduction/DMX/DMXEngine/Source/DMXEditor/Private/Customizations/DMXEntityReferenceCustomization.cpp
@@ -34,6 +34,7 @@ void FDMXEntityReferenceCustomization::CustomizeHeader(TSharedRef<IPropertyHandl
 			: InPropertyHandle->CreatePropertyValueWidget(false);
 
 		InHeaderRow
+			.ShouldAutoExpand(true)
 			.NameContent()
 			[
 				InPropertyHandle->CreatePropertyNameWidget()
diff --git a/Engine/Plugins/VirtualProduction/DMX/DMXEngine/Source/DMXEditor/Private/Customizations/DMXLibraryPortReferencesCustomization.cpp b/Engine/Plugins/VirtualProduction/DMX/DMXEngine/Source/DMXEditor/Private/Customizations/DMXLibraryPortReferencesCustomization.cpp
--- a/Engine/Plugins/VirtualProduction/DMX/DMXEngine/Source/DMXEditor/Private/Customizations/DMXLibraryPortReferencesCustomization.cpp
+++ b/Engine/Plugins/VirtualProduction/DMX/DMXEngine/Source/DMXEditor/Private/Customizations/DMXLibraryPortReferencesCustomization.cpp
@@ -182,6 +182,7 @@ void FDMXLibraryPortReferencesCustomization::CustomizeHeader(TSharedRef<IPropert
 	const FText DisplayToolTipOverride = LOCTEXT("HeaderToolTip", "The ports to be used with this Library.");
 
 	HeaderRow
+		.ShouldAutoExpand(true)
 		.NameContent()
 		[
 			SNullWidget::NullWidget
diff --git a/Engine/Plugins/VirtualProduction/DMX/DMXEngine/Source/DMXEditor/Private/DMXEditorModule.cpp b/Engine/Plugins/VirtualProduction/DMX/DMXEngine/Source/DMXEditor/Private/DMXEditorModule.cpp
--- a/Engine/Plugins/VirtualProduction/DMX/DMXEngine/Source/DMXEditor/Private/DMXEditorModule.cpp
+++ b/Engine/Plugins/VirtualProduction/DMX/DMXEngine/Source/DMXEditor/Private/DMXEditorModule.cpp
@@ -8,6 +8,7 @@
 #include "Commands/DMXConflictMonitorCommands.h"
 #include "Commands/DMXEditorCommands.h"
 #include "Customizations/DMXAttributeNameCustomization.h"
+#include "Customizations/DMXAutoExpandedStructCustomization.h"
 #include "Customizations/DMXEntityFixtureTypeDetails.h"
 #include "Customizations/DMXEntityReferenceCustomization.h"
 #include "Customizations/DMXFixtureCategoryCustomization.h"
@@ -406,6 +407,23 @@ void FDMXEditorModule::RegisterPropertyTypeCustomizations()
 	RegisterCustomPropertyTypeLayout(FAddAllPatchesButton::StaticStruct()->GetFName(), 
 		FOnGetPropertyTypeCustomizationInstance::CreateStatic(&FTakeRecorderDMXLibrarySourceEditorCustomization::MakeInstance)
 	);
+
+	// Customizations for auto expanded structs. The listed structs need to be auto expanded by default in most use cases.
+	const TArray<FName> StructNames
+	{
+		FDMXFixtureMode::StaticStruct()->GetFName(),
+		FDMXInputPortReference::StaticStruct()->GetFName(),
+		FDMXOutputPortReference::StaticStruct()->GetFName()
+	};
+
+	for (const FName& StructName : StructNames)
+	{
+		RegisterCustomPropertyTypeLayout
+		(
+			StructName,
+			FOnGetPropertyTypeCustomizationInstance::CreateStatic(&FDMXAutoExpandedStructCustomization::MakeInstance)
+		);
+	}
 }
 
 void FDMXEditorModule::RegisterSequencerTypes()
