DMX: DMX Libraries now spawn Actors for all Fixture Types instead of GDTF based Fixture Types only when dropped onto a level

* MVR Scene Actor groups Actors per Fixture Type instead of per GDTF
* Fixture Patch now has a Default Transform property used when its DMX Library is spawned or exported as MVR file
* The GDTF assets and Fixture Types now allow to select an Actor class to spawn
* The MVR Scene Actor now offers an option to write Default Transforms to its DMX Library

diff --git a/Engine/Plugins/VirtualProduction/DMX/DMXEngine/Source/DMXEditor/Private/Customizations/DMXEntityFixtureTypeDetails.cpp b/Engine/Plugins/VirtualProduction/DMX/DMXEngine/Source/DMXEditor/Private/Customizations/DMXEntityFixtureTypeDetails.cpp
--- a/Engine/Plugins/VirtualProduction/DMX/DMXEngine/Source/DMXEditor/Private/Customizations/DMXEntityFixtureTypeDetails.cpp
+++ b/Engine/Plugins/VirtualProduction/DMX/DMXEngine/Source/DMXEditor/Private/Customizations/DMXEntityFixtureTypeDetails.cpp
@@ -57,8 +57,14 @@ void FDMXEntityFixtureTypeDetails::OnGDTFSourceChanged()
 				if (UDMXEntityFixtureType* FixtureType = Cast<UDMXEntityFixtureType>(WeakFixtureTypeObject.Get()))
 				{
 					FixtureType->PreEditChange(nullptr);
+
+					// Generate GDTF
 					constexpr bool bUpdateFixtureTypeName = true;
 					FDMXGDTFToFixtureTypeConverter::ConvertGDTF(*FixtureType, *GDTF, bUpdateFixtureTypeName);
+
+					// Set Actor Class to Spawn
+					FixtureType->ActorClassToSpawn = GDTF->GetActorClass();
+
 					FixtureType->PostEditChange();
 				}
 			}
@@ -70,7 +76,13 @@ void FDMXEntityFixtureTypeDetails::OnGDTFSourceChanged()
 				if (UDMXEntityFixtureType* FixtureType = Cast<UDMXEntityFixtureType>(WeakFixtureTypeObject.Get()))
 				{
 					FixtureType->PreEditChange(nullptr);
+
+					// Reset GDTF
 					FixtureType->Modes.Reset();
+
+					// Reset Actor Class to Spawn
+					FixtureType->ActorClassToSpawn.Reset();
+
 					FixtureType->PostEditChange();
 				}
 			}
diff --git a/Engine/Plugins/VirtualProduction/DMX/DMXEngine/Source/DMXEditor/Private/Customizations/DMXMVRSceneActorDetails.cpp b/Engine/Plugins/VirtualProduction/DMX/DMXEngine/Source/DMXEditor/Private/Customizations/DMXMVRSceneActorDetails.cpp
--- a/Engine/Plugins/VirtualProduction/DMX/DMXEngine/Source/DMXEditor/Private/Customizations/DMXMVRSceneActorDetails.cpp
+++ b/Engine/Plugins/VirtualProduction/DMX/DMXEngine/Source/DMXEditor/Private/Customizations/DMXMVRSceneActorDetails.cpp
@@ -2,25 +2,23 @@
 
 #include "Customizations/DMXMVRSceneActorDetails.h"
 
-#include "MVR/DMXMVRSceneActor.h"
-#include "Library/DMXEntityFixturePatch.h"
-#include "Library/DMXEntityFixtureType.h"
-#include "Library/DMXImportGDTF.h"
-#include "Library/DMXLibrary.h"
-
 #include "DetailCategoryBuilder.h"
 #include "DetailLayoutBuilder.h"
 #include "DetailWidgetRow.h"
+#include "Game/DMXComponent.h"
 #include "IPropertyUtilities.h"
 #include "ISceneOutliner.h"
 #include "LevelEditor.h"
-#include "ScopedTransaction.h"
+#include "Library/DMXEntityFixturePatch.h"
+#include "Library/DMXEntityFixtureType.h"
+#include "Library/DMXLibrary.h"
 #include "Modules/ModuleManager.h"
+#include "MVR/DMXMVRSceneActor.h"
+#include "ScopedTransaction.h"
 #include "Subsystems/EditorActorSubsystem.h"
 #include "Widgets/Input/SButton.h"
 #include "Widgets/Layout/SWrapBox.h"
 
-
 #define LOCTEXT_NAMESPACE "DMXMVRSceneActorDetails"
 
 TSharedRef<IDetailCustomization> FDMXMVRSceneActorDetails::MakeInstance()
@@ -42,8 +40,8 @@ void FDMXMVRSceneActorDetails::CustomizeDetails(IDetailLayoutBuilder& DetailBuil
 		}
 	}
 
-	CreateRefreshMVRSceneSection(DetailBuilder);
-	CreateGDTFToActorClassSection(DetailBuilder);
+	CreateDMXLibrarySection(DetailBuilder);
+	CreateFixtureTypeToActorClassSection(DetailBuilder);
 
 	// Listen to map and actor changes
 	FEditorDelegates::MapChange.AddSP(this, &FDMXMVRSceneActorDetails::OnMapChange);
@@ -54,24 +52,47 @@ void FDMXMVRSceneActorDetails::CustomizeDetails(IDetailLayoutBuilder& DetailBuil
 	}
 }
 
-void FDMXMVRSceneActorDetails::CreateRefreshMVRSceneSection(IDetailLayoutBuilder& DetailBuilder)
+void FDMXMVRSceneActorDetails::CreateDMXLibrarySection(IDetailLayoutBuilder& DetailBuilder)
 {
 	IDetailCategoryBuilder& MVRCategory = DetailBuilder.EditCategory("MVR");
-
+	
+	// DMX Library
 	MVRCategory.AddProperty(DetailBuilder.GetProperty(ADMXMVRSceneActor::GetDMXLibraryPropertyNameChecked()));
+
+	// Write Transforms to DMX Library button
+	MVRCategory.AddCustomRow(LOCTEXT("WriteTransformsFilterText", "Write Transforms to DMX Library"))
+		[
+			SNew(SBorder)
+			.HAlign(HAlign_Fill)
+			.VAlign(VAlign_Center)
+			.Padding(8.f, 1.f, 0.f, 1.f)
+			.BorderImage(FAppStyle::GetBrush("NoBorder"))
+			[
+				SNew(SButton)
+				.Text(LOCTEXT("WriteTransformsToDMXLibraryCaption", "Write Transforms to DMX Library"))
+				.ToolTipText(LOCTEXT("WriteTransformsToDMXLibraryTooltip", "Sets the transform of the Fixture Actors as Default Transforms for the Fixture Patches.\n\nThe transforms will be used when the DMX Library is spawned in another level.\nThe transforms will be used when the DMX Library is exported as MVR."))
+				.OnClicked(this, &FDMXMVRSceneActorDetails::OnWriteTransformsToDMXLibraryClicked)
+			]
+		];
+
+	// Refresh from DMX Library button
 	MVRCategory.AddCustomRow(LOCTEXT("RefreshSceneFilterText", "Refresh from DMX Library"))
 		.WholeRowContent()
 		[			
 			SNew(SBorder)
-			.Padding(8.f, 0.f, 0.f, 0.f)
+			.HAlign(HAlign_Fill)
+			.VAlign(VAlign_Center)
+			.Padding(8.f, 1.f, 0.f, 1.f)
 			.BorderImage(FAppStyle::GetBrush("NoBorder"))
 			[
 				SNew(SButton)
-				.Text(LOCTEXT("RefreshActorsFromDMXLibraryButtonCaption", "Update Actors from DMX Library"))
+				.Text(LOCTEXT("RefreshActorsFromDMXLibraryCaption", "Refresh Actors from DMX Library"))
+				.ToolTipText(LOCTEXT("RefreshActorsFromDMXLibraryTooltip", "Updates the MVR Scene to reflect the DMX Library, possibly respwaning deleted actors and reseting to default transforms according to options."))
 				.OnClicked(this, &FDMXMVRSceneActorDetails::OnRefreshActorsFromDMXLibraryClicked)
 			]
 		];
 
+	// Respawn Deleted Actors option
 	const TSharedRef<IPropertyHandle> RespawnDeletedActorHandle = DetailBuilder.GetProperty(GET_MEMBER_NAME_CHECKED(ADMXMVRSceneActor, bRespawnDeletedActorsOnRefresh));
 	RespawnDeletedActorHandle->MarkHiddenByCustomization();
 	MVRCategory.AddCustomRow(LOCTEXT("RespawnDeletedActorsFilterText", "Respawn Deleted Actors"))
@@ -89,6 +110,7 @@ void FDMXMVRSceneActorDetails::CreateRefreshMVRSceneSection(IDetailLayoutBuilder
 			RespawnDeletedActorHandle->CreatePropertyValueWidget()
 		];
 	
+	// Reset Transforms option
 	const TSharedRef<IPropertyHandle> UpdateTransformHandle = DetailBuilder.GetProperty(GET_MEMBER_NAME_CHECKED(ADMXMVRSceneActor, bUpdateTransformsOnRefresh));
 	UpdateTransformHandle->MarkHiddenByCustomization();
 	MVRCategory.AddCustomRow(LOCTEXT("ResetTransformsFilterText", "Reset Transforms"))
@@ -107,52 +129,57 @@ void FDMXMVRSceneActorDetails::CreateRefreshMVRSceneSection(IDetailLayoutBuilder
 		];
 }
 
-void FDMXMVRSceneActorDetails::CreateGDTFToActorClassSection(IDetailLayoutBuilder& DetailBuilder)
+void FDMXMVRSceneActorDetails::CreateFixtureTypeToActorClassSection(IDetailLayoutBuilder& DetailBuilder)
 {
-	IDetailCategoryBuilder& ActorTypeForGDTFCategory = DetailBuilder.EditCategory("GDTF to Spawned Actor");
-	ActorTypeForGDTFCategory.InitiallyCollapsed(false);
+	IDetailCategoryBuilder& ActorTypeForFixtureTypeCategory = DetailBuilder.EditCategory("Fixture Type to Spawned Actor");
+	ActorTypeForFixtureTypeCategory.InitiallyCollapsed(false);
 
-	const TSharedRef<IPropertyHandle> GDTFToDefaultActorClassesHandle = DetailBuilder.GetProperty(ADMXMVRSceneActor::GetGDTFToDefaultActorClassesPropertyNameChecked());
-	GDTFToDefaultActorClassesHandle->MarkHiddenByCustomization();
+	const TSharedRef<IPropertyHandle> FixtureTypeToActorClassesHandle = DetailBuilder.GetProperty(ADMXMVRSceneActor::GetFixtureTypeToActorClassesPropertyNameChecked());
+	FixtureTypeToActorClassesHandle->MarkHiddenByCustomization();
 
-	const TSharedPtr<IPropertyHandleArray> GDTFToDefaultActorClassHandleArray = GDTFToDefaultActorClassesHandle->AsArray();
-	FSimpleDelegate GDTFToDefaultActorClassArrayChangedDelegate = FSimpleDelegate::CreateSP(this, &FDMXMVRSceneActorDetails::RequestRefresh);
-	GDTFToDefaultActorClassHandleArray->SetOnNumElementsChanged(GDTFToDefaultActorClassArrayChangedDelegate);
+	const TSharedPtr<IPropertyHandleArray> FixtureTypeToActorClassesHandleArray = FixtureTypeToActorClassesHandle->AsArray();
+	FixtureTypeToActorClassesHandleArray->SetOnNumElementsChanged(FSimpleDelegate::CreateSP(this, &FDMXMVRSceneActorDetails::RequestRefresh));
 
-	uint32 NumGDTFToDefaultActorClassElements;
-	if (!ensure(GDTFToDefaultActorClassHandleArray->GetNumElements(NumGDTFToDefaultActorClassElements) == FPropertyAccess::Success))
+	uint32 NumFixtureTypeToActorClassElements;
+	if (!ensure(FixtureTypeToActorClassesHandleArray->GetNumElements(NumFixtureTypeToActorClassElements) == FPropertyAccess::Success))
 	{
 		return;
 	}
 
-	for (uint32 GDTFToDefaultActorClassElementIndex = 0; GDTFToDefaultActorClassElementIndex < NumGDTFToDefaultActorClassElements; GDTFToDefaultActorClassElementIndex++)
+	for (uint32 FixtureTypeToActorClassElementIndex = 0; FixtureTypeToActorClassElementIndex < NumFixtureTypeToActorClassElements; FixtureTypeToActorClassElementIndex++)
 	{
-		const TSharedPtr<IPropertyHandle> GDTFToDefaultActorClassHandle = GDTFToDefaultActorClassHandleArray->GetElement(GDTFToDefaultActorClassElementIndex);
-		const TSharedPtr<IPropertyHandle> GDTFHandle = GDTFToDefaultActorClassHandle->GetChildHandle(GET_MEMBER_NAME_CHECKED(FDMXMVRSceneGDTFToActorClassPair, GDTF));
-		DefaultActorClassHandle = GDTFToDefaultActorClassHandle->GetChildHandle(GET_MEMBER_NAME_CHECKED(FDMXMVRSceneGDTFToActorClassPair, ActorClass));
-		DefaultActorClassHandle->SetOnChildPropertyValuePreChange(FSimpleDelegate::CreateSP(this, &FDMXMVRSceneActorDetails::OnPreEditChangeActorClassInGDTFToActorClasses));
-		DefaultActorClassHandle->SetOnChildPropertyValueChanged(FSimpleDelegate::CreateSP(this, &FDMXMVRSceneActorDetails::OnPostEditChangeActorClassInGDTFToActorClasses));
-
-		UObject* GDTFObject;
-		if (!GDTFHandle->GetValue(GDTFObject))
+		const TSharedPtr<IPropertyHandle> FixtureTypeToActorClassHandle = FixtureTypeToActorClassesHandleArray->GetElement(FixtureTypeToActorClassElementIndex);
+		const TSharedPtr<IPropertyHandle> FixtureTypeHandle = FixtureTypeToActorClassHandle->GetChildHandle(GET_MEMBER_NAME_CHECKED(FDMXMVRSceneFixtureTypeToActorClassPair, FixtureType));
+		TSharedPtr<IPropertyHandle> ActorClassHandle = FixtureTypeToActorClassHandle->GetChildHandle(GET_MEMBER_NAME_CHECKED(FDMXMVRSceneFixtureTypeToActorClassPair, ActorClass));
+		ActorClassHandle->SetOnChildPropertyValuePreChange(FSimpleDelegate::CreateSP(this, &FDMXMVRSceneActorDetails::OnPreEditChangeActorClassInFixtureTypeToActorClasses));
+		ActorClassHandle->SetOnChildPropertyValueChanged(FSimpleDelegate::CreateSP(this, &FDMXMVRSceneActorDetails::OnPostEditChangeActorClassInFixtureTypeToActorClasses));
+
+		UObject* FixtureTypeObject = nullptr;
+		if (!FixtureTypeHandle->GetValue(FixtureTypeObject))
 		{
 			return;
 		}
-		UDMXImportGDTF* GDTF = Cast<UDMXImportGDTF>(GDTFObject);
-		if (!GDTF)
-		{
-			continue;
-		}
 
-		if (!IsAnyActorUsingGDTF(GDTF))
+		UDMXEntityFixtureType* FixtureType = Cast<UDMXEntityFixtureType>(FixtureTypeObject);
+		if (!FixtureType)
 		{
 			continue;
 		}
 
-		ActorTypeForGDTFCategory.AddCustomRow(LOCTEXT("GDTFToDefaultActorClassFilter", "GDTF"))
+		ActorTypeForFixtureTypeCategory.AddCustomRow(LOCTEXT("FixtureTypeToActorClassFilter", "Fixture Type to Spawned Actor"))
 		.NameContent()
 		[
-			GDTFHandle->CreatePropertyValueWidget()
+			SNew(STextBlock)
+			.Font(DetailBuilder.GetDetailFont())
+			.Text_Lambda([WeakFixtureType = TWeakObjectPtr<UDMXEntityFixtureType>(FixtureType)]()
+				{
+					if (WeakFixtureType.IsValid())
+					{
+						return FText::FromString(WeakFixtureType.Get()->Name);
+					}
+					
+					return LOCTEXT("InvalidFixtureTypeName", "Invalid Fixture Type");
+				})
 		]
 		.ValueContent()
 		[
@@ -160,14 +187,14 @@ void FDMXMVRSceneActorDetails::CreateGDTFToActorClassSection(IDetailLayoutBuilde
 
 			+ SWrapBox::Slot()
 			[
-				DefaultActorClassHandle->CreatePropertyValueWidget()
+				ActorClassHandle->CreatePropertyValueWidget()
 			]
 
 			+ SWrapBox::Slot()
 			[
 				SNew(SButton)
-				.OnClicked(this, &FDMXMVRSceneActorDetails::OnSelectGDTFToActorClassGroupClicked, GDTFObject)
-				.Text(LOCTEXT("SelectGDTFGroupButtonCaption", "Select"))
+				.OnClicked(this, &FDMXMVRSceneActorDetails::OnFixtureTypeToActorClassGroupSelected, FixtureTypeObject)
+				.Text(LOCTEXT("SelectFixtureTypeGroupButtonCaption", "Select"))
 			]
 		];
 	}
@@ -175,13 +202,21 @@ void FDMXMVRSceneActorDetails::CreateGDTFToActorClassSection(IDetailLayoutBuilde
 
 FReply FDMXMVRSceneActorDetails::OnRefreshActorsFromDMXLibraryClicked()
 {
-	const TArray<TWeakObjectPtr<UObject>> SelectedObjects = PropertyUtilities->GetSelectedObjects();
+	const FScopedTransaction RefreshActorsFromDMXLibraryTransaction(LOCTEXT("RefreshActorsFromDMXLibraryTransaction", "Update MVR Scene form DMX Library"));
 
-	for (TWeakObjectPtr<UObject> SelectedObject : SelectedObjects)
+	const TArray<TWeakObjectPtr<UObject>> SelectedObjects = PropertyUtilities->GetSelectedObjects();
+	
+	AActor* FirstActor = SelectedObjects.IsEmpty() ? nullptr : Cast<AActor>(SelectedObjects[0]);
+	UWorld* World = FirstActor ? FirstActor->GetWorld() : nullptr;
+	if (World)
+	{
+		World->PreEditChange(nullptr);
+	}
+	
+	for (const TWeakObjectPtr<UObject>& SelectedObject : SelectedObjects)
 	{
 		if (ADMXMVRSceneActor* MVRSceneActor = Cast<ADMXMVRSceneActor>(SelectedObject.Get()))
 		{
-			const FScopedTransaction RefreshActorsFromDMXLibraryTransaction(LOCTEXT("RefreshActorsFromDMXLibraryTransaction", "Update MVR Scene form DMX Library"));
 			MVRSceneActor->PreEditChange(ADMXMVRSceneActor::StaticClass()->FindPropertyByName(ADMXMVRSceneActor::GetRelatedAcctorsPropertyNameChecked()));
 
 			MVRSceneActor->RefreshFromDMXLibrary();
@@ -190,29 +225,90 @@ FReply FDMXMVRSceneActorDetails::OnRefreshActorsFromDMXLibraryClicked()
 		}
 	}
 
+	if (World)
+	{
+		World->PostEditChange();
+	}
+
+	RequestRefresh();
+
 	return FReply::Handled();
 }
 
-FReply FDMXMVRSceneActorDetails::OnSelectGDTFToActorClassGroupClicked(UObject* GDTFObject)
+FReply FDMXMVRSceneActorDetails::OnWriteTransformsToDMXLibraryClicked()
 {
-	UDMXImportGDTF* GDTF = Cast<UDMXImportGDTF>(GDTFObject);
-	if (!GDTF)
+	const FScopedTransaction WriteTransformsToDMXLibraryTransaction(LOCTEXT("WriteTransformsToDMXLibraryTransaction", "Write MVR Scene Transforms to DMX Library"));
+
+	const TArray<TWeakObjectPtr<UObject>> SelectedObjects = PropertyUtilities->GetSelectedObjects();
+	for (const TWeakObjectPtr<UObject>& SelectedObject : SelectedObjects)
+	{
+		ADMXMVRSceneActor* MVRSceneActor = Cast<ADMXMVRSceneActor>(SelectedObject.Get());
+		UDMXLibrary* DMXLibrary = MVRSceneActor ? MVRSceneActor->GetDMXLibrary() : nullptr;
+
+		if (!MVRSceneActor || !DMXLibrary)
+		{
+			continue;
+		}
+
+		const TArray<TSoftObjectPtr<AActor>> SoftRelatedActors = MVRSceneActor->GetRelatedActors();
+		for (const TSoftObjectPtr<AActor>& SoftRelatedActor : SoftRelatedActors)
+		{
+			if (!SoftRelatedActor.IsValid())
+			{
+				continue;
+			}
+			AActor* RelatedActor = SoftRelatedActor.Get();
+			UDMXEntityFixturePatch* FixturePatch = GetFixturePatchFromActor(RelatedActor);
+
+			if (FixturePatch)
+			{
+				FixturePatch->PreEditChange(nullptr);
+				FixturePatch->SetDefaultTransform(RelatedActor->GetTransform());
+				FixturePatch->PostEditChange();
+			}
+		}
+	}
+
+	RequestRefresh();
+
+	return FReply::Handled();
+}
+
+UDMXEntityFixturePatch* FDMXMVRSceneActorDetails::GetFixturePatchFromActor(AActor* Actor) const
+{
+	TArray<UDMXComponent*> DMXComponents;
+	Actor->GetComponents<UDMXComponent>(DMXComponents);
+	if (!ensureAlwaysMsgf(!DMXComponents.IsEmpty(), TEXT("Cannot find DMX component for Actor '%s'. Cannot get Fixture Patch from Actor."), *Actor->GetName()))
+	{
+		return nullptr;
+	}
+	ensureAlwaysMsgf(DMXComponents.Num() == 1, TEXT("Actor '%s' has more than one DMX component. A single DMX component is required to clearly identify the fixture. Cannot get Fixture Patch from Actor."), *Actor->GetName());
+
+	UDMXEntityFixturePatch* FixturePatch = DMXComponents[0]->GetFixturePatch();
+
+	return FixturePatch;
+}
+
+FReply FDMXMVRSceneActorDetails::OnFixtureTypeToActorClassGroupSelected(UObject* FixtureTypeObject)
+{
+	UDMXEntityFixtureType* FixtureType = Cast<UDMXEntityFixtureType>(FixtureTypeObject);
+	if (!FixtureType)
 	{
 		return FReply::Unhandled();
 	}
 
 	const TArray<TWeakObjectPtr<UObject>> SelectedObjects = PropertyUtilities->GetSelectedObjects();
 
-	for (TWeakObjectPtr<UObject> SelectedObject : SelectedObjects)
+	for (const TWeakObjectPtr<UObject>& SelectedObject : SelectedObjects)
 	{
 		if (ADMXMVRSceneActor* MVRSceneActor = Cast<ADMXMVRSceneActor>(SelectedObject.Get()))
 		{
-			const TArray<AActor*> ActorsForThisGDTF = MVRSceneActor->GetActorsSpawnedForGDTF(GDTF);
+			const TArray<AActor*> ActorsForThisFixtureType = MVRSceneActor->GetActorsSpawnedForFixtureType(FixtureType);
 			UEditorActorSubsystem* EditorActorSubsystem = GEditor->GetEditorSubsystem<UEditorActorSubsystem>();
 
 			if (EditorActorSubsystem)
 			{
-				EditorActorSubsystem->SetSelectedLevelActors(ActorsForThisGDTF);
+				EditorActorSubsystem->SetSelectedLevelActors(ActorsForThisFixtureType);
 			}
 		}
 	}
@@ -248,18 +344,18 @@ void FDMXMVRSceneActorDetails::OnActorDeleted(AActor* DeletedActor)
 	RequestRefresh();
 }
 
-void FDMXMVRSceneActorDetails::OnPreEditChangeActorClassInGDTFToActorClasses()
+void FDMXMVRSceneActorDetails::OnPreEditChangeActorClassInFixtureTypeToActorClasses()
 {
 	for (TWeakObjectPtr<ADMXMVRSceneActor> WeakMVRSceneActor : OuterSceneActors)
 	{
 		if (ADMXMVRSceneActor* MVRSceneActor = WeakMVRSceneActor.Get())
 		{
-			MVRSceneActor->PreEditChange(FDMXMVRSceneGDTFToActorClassPair::StaticStruct()->FindPropertyByName(GET_MEMBER_NAME_CHECKED(FDMXMVRSceneGDTFToActorClassPair, ActorClass)));
+			MVRSceneActor->PreEditChange(FDMXMVRSceneFixtureTypeToActorClassPair::StaticStruct()->FindPropertyByName(GET_MEMBER_NAME_CHECKED(FDMXMVRSceneFixtureTypeToActorClassPair, ActorClass)));
 		}
 	}
 }
 
-void FDMXMVRSceneActorDetails::OnPostEditChangeActorClassInGDTFToActorClasses()
+void FDMXMVRSceneActorDetails::OnPostEditChangeActorClassInFixtureTypeToActorClasses()
 {
 	for (TWeakObjectPtr<ADMXMVRSceneActor> WeakMVRSceneActor : OuterSceneActors)
 	{
@@ -270,19 +366,6 @@ void FDMXMVRSceneActorDetails::OnPostEditChangeActorClassInGDTFToActorClasses()
 	}
 }
 
-bool FDMXMVRSceneActorDetails::IsAnyActorUsingGDTF(const UDMXImportGDTF* GDTF) const
-{
-	const TArray<TWeakObjectPtr<UObject>> SelectedObjects = PropertyUtilities->GetSelectedObjects();
-	for (TWeakObjectPtr<UObject> SelectedObject : SelectedObjects)
-	{
-		if (ADMXMVRSceneActor* MVRSceneActor = Cast<ADMXMVRSceneActor>(SelectedObject.Get()))
-		{
-			return !MVRSceneActor->GetActorsSpawnedForGDTF(GDTF).IsEmpty();
-		}
-	}
-	return false;
-}
-
 void FDMXMVRSceneActorDetails::RequestRefresh()
 {
 	PropertyUtilities->RequestRefresh();
diff --git a/Engine/Plugins/VirtualProduction/DMX/DMXEngine/Source/DMXEditor/Private/Customizations/DMXMVRSceneActorDetails.h b/Engine/Plugins/VirtualProduction/DMX/DMXEngine/Source/DMXEditor/Private/Customizations/DMXMVRSceneActorDetails.h
--- a/Engine/Plugins/VirtualProduction/DMX/DMXEngine/Source/DMXEditor/Private/Customizations/DMXMVRSceneActorDetails.h
+++ b/Engine/Plugins/VirtualProduction/DMX/DMXEngine/Source/DMXEditor/Private/Customizations/DMXMVRSceneActorDetails.h
@@ -2,19 +2,15 @@
 
 #pragma once
 
-#include "CoreMinimal.h"
 #include "IDetailCustomization.h"
 #include "Input/Reply.h"
 
 class ADMXMVRSceneActor;
+class IPropertyHandle;
+class IPropertyUtilities;
 class UDMXEntityFixturePatch;
 class UDMXEntityFixtureType;
-class UDMXImportGDTF;
-
 enum class ECheckBoxState : uint8;
-class IPropertyHandle;
-class IPropertyUtilities;
-
 
 /** Details customization for the 'FixtureType FunctionProperties' details view */
 class FDMXMVRSceneActorDetails
@@ -30,17 +26,23 @@ protected:
 	//~ End IDetailCustomization interface
 
 private:
-	/** Creates the section where the user can refresh the MVR scene from a changed DMX Library */
-	void CreateRefreshMVRSceneSection(IDetailLayoutBuilder& DetailBuilder);
+	/** Creates the DMX Library related section */
+	void CreateDMXLibrarySection(IDetailLayoutBuilder& DetailBuilder);
+
+	/** Creates the section where the user can select an actor class for each Fixture Type in the MVR Scene */
+	void CreateFixtureTypeToActorClassSection(IDetailLayoutBuilder& DetailBuilder);
 
-	/** Creates the section where the user can select an actor class for each GDTF in the MVR Scene */
-	void CreateGDTFToActorClassSection(IDetailLayoutBuilder& DetailBuilder);
+	/** Returns the Fixture Patch from an Actor, or nullptr if the Actor has no Fixture Patch set */
+	UDMXEntityFixturePatch* GetFixturePatchFromActor(AActor* Actor) const;
 
 	/** Called when the Refresh Actors from DMX Library button was clicked */
 	FReply OnRefreshActorsFromDMXLibraryClicked();
 
-	/** Called when a GDTF to Actor Class Group was clicked */
-	FReply OnSelectGDTFToActorClassGroupClicked(UObject* GDTFObject);
+	/** Called when the Write Transforms To DMX Library button was clicked */
+	FReply OnWriteTransformsToDMXLibraryClicked();
+
+	/** Called when a Fixture Type to Actor Class Group was selected */
+	FReply OnFixtureTypeToActorClassGroupSelected(UObject* FixtureTypeObject);
 
 	/** Called when a Fixture Patch changed */
 	void OnFixturePatchChanged(const UDMXEntityFixturePatch* FixturePatch);
@@ -54,21 +56,15 @@ private:
 	/** Called when an actor got deleted in editor */
 	void OnActorDeleted(AActor* DeletedActor);
 
-	/** Called before an Actor Class changed in the GDTFToActorClasses member of the Actor */
-	void OnPreEditChangeActorClassInGDTFToActorClasses();
+	/** Called before an Actor Class changed in the FixtureTypeToActorClasses member of the Actor */
+	void OnPreEditChangeActorClassInFixtureTypeToActorClasses();
 
-	/** Called after an Actor Class changed in the GDTFToActorClasses member of the Actor */
-	void OnPostEditChangeActorClassInGDTFToActorClasses();
-
-	/** Returns true if any actor in the current level makes use of the specified GDTF */
-	bool IsAnyActorUsingGDTF(const UDMXImportGDTF* GDTF) const;
+	/** Called after an Actor Class changed in the FixtureTypeToActorClasses member of the Actor */
+	void OnPostEditChangeActorClassInFixtureTypeToActorClasses();
 
 	/** Requests this Details Customization to refresh */
 	void RequestRefresh();
 
-	/** Handle to the Actor Class Property in the GDTFToDefaultActorClasses struct */
-	TSharedPtr<IPropertyHandle> DefaultActorClassHandle;
-
 	/** The Actors being customized in this Detais Customization */
 	TArray<TWeakObjectPtr<ADMXMVRSceneActor>> OuterSceneActors;
 
diff --git a/Engine/Plugins/VirtualProduction/DMX/DMXEngine/Source/DMXEditor/Private/Factories/DMXLibraryFromMVRImporter.cpp b/Engine/Plugins/VirtualProduction/DMX/DMXEngine/Source/DMXEditor/Private/Factories/DMXLibraryFromMVRImporter.cpp
--- a/Engine/Plugins/VirtualProduction/DMX/DMXEngine/Source/DMXEditor/Private/Factories/DMXLibraryFromMVRImporter.cpp
+++ b/Engine/Plugins/VirtualProduction/DMX/DMXEngine/Source/DMXEditor/Private/Factories/DMXLibraryFromMVRImporter.cpp
@@ -332,6 +332,7 @@ void FDMXLibraryFromMVRImporter::InitializeDMXLibrary(UDMXLibrary* DMXLibrary, c
 				FixturePatchConstructionParams.UniverseID = FixtureNode->GetUniverseID();
 				FixturePatchConstructionParams.StartingAddress = FixtureNode->GetStartingChannel();
 				FixturePatchConstructionParams.MVRFixtureUUID = FixtureNode->UUID;
+				FixturePatchConstructionParams.DefaultTransform = FixtureNode->GetTransformAbsolute();
 
 				UDMXEntityFixturePatch* FixturePatch = UDMXEntityFixturePatch::CreateFixturePatchInLibrary(FixturePatchConstructionParams, FixtureNode->Name);
 				FixturePatch->GenerateFixtureID(DesiredFixtureID);
diff --git a/Engine/Plugins/VirtualProduction/DMX/DMXEngine/Source/DMXEditor/Private/Widgets/FixturePatch/SDMXFixturePatchList.cpp b/Engine/Plugins/VirtualProduction/DMX/DMXEngine/Source/DMXEditor/Private/Widgets/FixturePatch/SDMXFixturePatchList.cpp
--- a/Engine/Plugins/VirtualProduction/DMX/DMXEngine/Source/DMXEditor/Private/Widgets/FixturePatch/SDMXFixturePatchList.cpp
+++ b/Engine/Plugins/VirtualProduction/DMX/DMXEngine/Source/DMXEditor/Private/Widgets/FixturePatch/SDMXFixturePatchList.cpp
@@ -109,6 +109,7 @@ namespace UE::DMX::SDMXFixturePatchList::Private
 		ConstructionParams.ActiveMode = FixturePatchToDuplicate->GetActiveModeIndex();
 		ConstructionParams.UniverseID = FixturePatchToDuplicate->GetUniverseID();
 		ConstructionParams.StartingAddress = FixturePatchToDuplicate->GetStartingChannel();
+		ConstructionParams.DefaultTransform = FixturePatchToDuplicate->GetDefaultTransform();
 
 		constexpr bool bMarkLibraryDirty = false;
 		UDMXEntityFixturePatch* NewFixturePatch = UDMXEntityFixturePatch::CreateFixturePatchInLibrary(ConstructionParams, FixturePatchToDuplicate->Name, bMarkLibraryDirty);
diff --git a/Engine/Plugins/VirtualProduction/DMX/DMXEngine/Source/DMXRuntime/Private/DMXRuntimeMainStreamObjectVersion.h b/Engine/Plugins/VirtualProduction/DMX/DMXEngine/Source/DMXRuntime/Private/DMXRuntimeMainStreamObjectVersion.h
--- a/Engine/Plugins/VirtualProduction/DMX/DMXEngine/Source/DMXRuntime/Private/DMXRuntimeMainStreamObjectVersion.h
+++ b/Engine/Plugins/VirtualProduction/DMX/DMXEngine/Source/DMXRuntime/Private/DMXRuntimeMainStreamObjectVersion.h
@@ -42,6 +42,12 @@ struct FDMXRuntimeMainStreamObjectVersion
 		// Upgrade Fixture Functions to make use of pysical properties where required (e.g. Pan, Tilt, Zoom)
 		DMXUpgradeFixtureFunctionsToUsePhysicalProperties,
 
+		// Upgrade Fixture Patches to have a Default Transform
+		DMXFixturePatchesHaveDefaultTransform,
+
+		// 5.5 Upgrade MVR Scene Actor to spawn Actors per Fixture Type of per GDTF
+		DMXMVRSceneActorSpawnsActorsPerFixtureType,
+
 		VersionPlusOne,
 		LatestVersion = VersionPlusOne - 1
 	};
diff --git a/Engine/Plugins/VirtualProduction/DMX/DMXEngine/Source/DMXRuntime/Private/Library/DMXEntityFixturePatch.cpp b/Engine/Plugins/VirtualProduction/DMX/DMXEngine/Source/DMXRuntime/Private/Library/DMXEntityFixturePatch.cpp
--- a/Engine/Plugins/VirtualProduction/DMX/DMXEngine/Source/DMXRuntime/Private/Library/DMXEntityFixturePatch.cpp
+++ b/Engine/Plugins/VirtualProduction/DMX/DMXEngine/Source/DMXRuntime/Private/Library/DMXEntityFixturePatch.cpp
@@ -66,7 +66,11 @@ UDMXEntityFixturePatch* UDMXEntityFixturePatch::CreateFixturePatchInLibrary(FDMX
 			NewFixturePatch->SetUniverseID(ConstructionParams.UniverseID);
 			NewFixturePatch->SetStartingChannel(ConstructionParams.StartingAddress);
 			NewFixturePatch->SetActiveModeIndex(ConstructionParams.ActiveMode);
-			
+
+#if WITH_EDITOR
+			NewFixturePatch->SetDefaultTransform(ConstructionParams.DefaultTransform);
+#endif 
+
 			if (ConstructionParams.MVRFixtureUUID.IsValid())
 			{
 				// Make sure the MVR UUID is truly unique across the DNX Library
@@ -196,30 +200,46 @@ void UDMXEntityFixturePatch::PostLoad()
 
 	if (!HasAnyFlags(RF_ClassDefaultObject | RF_ArchetypeObject))
 	{		
-		RebuildCache();
-		
-		// Upgrade Fixture ID and UUID from deprecated general scene description member in DMX Library
-		PRAGMA_DISABLE_DEPRECATION_WARNINGS
-		UDMXLibrary* DMXLibrary = GetParentLibrary();
-		UDMXMVRGeneralSceneDescription* GeneralSceneDescription = DMXLibrary ? DMXLibrary->GetLazyGeneralSceneDescription() : nullptr;
-		UDMXMVRFixtureNode* FixtureNode = GeneralSceneDescription ? GeneralSceneDescription->FindFixtureNode(MVRFixtureUUID) : nullptr;
-		if (FixtureNode)
+
+#if WITH_EDITOR	
+		// ~UE5.3, upgrade the patch to hold its fixture ID
+		if (GetLinkerCustomVersion(FDMXRuntimeMainStreamObjectVersion::GUID) < FDMXRuntimeMainStreamObjectVersion::DMXFixturePatchHasFixtureID)
 		{
-			MVRFixtureUUID = FixtureNode->UUID;
+			PRAGMA_DISABLE_DEPRECATION_WARNINGS
+			UDMXLibrary* DMXLibrary = GetParentLibrary();
+			UDMXMVRGeneralSceneDescription* GeneralSceneDescription = DMXLibrary ? DMXLibrary->GetLazyGeneralSceneDescription() : nullptr;
+			UDMXMVRFixtureNode* FixtureNode = GeneralSceneDescription ? GeneralSceneDescription->FindFixtureNode(MVRFixtureUUID) : nullptr;
+			if (FixtureNode)
+			{
+				int32 IntegralMVRFixtureID;
+				if (LexTryParseString(IntegralMVRFixtureID, *FixtureNode->FixtureID))
+				{
+					FixtureID = IntegralMVRFixtureID;
+				}
+			}
+			PRAGMA_ENABLE_DEPRECATION_WARNINGS
+		}
 
-			int32 IntegralMVRFixtureID;
-			if (LexTryParseString(IntegralMVRFixtureID, *FixtureNode->FixtureID))
+		// ~UE5.5, adopt the DefaultTransform property from the General Scene Description
+		if (GetLinkerCustomVersion(FDMXRuntimeMainStreamObjectVersion::GUID) < FDMXRuntimeMainStreamObjectVersion::DMXFixturePatchesHaveDefaultTransform)
+		{
+			UDMXLibrary* DMXLibrary = GetParentLibrary();
+			UDMXMVRGeneralSceneDescription* GeneralSceneDescription = DMXLibrary ? DMXLibrary->GetLazyGeneralSceneDescription() : nullptr;
+			UDMXMVRFixtureNode* FixtureNode = GeneralSceneDescription ? GeneralSceneDescription->FindFixtureNode(MVRFixtureUUID) : nullptr;
+			if (FixtureNode)
 			{
-				FixtureID = IntegralMVRFixtureID;
+				DefaultTransform = FixtureNode->GetTransformAbsolute();
 			}
 		}
-		PRAGMA_ENABLE_DEPRECATION_WARNINGS
+#endif
 
 		// Mend invalid fixture IDs
 		if (FixtureID < 1)
 		{
 			GenerateFixtureID();
 		}
+
+		RebuildCache();
 	}
 }
 
diff --git a/Engine/Plugins/VirtualProduction/DMX/DMXEngine/Source/DMXRuntime/Private/MVR/DMXMVRFixtureActorLibrary.cpp b/Engine/Plugins/VirtualProduction/DMX/DMXEngine/Source/DMXRuntime/Private/MVR/DMXMVRFixtureActorLibrary.cpp
--- a/Engine/Plugins/VirtualProduction/DMX/DMXEngine/Source/DMXRuntime/Private/MVR/DMXMVRFixtureActorLibrary.cpp
+++ b/Engine/Plugins/VirtualProduction/DMX/DMXEngine/Source/DMXRuntime/Private/MVR/DMXMVRFixtureActorLibrary.cpp
@@ -310,8 +310,30 @@ FDMXMVRFixtureActorLibrary::FDMXMVRFixtureActorLibrary()
 	}
 }
 
-UClass* FDMXMVRFixtureActorLibrary::FindMostAppropriateActorClassForPatch(const UDMXEntityFixturePatch* const Patch) const
+UClass* FDMXMVRFixtureActorLibrary::FindMostAppropriateActorClassForPatch(const UDMXEntityFixturePatch* Patch) const
 {
+	if (!ensureMsgf(Patch, TEXT("Invalid Fixture Patch provided when trying to find most appropriate actor class for patch. Cannot spawn Fixture Patch")))
+	{
+		return nullptr;
+	}
+
+	UDMXEntityFixtureType* FixtureType = Patch->GetFixtureType();
+	if (!FixtureType)
+	{
+		// May be null. Don't ensure, just log
+		UE_LOG(LogDMXRuntime, Warning, TEXT("Cannot spawn Actor for Fixture Patch, Fixture Patch '%s'. Fixture Patch has no valid Fixture Type."), *Patch->Name);
+		return nullptr;
+	}
+
+#if WITH_EDITOR
+	// Prefer the actor class set in the fixture type if available
+	if (!FixtureType->ActorClassToSpawn.IsNull())
+	{
+		return FixtureType->ActorClassToSpawn.LoadSynchronous();
+	}
+#endif
+
+	// Find the actor class with the most matching attributes
 	const FDMXFixtureMode* const ModePtr = Patch->GetActiveMode();
 	if (!ModePtr)
 	{
diff --git a/Engine/Plugins/VirtualProduction/DMX/DMXEngine/Source/DMXRuntime/Private/MVR/DMXMVRGeneralSceneDescription.cpp b/Engine/Plugins/VirtualProduction/DMX/DMXEngine/Source/DMXRuntime/Private/MVR/DMXMVRGeneralSceneDescription.cpp
--- a/Engine/Plugins/VirtualProduction/DMX/DMXEngine/Source/DMXRuntime/Private/MVR/DMXMVRGeneralSceneDescription.cpp
+++ b/Engine/Plugins/VirtualProduction/DMX/DMXEngine/Source/DMXRuntime/Private/MVR/DMXMVRGeneralSceneDescription.cpp
@@ -115,7 +115,7 @@ void UDMXMVRGeneralSceneDescription::WriteDMXLibrary(const UDMXLibrary& DMXLibra
 	{
 		if (FixturePatch)
 		{
-			WriteFixturePatch(*FixturePatch);
+			WriteFixturePatch(*FixturePatch, FixturePatch->GetDefaultTransform());
 		}
 	}
 	
@@ -146,13 +146,13 @@ void UDMXMVRGeneralSceneDescription::WriteDMXLibrary(const UDMXLibrary& DMXLibra
 				// Remember the MVR UUID of the first patch as multi patch UUID
 				MultiPatchUUID = FixturePatchInWorld->GetMVRFixtureUUID();
 
-				const TOptional<FTransform> OptionalTransform = WorldParams.bUseTransformsFromLevel ? Actor->GetTransform() : TOptional<FTransform>();
-				WriteFixturePatch(*FixturePatchInLibrary, OptionalTransform);
+				const FTransform Transform = WorldParams.bUseTransformsFromLevel ? Actor->GetTransform() : FixturePatchInLibrary->GetDefaultTransform();
+				WriteFixturePatch(*FixturePatchInLibrary, Transform);
 			}
 			else if (WorldParams.bCreateMultiPatchFixtures)
 			{
-				const TOptional<FTransform> OptionalTransform = WorldParams.bUseTransformsFromLevel ? Actor->GetTransform() : TOptional<FTransform>();
-				WriteFixturePatch(*FixturePatchInLibrary, OptionalTransform, MultiPatchUUID);
+				const FTransform Transform = WorldParams.bUseTransformsFromLevel ? Actor->GetTransform() : FixturePatchInLibrary->GetDefaultTransform();
+				WriteFixturePatch(*FixturePatchInLibrary, Transform, MultiPatchUUID);
 			}
 		}
 	}
@@ -230,7 +230,7 @@ TSharedPtr<FXmlFile> UDMXMVRGeneralSceneDescription::CreateXmlFile() const
 #endif // WITH_EDITOR
 
 #if WITH_EDITOR
-void UDMXMVRGeneralSceneDescription::WriteFixturePatch(const UDMXEntityFixturePatch& FixturePatch, const TOptional<FTransform>& OptionalTransform, const FGuid& MultiPatchUUID)
+void UDMXMVRGeneralSceneDescription::WriteFixturePatch(const UDMXEntityFixturePatch& FixturePatch, const FTransform& Transform, const FGuid& MultiPatchUUID)
 {
 	TArray<UDMXMVRFixtureNode*> FixtureNodes;
 	RootNode->GetFixtureNodes(FixtureNodes);
@@ -284,11 +284,7 @@ void UDMXMVRGeneralSceneDescription::WriteFixturePatch(const UDMXEntityFixturePa
 	}
 	check(MVRFixtureNode)
 
-	if(OptionalTransform.IsSet())
-	{
-		MVRFixtureNode->SetTransformAbsolute(OptionalTransform.GetValue());
-	}
-
+	MVRFixtureNode->SetTransformAbsolute(Transform);
 	MVRFixtureNode->SetUniverseID(FixturePatch.GetUniverseID());
 	MVRFixtureNode->SetStartingChannel(FixturePatch.GetStartingChannel());
 
diff --git a/Engine/Plugins/VirtualProduction/DMX/DMXEngine/Source/DMXRuntime/Private/MVR/DMXMVRSceneActor.cpp b/Engine/Plugins/VirtualProduction/DMX/DMXEngine/Source/DMXRuntime/Private/MVR/DMXMVRSceneActor.cpp
--- a/Engine/Plugins/VirtualProduction/DMX/DMXEngine/Source/DMXRuntime/Private/MVR/DMXMVRSceneActor.cpp
+++ b/Engine/Plugins/VirtualProduction/DMX/DMXEngine/Source/DMXRuntime/Private/MVR/DMXMVRSceneActor.cpp
@@ -2,22 +2,22 @@
 
 #include "MVR/DMXMVRSceneActor.h"
 
+#include "Algo/Find.h"
+#include "Algo/RemoveIf.h"
+#include "Components/SceneComponent.h"
+#include "DatasmithAssetUserData.h"
 #include "DMXMVRFixtureActorInterface.h"
 #include "DMXRuntimeLog.h"
+#include "DMXRuntimeMainStreamObjectVersion.h"
+#include "Engine/World.h"
+#include "EngineUtils.h"
 #include "Game/DMXComponent.h"
 #include "Library/DMXEntityFixturePatch.h"
 #include "Library/DMXImportGDTF.h"
 #include "Library/DMXLibrary.h"
-#include "MVR/DMXMVRFixtureActorLibrary.h"
 #include "MVR/DMXMVRAssetUserData.h"
+#include "MVR/DMXMVRFixtureActorLibrary.h"
 #include "MVR/Types/DMXMVRFixtureNode.h"
-
-#include "DatasmithAssetUserData.h"
-#include "EngineUtils.h"
-#include "Algo/Find.h"
-#include "Algo/RemoveIf.h"
-#include "Components/SceneComponent.h"
-#include "Engine/World.h"
 #include "UObject/UObjectIterator.h"
 
 #if WITH_EDITOR
@@ -28,7 +28,7 @@
 ADMXMVRSceneActor::ADMXMVRSceneActor()
 {
 #if WITH_EDITOR
-	if (HasAnyFlags(RF_ClassDefaultObject))
+	if (IsTemplate())
 	{
 		return;
 	}
@@ -51,10 +51,13 @@ ADMXMVRSceneActor::ADMXMVRSceneActor()
 	AddInstanceComponent(MVRSceneRoot);
 }
 
+PRAGMA_DISABLE_DEPRECATION_WARNINGS // Allow destruct complex deprecated members
 ADMXMVRSceneActor::~ADMXMVRSceneActor()
 {
+PRAGMA_ENABLE_DEPRECATION_WARNINGS
+
 #if WITH_EDITOR
-	if (HasAnyFlags(RF_ClassDefaultObject))
+	if (IsTemplate())
 	{
 		return;
 	}
@@ -73,6 +76,24 @@ ADMXMVRSceneActor::~ADMXMVRSceneActor()
 #endif // WITH_EDITOR
 }
 
+void ADMXMVRSceneActor::Serialize(FArchive& Ar)
+{
+	Super::Serialize(Ar);
+
+	Ar.UsingCustomVersion(FDMXRuntimeMainStreamObjectVersion::GUID);
+	
+#if WITH_EDITOR
+	if (Ar.IsLoading())
+	{
+		// Upgrade from actors per Fixture Type instead of per GDTF 
+		if (Ar.CustomVer(FDMXRuntimeMainStreamObjectVersion::GUID) < FDMXRuntimeMainStreamObjectVersion::DMXMVRSceneActorSpawnsActorsPerFixtureType)
+		{
+			UpgradeToFixtureTypeToActorClasses();
+		}
+	}
+#endif 
+}
+
 void ADMXMVRSceneActor::PostLoad()
 {
 	Super::PostLoad();
@@ -105,9 +126,9 @@ void ADMXMVRSceneActor::PreEditChange(FProperty* PropertyAboutToChange)
 {
 	Super::PreEditChange(PropertyAboutToChange);
 
-	if (PropertyAboutToChange && PropertyAboutToChange->GetFName() == GET_MEMBER_NAME_CHECKED(FDMXMVRSceneGDTFToActorClassPair, ActorClass))
+	if (PropertyAboutToChange && PropertyAboutToChange->GetFName() == GET_MEMBER_NAME_CHECKED(FDMXMVRSceneFixtureTypeToActorClassPair, ActorClass))
 	{
-		GDTFToDefaultActorClasses_PreEditChange = GDTFToDefaultActorClasses;
+		FixtureTypeToActorClasses_PreEditChange = FixtureTypeToActorClasses;
 	}
 }
 #endif // WITH_EDITOR
@@ -119,9 +140,9 @@ void ADMXMVRSceneActor::PostEditChangeProperty(FPropertyChangedEvent& PropertyCh
 
 	const FName PropertyName = PropertyChangedEvent.GetPropertyName();
 
-	if (PropertyName == GET_MEMBER_NAME_CHECKED(FDMXMVRSceneGDTFToActorClassPair, ActorClass))
+	if (PropertyName == GET_MEMBER_NAME_CHECKED(FDMXMVRSceneFixtureTypeToActorClassPair, ActorClass))
 	{
-		HandleDefaultActorClassForGDTFChanged();
+		HandleDefaultActorClassForFixtureTypeChanged();
 	}
 }
 #endif // WITH_EDITOR
@@ -230,7 +251,8 @@ void ADMXMVRSceneActor::RefreshFromDMXLibrary()
 	// Spawn newly added and if requested, previously deleted actors
 	for (UDMXEntityFixturePatch* FixturePatch : FixturePatches)
 	{
-		if (FixturePatchToSpawnedActorMap.Contains(FixturePatch))
+		if (!FixturePatch ||
+			FixturePatchToSpawnedActorMap.Contains(FixturePatch))
 		{
 			continue;
 		}
@@ -241,20 +263,28 @@ void ADMXMVRSceneActor::RefreshFromDMXLibrary()
 			continue;
 		}
 
-		UClass* ActorClass = MVRFixtureActorLibrary->FindMostAppropriateActorClassForPatch(FixturePatch);
-		if (!ActorClass)
+		const UDMXEntityFixtureType* FixtureType = FixturePatch->GetFixtureType();
+		if (!FixtureType)
 		{
 			continue;
 		}
 
-		const FGuid& MVRFixtureUUID = FixturePatch->GetMVRFixtureUUID();
-		const UDMXMVRFixtureNode* FixtureNode = GeneralSceneDescription->FindFixtureNode(MVRFixtureUUID);
-		if (!FixtureNode)
+		const FDMXMVRSceneFixtureTypeToActorClassPair* FixtureTypeToActorClassPairPtr = 
+			Algo::FindBy(FixtureTypeToActorClasses, FixtureType, &FDMXMVRSceneFixtureTypeToActorClassPair::FixtureType);
+
+		UClass* ActorClass = FixtureTypeToActorClassPairPtr ? 
+			(*FixtureTypeToActorClassPairPtr).ActorClass.LoadSynchronous() : 
+			MVRFixtureActorLibrary->FindMostAppropriateActorClassForPatch(FixturePatch);
+
+		if (!ActorClass)
 		{
 			continue;
 		}
 
-		const FTransform Transform = FixtureNode->GetTransformAbsolute();
+		const FGuid& MVRFixtureUUID = FixturePatch->GetMVRFixtureUUID();
+		const UDMXMVRFixtureNode* FixtureNode = GeneralSceneDescription->FindFixtureNode(MVRFixtureUUID);
+		const FTransform Transform = FixtureNode ? FixtureNode->GetTransformAbsolute() : FTransform::Identity;
+
 		SpawnMVRActor(ActorClass, FixturePatch, Transform);
 	}
 
@@ -277,17 +307,47 @@ void ADMXMVRSceneActor::RefreshFromDMXLibrary()
 
 			const FGuid& MVRFixtureUUID = FixturePatch->GetMVRFixtureUUID();
 			const UDMXMVRFixtureNode* FixtureNode = GeneralSceneDescription->FindFixtureNode(MVRFixtureUUID);
-			if (!FixtureNode)
-			{
-				continue;
-			}
+			const FTransform Transform = FixtureNode ? FixtureNode->GetTransformAbsolute() : FTransform::Identity;
 
-			const FTransform Transform = FixtureNode->GetTransformAbsolute();
 			RelatedActor->SetActorTransform(Transform);
 		}
 	}
 
-	UpdateGDTFToDefaultActorClasses(MVRFixtureActorLibrary);
+	UpdateFixtureTypeToDefaultActorClasses(MVRFixtureActorLibrary);
+}
+#endif // WITH_EDITOR
+
+#if WITH_EDITOR
+TArray<AActor*> ADMXMVRSceneActor::GetActorsSpawnedForFixtureType(const UDMXEntityFixtureType* FixtureType) const
+{
+	// DEPRECATED 5.5
+	TArray<AActor*> Result;
+	if (!FixtureType)
+	{
+		return Result;
+	}
+
+	for (const TSoftObjectPtr<AActor>& SoftActorPtr : RelatedActors)
+	{
+		if (!SoftActorPtr.IsValid())
+		{
+			continue;
+		}
+		AActor* Actor = SoftActorPtr.Get();
+
+		UDMXEntityFixturePatch* FixturePatch = GetFixturePatchFromActor(Actor);
+		if (!FixturePatch || !FixturePatch->GetFixtureType())
+		{
+			continue;
+		}
+
+		if (FixturePatch->GetFixtureType() == FixtureType)
+		{
+			Result.Add(Actor);
+		}
+	}
+
+	return Result;
 }
 #endif // WITH_EDITOR
 
@@ -343,46 +403,34 @@ void ADMXMVRSceneActor::EnsureMVRUUIDsForRelatedActors()
 	}
 }
 
-void ADMXMVRSceneActor::UpdateGDTFToDefaultActorClasses(const TSharedRef<FDMXMVRFixtureActorLibrary>& MVRFixtureActorLibrary)
+#if WITH_EDITOR
+void ADMXMVRSceneActor::UpdateFixtureTypeToDefaultActorClasses(const TSharedRef<FDMXMVRFixtureActorLibrary>& MVRFixtureActorLibrary)
 {
 	const TArray<UDMXEntityFixturePatch*> FixturePatches = DMXLibrary->GetEntitiesTypeCast<UDMXEntityFixturePatch>();
 	for (UDMXEntityFixturePatch* FixturePatch : FixturePatches)
 	{
-		const UDMXEntityFixtureType* FixtureType = FixturePatch->GetFixtureType();
-		if (FixtureType && !FixtureType->GDTFSource.IsNull())
+		UDMXEntityFixtureType* FixtureType = FixturePatch->GetFixtureType();
+		if (FixtureType)
 		{
-			const FDMXMVRSceneGDTFToActorClassPair* ExistingGDTFToActorClassPairPtr = Algo::FindByPredicate(GDTFToDefaultActorClasses, [FixtureType](const FDMXMVRSceneGDTFToActorClassPair& GDTFToActorClassPair)
+			const FDMXMVRSceneFixtureTypeToActorClassPair* ExistingFixtureTypeToActorClassPairPtr = Algo::FindByPredicate(FixtureTypeToActorClasses, [FixtureType](const FDMXMVRSceneFixtureTypeToActorClassPair& FixtureTypeToActorClassPair)
 				{
-					return GDTFToActorClassPair.GDTF == FixtureType->GDTFSource;
+					return FixtureTypeToActorClassPair.FixtureType == FixtureType;
 				});
-			if (ExistingGDTFToActorClassPairPtr)
+			if (ExistingFixtureTypeToActorClassPairPtr)
 			{
 				continue;
 			}
 
 			UClass* ActorClass = MVRFixtureActorLibrary->FindMostAppropriateActorClassForPatch(FixturePatch);
-			FDMXMVRSceneGDTFToActorClassPair GDTFToActorClassPair;
-			GDTFToActorClassPair.GDTF = FixtureType->GDTFSource;
-			GDTFToActorClassPair.ActorClass = ActorClass;
+			FDMXMVRSceneFixtureTypeToActorClassPair FixtureTypeToActorClassPair;
+			FixtureTypeToActorClassPair.FixtureType = FixtureType;
+			FixtureTypeToActorClassPair.ActorClass = ActorClass;
 
-			GDTFToDefaultActorClasses.Add(GDTFToActorClassPair);
+			FixtureTypeToActorClasses.Add(FixtureTypeToActorClassPair);
 		}
 	}
-
-	// Remove those GDTFs that are no longer used by any actor
-	GDTFToDefaultActorClasses.SetNum(Algo::RemoveIf(GDTFToDefaultActorClasses, [this](const FDMXMVRSceneGDTFToActorClassPair& GDTFToActorClassPair)
-		{
-			const TSoftObjectPtr<AActor>* ActorThatUsesGDTF = Algo::FindByPredicate(RelatedActors, [this, GDTFToActorClassPair](const TSoftObjectPtr<AActor>& Actor)
-				{
-					const UDMXEntityFixturePatch* FixturePatch = Actor.IsValid() ? GetFixturePatchFromActor(Actor.Get()) : nullptr;
-					const UDMXEntityFixtureType* FixtureType = FixturePatch ? FixturePatch->GetFixtureType() : nullptr;
-					const UDMXImportGDTF* GDTF = FixtureType ? FixtureType->GDTFSource.LoadSynchronous() : nullptr;
-					
-					return GDTFToActorClassPair.GDTF == GDTF;
-				});
-			return ActorThatUsesGDTF == nullptr;
-		}));
 }
+#endif // WITH_EDITOR
 
 #if WITH_EDITOR
 void ADMXMVRSceneActor::OnMapChange(uint32 MapEventFlags)
@@ -449,22 +497,22 @@ void ADMXMVRSceneActor::OnAssetPostImport(UFactory* InFactory, UObject* ActorAdd
 #endif // WITH_EDITOR
 
 #if WITH_EDITOR
-void ADMXMVRSceneActor::HandleDefaultActorClassForGDTFChanged()
+void ADMXMVRSceneActor::HandleDefaultActorClassForFixtureTypeChanged()
 {
 	// Handle element changes, but not add/remove
-	if (GDTFToDefaultActorClasses_PreEditChange.Num() != GDTFToDefaultActorClasses.Num())
+	if (FixtureTypeToActorClasses_PreEditChange.Num() != FixtureTypeToActorClasses.Num())
 	{
 		return;
 	}
 
 	int32 IndexOfChangedElement = INDEX_NONE;
-	for (const FDMXMVRSceneGDTFToActorClassPair& GDTFToDefaultActorClassPair : GDTFToDefaultActorClasses)
+	for (const FDMXMVRSceneFixtureTypeToActorClassPair& FixtureTypeToActorClassPair : FixtureTypeToActorClasses)
 	{
-		IndexOfChangedElement = GDTFToDefaultActorClasses_PreEditChange.IndexOfByPredicate([&GDTFToDefaultActorClassPair](const FDMXMVRSceneGDTFToActorClassPair& GDTFToActorClassPair)
+		IndexOfChangedElement = FixtureTypeToActorClasses_PreEditChange.IndexOfByPredicate([&FixtureTypeToActorClassPair](const FDMXMVRSceneFixtureTypeToActorClassPair& OtherFixtureTypeoActorClassPair)
 			{
 				return
-					GDTFToActorClassPair.GDTF == GDTFToDefaultActorClassPair.GDTF &&
-					GDTFToActorClassPair.ActorClass != GDTFToDefaultActorClassPair.ActorClass;
+					OtherFixtureTypeoActorClassPair.FixtureType == FixtureTypeToActorClassPair.FixtureType &&
+					OtherFixtureTypeoActorClassPair.ActorClass != FixtureTypeToActorClassPair.ActorClass;
 			});
 
 		if (IndexOfChangedElement != INDEX_NONE)
@@ -478,7 +526,7 @@ void ADMXMVRSceneActor::HandleDefaultActorClassForGDTFChanged()
 		return;
 	}
 
-	const TSubclassOf<AActor> Class = GDTFToDefaultActorClasses[IndexOfChangedElement].ActorClass.Get();
+	const TSubclassOf<AActor> Class = FixtureTypeToActorClasses[IndexOfChangedElement].ActorClass.Get();
 	if (!Class.Get())
 	{
 		return;
@@ -495,7 +543,7 @@ void ADMXMVRSceneActor::HandleDefaultActorClassForGDTFChanged()
 		UDMXEntityFixturePatch* FixturePatch = GetFixturePatchFromActor(Actor);
 		if (FixturePatch && 
 			FixturePatch->GetFixtureType() && 
-			FixturePatch->GetFixtureType()->GDTFSource == GDTFToDefaultActorClasses[IndexOfChangedElement].GDTF)
+			FixturePatch->GetFixtureType() == FixtureTypeToActorClasses[IndexOfChangedElement].FixtureType)
 		{
 			ReplaceMVRActor(Actor, Class);
 		}
@@ -517,9 +565,9 @@ AActor* ADMXMVRSceneActor::SpawnMVRActor(const TSubclassOf<AActor>&ActorClass, U
 		return nullptr;
 	}
 
-	if (!FixturePatch->GetFixtureType() || FixturePatch->GetFixtureType()->GDTFSource.IsNull())
+	if (!FixturePatch->GetFixtureType())
 	{
-		UE_LOG(LogDMXRuntime, Warning, TEXT("Cannot spawn fixture for Fixture Patch '%s'. It's Fixture Type does not source from a GDTF."), *FixturePatch->Name);
+		UE_LOG(LogDMXRuntime, Warning, TEXT("Cannot spawn fixture for Fixture Patch '%s'. Fixture Patch has no Fixture Type set."), *FixturePatch->Name);
 		return nullptr;
 	}
 
@@ -564,6 +612,8 @@ AActor* ADMXMVRSceneActor::SpawnMVRActor(const TSubclassOf<AActor>&ActorClass, U
 	SetFixturePatchOnActor(NewFixtureActor, FixturePatch);
 	RelatedActors.Add(NewFixtureActor);
 
+	DeletedMVRFixtureUUIDs.Remove(FixturePatch->GetMVRFixtureUUID());
+
 #if WITH_EDITOR
 	PostEditChange();
 	NewFixtureActor->PostEditChange();
@@ -634,3 +684,86 @@ AActor* ADMXMVRSceneActor::ReplaceMVRActor(AActor* ActorToReplace, const TSubcla
 	return nullptr;
 }
 #endif // WITH_EDITOR
+
+#if WITH_EDITOR
+void ADMXMVRSceneActor::UpgradeToFixtureTypeToActorClasses()
+{
+	PRAGMA_DISABLE_DEPRECATION_WARNINGS
+	if (!DMXLibrary)
+	{
+		return;
+	}
+
+	const TArray<UDMXEntityFixtureType*> FixtureTypes = DMXLibrary->GetEntitiesTypeCast<UDMXEntityFixtureType>();
+	
+	TArray<UDMXEntityFixtureType*> PreviouslyIgnoredFixtureTypes;
+	for (UDMXEntityFixtureType* FixtureType : FixtureTypes)
+	{
+		if (!FixtureType)
+		{
+			continue;
+		}
+
+		const FDMXMVRSceneGDTFToActorClassPair* GDTFToActorClassPairPtr  = Algo::FindByPredicate(GDTFToDefaultActorClasses_DEPRECATED,
+			[&FixtureType](const FDMXMVRSceneGDTFToActorClassPair& GDTFToActorClassPair)
+			{
+				return
+					GDTFToActorClassPair.ActorClass &&
+					!GDTFToActorClassPair.GDTF.IsNull() &&
+					GDTFToActorClassPair.GDTF == FixtureType->GDTFSource;
+			});
+
+		if (GDTFToActorClassPairPtr)
+		{
+			FDMXMVRSceneFixtureTypeToActorClassPair FixtureTypeToActorClassPair;
+			FixtureTypeToActorClassPair.FixtureType = FixtureType;
+			FixtureTypeToActorClassPair.ActorClass = GDTFToActorClassPairPtr->ActorClass;
+
+			FixtureTypeToActorClasses.Add(FixtureTypeToActorClassPair);
+		}
+		else
+		{
+			PreviouslyIgnoredFixtureTypes.Add(FixtureType);
+		}
+	}
+
+	// Early out if possible to avoid any overhead
+	if (PreviouslyIgnoredFixtureTypes.IsEmpty())
+	{
+		return;
+	}
+
+	const TSharedRef<FDMXMVRFixtureActorLibrary> MVRFixtureActorLibrary = MakeShared<FDMXMVRFixtureActorLibrary>();
+	const TArray<UDMXEntityFixturePatch*> FixturePatches = DMXLibrary->GetEntitiesTypeCast<UDMXEntityFixturePatch>();
+
+	TArray<const UDMXEntityFixtureType*> UpgradedFixtureTypes;
+	for (UDMXEntityFixturePatch* FixturePatch : FixturePatches)
+	{
+		const UDMXEntityFixtureType* FixtureType = FixturePatch ? FixturePatch->GetFixtureType() : nullptr;
+		if (!FixturePatch || !FixtureType)
+		{
+			continue;
+		}
+
+		if (PreviouslyIgnoredFixtureTypes.Contains(FixtureType))
+		{
+			// Treat previously ignored Fixture Types as deleted from the Level
+			DeletedMVRFixtureUUIDs.Add(FixturePatch->GetMVRFixtureUUID());
+
+			if (!UpgradedFixtureTypes.Contains(FixtureType))
+			{
+				// Upgrade to use Fixture Types which don't have a GDTF set
+				FDMXMVRSceneFixtureTypeToActorClassPair FixtureTypeToActorClassPair;
+				FixtureTypeToActorClassPair.FixtureType = FixtureType;
+				FixtureTypeToActorClassPair.ActorClass = MVRFixtureActorLibrary->FindMostAppropriateActorClassForPatch(FixturePatch);
+
+				FixtureTypeToActorClasses.Add(FixtureTypeToActorClassPair);
+
+				UpgradedFixtureTypes.Add(FixtureType);
+			}
+		}
+	}
+
+	PRAGMA_ENABLE_DEPRECATION_WARNINGS
+}
+#endif // WITH_EDITOR
diff --git a/Engine/Plugins/VirtualProduction/DMX/DMXEngine/Source/DMXRuntime/Public/Library/DMXEntityFixturePatch.h b/Engine/Plugins/VirtualProduction/DMX/DMXEngine/Source/DMXRuntime/Public/Library/DMXEntityFixturePatch.h
--- a/Engine/Plugins/VirtualProduction/DMX/DMXEngine/Source/DMXRuntime/Public/Library/DMXEntityFixturePatch.h
+++ b/Engine/Plugins/VirtualProduction/DMX/DMXEngine/Source/DMXRuntime/Public/Library/DMXEntityFixturePatch.h
@@ -46,6 +46,15 @@ struct DMXRUNTIME_API FDMXEntityFixturePatchConstructionParams
 	UPROPERTY(EditAnywhere, BlueprintReadWrite, Category = "Fixture Patch", meta = (DisplayName = "Starting Channel", UIMin = "1", UIMax = "512", ClampMin = "1", ClampMax = "512"))
 	int32 StartingAddress = 1;
 
+#if WITH_EDITOR
+	/**
+	 * The transform used when the DMX Library is spawned in a level.
+	 *
+	 * When the DMX Library is exported as an MVR file, this transform is used unless the 'Use Transforms from Level' export option is checked.
+	 */
+	FTransform DefaultTransform = FTransform::Identity;
+#endif 
+
 	/** 
 	 * When spawning the DMX Library as MVR Scene in Editor, each Fixture Patch has to correspond to a Fixture in the World (if it is desired to export the Scene as MVR later).
 	 * Mostly useful when importing an MVR into the DMX Library (see DMXLibraryFromMVRFactory). If left '00000000-00000000-00000000-00000000', a Unique ID will be generated for the patch. 
@@ -218,10 +227,25 @@ public:
 	bool FindFixtureID(int32& OutFixtureID) const;
 
 #if WITH_EDITOR
+	/**
+	 * Sets The transform used when the DMX Library is spawned in a level.
+	 *
+	 * When the DMX Library is exported as an MVR file, this transform is used unless the 'Use Transforms from Level' export option is checked.
+	 */
+	void SetDefaultTransform(const FTransform& NewDefaultTransform) { DefaultTransform = NewDefaultTransform; }
+
+	/**
+	 * Returns The transform used when the DMX Library is spawned in a level.
+	 *
+	 * When the DMX Library is exported as an MVR file, this transform is used unless the 'Use Transforms from Level' export option is checked.
+	 */
+	const FTransform& GetDefaultTransform() const { return DefaultTransform; }
+
 	/** Property name getters. When accessing the this way, use with care. The patch will need to update its cache after using property setters, use Pre/PostEditChanged events to achieve that. */
 	static FName GetUniverseIDPropertyNameChecked() { return GET_MEMBER_NAME_CHECKED(UDMXEntityFixturePatch, UniverseID); }
 	static FName GetParentFixtureTypeTemplatePropertyNameChecked() { return GET_MEMBER_NAME_CHECKED(UDMXEntityFixturePatch, ParentFixtureTypeTemplate); }
 	static FName GetActiveModePropertyNameChecked() { return GET_MEMBER_NAME_CHECKED(UDMXEntityFixturePatch, ActiveMode); }
+	static FName GetDefaultTransformPropertyNameChecked() { return GET_MEMBER_NAME_CHECKED(UDMXEntityFixturePatch, DefaultTransform); }
 	static FName GetMVRFixtureUUIDPropertyNameChecked() { return GET_MEMBER_NAME_CHECKED(UDMXEntityFixturePatch, MVRFixtureUUID); }
 	static FName GetFixtureIDPropertyNameChecked() { return GET_MEMBER_NAME_CHECKED(UDMXEntityFixturePatch, FixtureID); }
 	static FName GetStartingChannelPropertyNameChecked() { return GET_MEMBER_NAME_CHECKED(UDMXEntityFixturePatch, StartingChannel); }
@@ -258,13 +282,23 @@ protected:
 	UPROPERTY(EditAnywhere, BlueprintReadOnly, meta = (DisplayName = "Active Mode Index"), Category = "Fixture Patch")
 	int32 ActiveMode;
 
-	/** The MVR Fixture UUID */
-	UPROPERTY(BlueprintReadOnly, VisibleAnywhere, meta = (DisplayName = "MVR Fixture UUID"), Category = "Fixture Patch")
-	FGuid MVRFixtureUUID;
+#if WITH_EDITORONLY_DATA
+	/** 
+	 * The transform used when the DMX Library is spawned in a level.
+	 * 
+	 * When the DMX Library is exported as an MVR file, this transform is used unless the 'Use Transforms from Level' export option is checked.
+	 */
+	UPROPERTY(EditAnywhere, BlueprintReadWrite, Category = "MVR")
+	FTransform DefaultTransform = FTransform::Identity;
+#endif // WITH_EDITORONLY_DATA
 
-	/** The Fixture ID. Note, fixture patch much like some lighting consoles only supports numerical fixture IDs, and not a string as per MVR specs. */
-	UPROPERTY(BlueprintReadOnly, VisibleAnywhere, Category = "Fixture Patch")
+	/** The Fixture ID of this patch */
+	UPROPERTY(BlueprintReadOnly, VisibleAnywhere, Category = "MVR")
 	int32 FixtureID = 0;
+	
+	/** The MVR Fixture UUID */
+	UPROPERTY(BlueprintReadOnly, VisibleAnywhere, meta = (DisplayName = "MVR Fixture UUID"), Category = "MVR")
+	FGuid MVRFixtureUUID;
 
 	/** Delegate broadcast when a Fixture Patch changed */
 	static FDMXOnFixturePatchChangedDelegate OnFixturePatchChangedDelegate;
diff --git a/Engine/Plugins/VirtualProduction/DMX/DMXEngine/Source/DMXRuntime/Public/Library/DMXEntityFixtureType.h b/Engine/Plugins/VirtualProduction/DMX/DMXEngine/Source/DMXRuntime/Public/Library/DMXEntityFixtureType.h
--- a/Engine/Plugins/VirtualProduction/DMX/DMXEngine/Source/DMXRuntime/Public/Library/DMXEntityFixtureType.h
+++ b/Engine/Plugins/VirtualProduction/DMX/DMXEngine/Source/DMXRuntime/Public/Library/DMXEntityFixtureType.h
@@ -598,6 +598,17 @@ public:
 	UPROPERTY(EditAnywhere, AdvancedDisplay, Category = "Fixture Type", meta = (DisplayName = "Generate GDTF (not recommended)"))
 	bool bExportGeneratedGDTF = false;
 
+#if WITH_EDITORONLY_DATA
+	/**
+	 * The Actor Class that is spawned when the DMX Library dropped onto a Level.
+	 * Only Actors that implement the MVR Fixture Actor Interface can be used.
+	 *
+	 * Can be left blank. If so, any Actor Class with the most matching Attributes will be spawned.
+	 */
+	UPROPERTY(EditAnywhere, Category = "MVR", Meta = (MustImplement = "/Script/DMXFixtureActorInterface.DMXMVRFixtureActorInterface"))
+	TSoftClassPtr<AActor> ActorClassToSpawn;
+#endif // WITH_EDITORONLY_DATA
+
 	/** If true only shows latest GDTF mode revisions in editor */
 	UPROPERTY()
 	bool bShowOnlyLatestGDTFModeRevisions = true;
diff --git a/Engine/Plugins/VirtualProduction/DMX/DMXEngine/Source/DMXRuntime/Public/Library/DMXImportGDTF.h b/Engine/Plugins/VirtualProduction/DMX/DMXEngine/Source/DMXRuntime/Public/Library/DMXImportGDTF.h
--- a/Engine/Plugins/VirtualProduction/DMX/DMXEngine/Source/DMXRuntime/Public/Library/DMXImportGDTF.h
+++ b/Engine/Plugins/VirtualProduction/DMX/DMXEngine/Source/DMXRuntime/Public/Library/DMXImportGDTF.h
@@ -48,9 +48,26 @@ public:
 #endif
 
 	/** Returns GDTF Asset Import Data for this GDTF */
-	FORCEINLINE UDMXGDTFAssetImportData* GetGDTFAssetImportData() const { return GDTFAssetImportData; }
- 
+	UDMXGDTFAssetImportData* GetGDTFAssetImportData() const { return GDTFAssetImportData; }
+
+#if WITH_EDITOR
+	/** Returns the Actor Class to spawn when spawned from a DMX Library */
+	const TSoftClassPtr<AActor>& GetActorClass() const { return ActorClass; }
+#endif
+
 private:
+#if WITH_EDITORONLY_DATA
+	/** 
+	 * The Actor Class that corresponds to this GDTF. Only Actors that implement the MVR Fixture Actor Interface can be used. 
+	 *
+	 * In the current Version this set as the Fixture Type Actor Class, when this GDTF is set for the Fixture Type. 
+	 * 
+	 * Can be left blank. If so, any Actor Class with the most matching Attributes will be spawned. 
+	 */
+	UPROPERTY(EditAnywhere, Category = "MVR", Meta = (MustImplement = "/Script/DMXFixtureActorInterface.DMXMVRFixtureActorInterface", AllowPrivateAccess = true))
+	TSoftClassPtr<AActor> ActorClass;
+#endif // WITH_EDITORONLY_DATA
+
 	/** The Asset Import Data used to generate the GDTF asset or nullptr, if not generated from a GDTF file */
 	UPROPERTY()
 	TObjectPtr<UDMXGDTFAssetImportData> GDTFAssetImportData;
diff --git a/Engine/Plugins/VirtualProduction/DMX/DMXEngine/Source/DMXRuntime/Public/MVR/DMXMVRFixtureActorLibrary.h b/Engine/Plugins/VirtualProduction/DMX/DMXEngine/Source/DMXRuntime/Public/MVR/DMXMVRFixtureActorLibrary.h
--- a/Engine/Plugins/VirtualProduction/DMX/DMXEngine/Source/DMXRuntime/Public/MVR/DMXMVRFixtureActorLibrary.h
+++ b/Engine/Plugins/VirtualProduction/DMX/DMXEngine/Source/DMXRuntime/Public/MVR/DMXMVRFixtureActorLibrary.h
@@ -23,7 +23,7 @@ public:
 	FDMXMVRFixtureActorLibrary();
 
 	/** Returns the Actor that can represent the patch best */
-	UClass* FindMostAppropriateActorClassForPatch(const UDMXEntityFixturePatch* const Patch) const;
+	UClass* FindMostAppropriateActorClassForPatch(const UDMXEntityFixturePatch* Patch) const;
 
 protected:
 	//~ Begin FGCObject interface
diff --git a/Engine/Plugins/VirtualProduction/DMX/DMXEngine/Source/DMXRuntime/Public/MVR/DMXMVRGeneralSceneDescription.h b/Engine/Plugins/VirtualProduction/DMX/DMXEngine/Source/DMXRuntime/Public/MVR/DMXMVRGeneralSceneDescription.h
--- a/Engine/Plugins/VirtualProduction/DMX/DMXEngine/Source/DMXRuntime/Public/MVR/DMXMVRGeneralSceneDescription.h
+++ b/Engine/Plugins/VirtualProduction/DMX/DMXEngine/Source/DMXRuntime/Public/MVR/DMXMVRGeneralSceneDescription.h
@@ -95,7 +95,7 @@ private:
 	 * Optionally a MultiPatch UUID can be passed so the patch is added as a multi patch.
 	 * When a MultiPatch UUID is provided, the parent with related UUID is expected to exist already.
 	 */
-	void WriteFixturePatch(const UDMXEntityFixturePatch& FixturePatch, const TOptional<FTransform>& OptionalTransform = TOptional<FTransform>(), const FGuid& MultiPatchUUID = FGuid());
+	void WriteFixturePatch(const UDMXEntityFixturePatch& FixturePatch, const FTransform& Transform, const FGuid& MultiPatchUUID = FGuid());
 
 	/** Makes sure the node has a unique MVR UUID and Fixture ID */
 	void  SanetizeFixtureNode(UDMXMVRFixtureNode& FixtureNode);
diff --git a/Engine/Plugins/VirtualProduction/DMX/DMXEngine/Source/DMXRuntime/Public/MVR/DMXMVRSceneActor.h b/Engine/Plugins/VirtualProduction/DMX/DMXEngine/Source/DMXRuntime/Public/MVR/DMXMVRSceneActor.h
--- a/Engine/Plugins/VirtualProduction/DMX/DMXEngine/Source/DMXRuntime/Public/MVR/DMXMVRSceneActor.h
+++ b/Engine/Plugins/VirtualProduction/DMX/DMXEngine/Source/DMXRuntime/Public/MVR/DMXMVRSceneActor.h
@@ -2,31 +2,49 @@
 
 #pragma once
 
-#include "CoreMinimal.h"
+#include "Containers/Map.h"
 #include "GameFramework/Actor.h"
 
+#if UE_ENABLE_INCLUDE_ORDER_DEPRECATED_IN_5_5
+#include "CoreMinimal.h"
+#endif 
+
 #include "DMXMVRSceneActor.generated.h"
 
 class FDMXMVRFixtureActorLibrary;
+class UFactory;
 class UDMXEntityFixturePatch;
+class UDMXEntityFixtureType;
 class UDMXImportGDTF;
 class UDMXLibrary;
 class UDMXMVRSceneComponent;
 
-class UFactory;
+USTRUCT(BlueprintType)
+struct DMXRUNTIME_API FDMXMVRSceneFixtureTypeToActorClassPair
+{
+	GENERATED_BODY()
 
+	/** The Fixture Type of the Actor. */
+	UPROPERTY(VisibleAnywhere, Category = "MVR")
+	TSoftObjectPtr<UDMXEntityFixtureType> FixtureType;
 
-USTRUCT(BlueprintType)
+	/** The Actor Class that should be or was spawned. Only Actors that implement the MVR Fixture Actor Interface can be used. */
+	UPROPERTY(EditAnywhere, BlueprintReadOnly, Category = MVR, NoClear, meta = (MustImplement = "/Script/DMXFixtureActorInterface.DMXMVRFixtureActorInterface"))
+	TSoftClassPtr<AActor> ActorClass;
+};
+
+struct UE_DEPRECATED(5.5, "Instead use FDMXMVRSceneFixtureTypeToActorClassPair") FDMXMVRSceneGDTFToActorClassPair;
+USTRUCT(BlueprintType, meta = (Deprecated = "Deprecated 5.5. Instead use DMXMVRSceneFixtureTypeToActorClassPair"))
 struct DMXRUNTIME_API FDMXMVRSceneGDTFToActorClassPair
 {
 	GENERATED_BODY()
 
-	/** The GDTF of the Actor (may not correlate with its Patch after the Patch changed) */
+	/** The GDTF of the Actor. */
 	UPROPERTY(VisibleAnywhere, Category = "MVR")
 	TSoftObjectPtr<UDMXImportGDTF> GDTF;
 	
-	/** The Actor Class that should be or was spawned */
-	UPROPERTY(EditAnywhere, BlueprintReadOnly, Category = MVR, NoClear, Meta = (MustImplement = "/Script/DMXFixtureActorInterface.DMXMVRFixtureActorInterface"))
+	/** The Actor Class that should be or was spawned. Only Actors that implement the MVR Fixture Actor Interface can be used. */
+	UPROPERTY(EditAnywhere, BlueprintReadOnly, Category = MVR, NoClear, meta = (MustImplement = "/Script/DMXFixtureActorInterface.DMXMVRFixtureActorInterface"))
 	TSoftClassPtr<AActor> ActorClass;
 };
 
@@ -37,61 +55,71 @@ class DMXRUNTIME_API ADMXMVRSceneActor
 	GENERATED_BODY()
 
 public:
-	/** Constructor */
 	ADMXMVRSceneActor();
-
-	/** Destructor */
 	~ADMXMVRSceneActor();
 
 	//~ Begin AActor interface
+	virtual void Serialize(FArchive& Ar) override;
 	virtual void PostLoad() override;
 	virtual void PostRegisterAllComponents() override;
 #if WITH_EDITOR
 	virtual void PreEditChange(FProperty* PropertyAboutToChange) override;
 	virtual void PostEditChangeProperty(FPropertyChangedEvent& PropertyChangedEvent) override;
-#endif
 	//~ End AActor interface
 
-#if WITH_EDITOR
 	/** Sets the dmx library for this MVR actor. Should only be called once, further calls will have no effect and hit an ensure condition */
 	void SetDMXLibrary(UDMXLibrary* NewDMXLibrary);
 
 	/** Refreshes the scene from the DMX Library */
 	void RefreshFromDMXLibrary();
 
+	/** Returns an array of actors that are currently spawned in the scene for the specified Fixture Type */
+	TArray<AActor*> GetActorsSpawnedForFixtureType(const UDMXEntityFixtureType* FixtureType) const;
+
 	/** Returns an array of actors that are currently spawned in the scene for specified GDTF */
+	UE_DEPRECATED(5.5, "MVRSceneActor now spawns actors for all Fixture Types, not only those that have a GDTF asset. Instead now use GetActorsSpawnedForFixtureType")
 	TArray<AActor*> GetActorsSpawnedForGDTF(const UDMXImportGDTF* GDTF) const;
 #endif
+
 	/** Returns the DMX Library of this MVR Scene Actor */
 	FORCEINLINE UDMXLibrary* GetDMXLibrary() const { return DMXLibrary; }
 
 	/** Returns the MVR UUID To Related Actor Map */
 	FORCEINLINE const TArray<TSoftObjectPtr<AActor>>& GetRelatedActors() const { return RelatedActors; };
 
-#if WITH_EDITOR
-	// Property name getters
-	static FName GetDMXLibraryPropertyNameChecked() { return GET_MEMBER_NAME_CHECKED(ADMXMVRSceneActor, DMXLibrary); }
-	static FName GetRelatedAcctorsPropertyNameChecked() { return GET_MEMBER_NAME_CHECKED(ADMXMVRSceneActor, RelatedActors); }
-	static FName GetGDTFToDefaultActorClassesPropertyNameChecked() { return GET_MEMBER_NAME_CHECKED(ADMXMVRSceneActor, GDTFToDefaultActorClasses); }
-	static FName GetMVRSceneRootPropertyNameChecked() { return GET_MEMBER_NAME_CHECKED(ADMXMVRSceneActor, MVRSceneRoot); }
-#endif
-
 #if WITH_EDITORONLY_DATA
 	/** If checked, respawns Fixture Actors deleted from the MVR Scene */
-	UPROPERTY(EditAnywhere, Category = "MVR", Meta = (DisplayName = "Respawn Deleted Actors"))
+	UPROPERTY(EditAnywhere, Category = "MVR", meta = (DisplayName = "Respawn Deleted Actors"))
 	bool bRespawnDeletedActorsOnRefresh = false;
 
 	/** If checked, updates transforms from the DMX Library */
-	UPROPERTY(EditAnywhere, Category = "MVR", Meta = (DisplayName = "Update Transforms"))
+	UPROPERTY(EditAnywhere, Category = "MVR", meta = (DisplayName = "Update Transforms"))
 	bool bUpdateTransformsOnRefresh = false;
 #endif
 
+#if WITH_EDITOR
+	// Property name getters
+	static FName GetDMXLibraryPropertyNameChecked() { return GET_MEMBER_NAME_CHECKED(ADMXMVRSceneActor, DMXLibrary); }
+	static FName GetRelatedAcctorsPropertyNameChecked() { return GET_MEMBER_NAME_CHECKED(ADMXMVRSceneActor, RelatedActors); }
+	static FName GetFixtureTypeToActorClassesPropertyNameChecked() { return GET_MEMBER_NAME_CHECKED(ADMXMVRSceneActor, FixtureTypeToActorClasses); }
+	static FName GetMVRSceneRootPropertyNameChecked() { return GET_MEMBER_NAME_CHECKED(ADMXMVRSceneActor, MVRSceneRoot); }
+
+	UE_DEPRECATED(5.5, "GDTFToDefaultActorClasses is no longer used as all Fixture Types can generate GDTFs. To upgrade please refer to FixtureTypeToActorClasses.")
+	static FName GetGDTFToDefaultActorClassesPropertyNameChecked() { 
+		PRAGMA_DISABLE_DEPRECATION_WARNINGS	
+		return GET_MEMBER_NAME_CHECKED(ADMXMVRSceneActor, GDTFToDefaultActorClasses_DEPRECATED); 
+		PRAGMA_ENABLE_DEPRECATION_WARNINGS
+	}
+#endif
+
 private:
 	/** Set MVR UUIDs for related Actors */
 	void EnsureMVRUUIDsForRelatedActors();
 
-	/** Updates the GDTFToDefaultActorClasses map so it only contains those entries present both in the DMX Libraryand the MVR Scene */
-	void UpdateGDTFToDefaultActorClasses(const TSharedRef<FDMXMVRFixtureActorLibrary>& MVRFixtureActorLibrary);
+#if WITH_EDITOR
+	/** Updates the FixtureTypeToDefaultActorClasses map so it only contains those entries present both in the DMX Library and the MVR Scene */
+	void UpdateFixtureTypeToDefaultActorClasses(const TSharedRef<FDMXMVRFixtureActorLibrary>& MVRFixtureActorLibrary);
+#endif // WITH_EDITOR
 
 	/** Gets the Fixture Patch of an Actor, or nullptr if the Fixture Patch cannot be retrieved anymore */
 	UDMXEntityFixturePatch* GetFixturePatchFromActor(AActor* Actor) const;
@@ -109,38 +137,48 @@ private:
 	/** Called when an asset was imported */
 	void OnAssetPostImport(UFactory* InFactory, UObject* ActorAdded);
 
-	/** Replaces related Actors that use the Default Actor Class for the GDTF with an instance of the new Default Actor Class */
-	void HandleDefaultActorClassForGDTFChanged();
+	/** Replaces related Actors that use the Default Actor Class for the Fixture Type with an instance of the new Default Actor Class */
+	void HandleDefaultActorClassForFixtureTypeChanged();
 
 	/** Spawns an MVR Actor in this Scene. Returns the newly spawned Actor, or nullptr if no Actor could be spawned. */
-	AActor* SpawnMVRActor(const TSubclassOf<AActor>&ActorClass, UDMXEntityFixturePatch * FixturePatch, const FTransform & Transform, AActor * Template = nullptr);
+	AActor* SpawnMVRActor(const TSubclassOf<AActor>& ActorClass, UDMXEntityFixturePatch * FixturePatch, const FTransform & Transform, AActor * Template = nullptr);
 
 	/** Replaces an MVR Actor in this Scene with another. Returns the newly spawned Actor, or nullptr if no Actor could be spawned. */
 	AActor* ReplaceMVRActor(AActor* ActorToReplace, const TSubclassOf<AActor>& ClassOfNewActor);
+
+	/** Upgrades this Actor created before 5.5 to use the FixtureTypeToActorClasses array instead of the deprecated GDTFToDefaultActorClasses array */
+	void UpgradeToFixtureTypeToActorClasses();
 #endif // WITH_EDITOR
 
 	/** The DMX Library this Scene Actor uses */
-	UPROPERTY(BlueprintReadOnly, VisibleAnywhere, Category = "MVR", Meta = (AllowPrivateAccess = true))
+	UPROPERTY(BlueprintReadOnly, VisibleAnywhere, Category = "MVR", meta = (AllowPrivateAccess = true))
 	TObjectPtr<UDMXLibrary> DMXLibrary;
 
 	/** The actors that created along with this scene */
-	UPROPERTY(VisibleAnywhere, Category = "Actor", AdvancedDisplay, Meta = (AllowPrivateAccess = true))
+	UPROPERTY(VisibleAnywhere, Category = "Actor", AdvancedDisplay, meta = (AllowPrivateAccess = true))
 	TArray<TSoftObjectPtr<AActor>> RelatedActors;
 	
-	/** The actor class that is spawned for a specific GDTF by default (can be overriden per MVR UUID, see below) */
-	UPROPERTY(EditAnywhere, Category = "MVR", Meta = (DispayName = "Default Actor used for GDTF"))
-	TArray<FDMXMVRSceneGDTFToActorClassPair> GDTFToDefaultActorClasses;
+	/** DEPRECATED 5.5, instead use FixtureTypeToActorClasses */
+	PRAGMA_DISABLE_DEPRECATION_WARNINGS
+	UPROPERTY()
+	TArray<FDMXMVRSceneGDTFToActorClassPair> GDTFToDefaultActorClasses_DEPRECATED;
+	PRAGMA_ENABLE_DEPRECATION_WARNINGS
+
+	/** The actor class that is spawned for a specific Fixture Type by default (can be overriden per MVR UUID, see below) */
+	UPROPERTY(EditAnywhere, Category = "MVR", meta = (DispayName = "Fixture Type to Spawned Actor"))
+	TArray<FDMXMVRSceneFixtureTypeToActorClassPair> FixtureTypeToActorClasses;
 
 	/** Array holding MVR Fixture UUIDs that were explicitly removed from the scene */
 	UPROPERTY()
 	TArray<FGuid> DeletedMVRFixtureUUIDs;
 
 #if WITH_EDITORONLY_DATA
-	/** The GDTFToDefaultActorClassMap cached of PreEditChange, to find changes. */
-	TArray<FDMXMVRSceneGDTFToActorClassPair> GDTFToDefaultActorClasses_PreEditChange;
+	/** The FixtureTypeToActorClasses cached of PreEditChange, to find changes. */
+	UPROPERTY(Transient)
+	TArray<FDMXMVRSceneFixtureTypeToActorClassPair> FixtureTypeToActorClasses_PreEditChange;
 #endif
 
 	/** The root component to which all actors are attached initially */
-	UPROPERTY(VisibleAnywhere, Category = "Actor", AdvancedDisplay, Meta = (AllowPrivateAccess = true))
+	UPROPERTY(VisibleAnywhere, Category = "Actor", AdvancedDisplay, meta = (AllowPrivateAccess = true))
 	TObjectPtr<USceneComponent> MVRSceneRoot;
 };
