DMX: Fix DMX Library editor related issues

* Modes cannot be selected from the Add Fixture Patch menu
* Fixture Patch editor does not always reflect Fixture Type changes correctly

diff --git a/Engine/Plugins/VirtualProduction/DMX/DMXEngine/Source/DMXEditor/Private/Widgets/FixturePatch/DMXFixturePatchListItem.cpp b/Engine/Plugins/VirtualProduction/DMX/DMXEngine/Source/DMXEditor/Private/Widgets/FixturePatch/DMXFixturePatchListItem.cpp
--- a/Engine/Plugins/VirtualProduction/DMX/DMXEngine/Source/DMXEditor/Private/Widgets/FixturePatch/DMXFixturePatchListItem.cpp
+++ b/Engine/Plugins/VirtualProduction/DMX/DMXEngine/Source/DMXEditor/Private/Widgets/FixturePatch/DMXFixturePatchListItem.cpp
@@ -49,6 +49,8 @@ FString FDMXFixturePatchListItem::GetFixturePatchName() const
 
 void FDMXFixturePatchListItem::SetFixturePatchName(const FString& InDesiredName, FString& OutNewName)
 {
+	const TGuardValue<bool> ChangeFixturePatchGuard(bChangingFixturePatch, true);
+
 	if (UDMXEntityFixturePatch* FixturePatch = WeakFixturePatch.Get())
 	{
 		if (FixturePatch->Name == InDesiredName)
@@ -74,6 +76,8 @@ FString FDMXFixturePatchListItem::GetFixtureID() const
 
 void FDMXFixturePatchListItem::SetFixtureID(int32 InFixtureID)
 {
+	const TGuardValue<bool> ChangeFixturePatchGuard(bChangingFixturePatch, true);
+
 	if (UDMXEntityFixturePatch* FixturePatch = WeakFixturePatch.Get())
 	{
 		const FScopedTransaction SetFixturePatchNameTransaction(LOCTEXT("SetFixturePatchFixtureIDTransaction", "Set Fixture ID"));
@@ -97,6 +101,8 @@ UDMXEntityFixtureType* FDMXFixturePatchListItem::GetFixtureType() const
 
 void FDMXFixturePatchListItem::SetFixtureType(UDMXEntityFixtureType* FixtureType)
 {
+	const TGuardValue<bool> ChangeFixturePatchGuard(bChangingFixturePatch, true);
+
 	UDMXLibrary* DMXLibrary = GetDMXLibrary();
 	UDMXEntityFixturePatch* FixturePatch = WeakFixturePatch.Get();
 
@@ -128,6 +134,14 @@ int32 FDMXFixturePatchListItem::GetModeIndex() const
 
 void FDMXFixturePatchListItem::SetModeIndex(int32 ModeIndex)
 {
+	// Don't handle other's transactions
+	if (GIsTransacting)
+	{
+		return;
+	}
+
+	const TGuardValue<bool> ChangeFixturePatchGuard(bChangingFixturePatch, true);
+
 	UDMXLibrary* DMXLibrary = GetDMXLibrary();
 	UDMXEntityFixturePatch* FixturePatch = WeakFixturePatch.Get();
 
@@ -148,6 +162,20 @@ void FDMXFixturePatchListItem::SetModeIndex(int32 ModeIndex)
 	}
 }
 
+bool FDMXFixturePatchListItem::GetActiveModeName(FString& OutModeName) const
+{
+	const UDMXEntityFixturePatch* FixturePatch = WeakFixturePatch.Get();
+	const FDMXFixtureMode* ActiveModePtr = FixturePatch ? FixturePatch->GetActiveMode() : nullptr;
+
+	if (ActiveModePtr)
+	{
+		OutModeName = ActiveModePtr->ModeName;
+		return true;
+	}
+
+	return false;
+}
+
 int32 FDMXFixturePatchListItem::GetUniverse() const
 {
 	if (UDMXEntityFixturePatch* FixturePatch = WeakFixturePatch.Get())
@@ -170,6 +198,8 @@ int32 FDMXFixturePatchListItem::GetAddress() const
 
 void FDMXFixturePatchListItem::SetAddresses(int32 Universe, int32 Address)
 {
+	const TGuardValue<bool> ChangeFixturePatchGuard(bChangingFixturePatch, true);
+
 	UDMXLibrary* DMXLibrary = GetDMXLibrary();
 	UDMXEntityFixturePatch* FixturePatch = WeakFixturePatch.Get();
 	const TSharedPtr<FDMXEditor> DMXEditor = WeakDMXEditor.Pin();
diff --git a/Engine/Plugins/VirtualProduction/DMX/DMXEngine/Source/DMXEditor/Private/Widgets/FixturePatch/DMXFixturePatchListItem.h b/Engine/Plugins/VirtualProduction/DMX/DMXEngine/Source/DMXEditor/Private/Widgets/FixturePatch/DMXFixturePatchListItem.h
--- a/Engine/Plugins/VirtualProduction/DMX/DMXEngine/Source/DMXEditor/Private/Widgets/FixturePatch/DMXFixturePatchListItem.h
+++ b/Engine/Plugins/VirtualProduction/DMX/DMXEngine/Source/DMXEditor/Private/Widgets/FixturePatch/DMXFixturePatchListItem.h
@@ -49,6 +49,9 @@ public:
 	/** Sets the Mode Index for the MVR */
 	void SetModeIndex(int32 ModeIndex);
 
+	/** Gets the name of the active mode. Returns false if there is no active mode for this patch. */
+	bool GetActiveModeName(FString& OutModeName) const;
+
 	/** Returns the Universe the MVR Fixture resides in */
 	int32 GetUniverse() const;
 
@@ -67,6 +70,9 @@ public:
 	/** Returns the DMX Library in which the MVR Fixture resides */
 	UDMXLibrary* GetDMXLibrary() const;
 
+	/** Returns true if this item is changing its Fixture Patch */
+	bool IsChangingFixturePatch() const { return bChangingFixturePatch; }
+
 	/** Warning Status Text of the Item */
 	FText WarningStatusText;
 
@@ -77,6 +83,9 @@ private:
 	/** True if this is in an even group */
 	bool bIsEvenGroup = false;
 
+	/** True while this item is changing the DMX Library */
+	bool bChangingFixturePatch = false;
+
 	/** The fixture patch the MVR Fixture UUID is assigned to */
 	TWeakObjectPtr<UDMXEntityFixturePatch> WeakFixturePatch;
 
diff --git a/Engine/Plugins/VirtualProduction/DMX/DMXEngine/Source/DMXEditor/Private/Widgets/FixturePatch/SAddFixturePatchMenu.cpp b/Engine/Plugins/VirtualProduction/DMX/DMXEngine/Source/DMXEditor/Private/Widgets/FixturePatch/SAddFixturePatchMenu.cpp
deleted file mode 100644
diff --git a/Engine/Plugins/VirtualProduction/DMX/DMXEngine/Source/DMXEditor/Private/Widgets/FixturePatch/SAddFixturePatchMenu.h b/Engine/Plugins/VirtualProduction/DMX/DMXEngine/Source/DMXEditor/Private/Widgets/FixturePatch/SAddFixturePatchMenu.h
deleted file mode 100644
diff --git a/Engine/Plugins/VirtualProduction/DMX/DMXEngine/Source/DMXEditor/Private/Widgets/FixturePatch/SDMXAddFixturePatchMenu.cpp b/Engine/Plugins/VirtualProduction/DMX/DMXEngine/Source/DMXEditor/Private/Widgets/FixturePatch/SDMXAddFixturePatchMenu.cpp
new file mode 100644
--- /dev/null
+++ b/Engine/Plugins/VirtualProduction/DMX/DMXEngine/Source/DMXEditor/Private/Widgets/FixturePatch/SDMXAddFixturePatchMenu.cpp
@@ -0,0 +1,630 @@
+// Copyright Epic Games, Inc. All Rights Reserved.
+
+#include "SDMXAddFixturePatchMenu.h"
+
+#include "Algo/Copy.h"
+#include "Algo/Find.h"
+#include "DMXAddFixturePatchMenuData.h"
+#include "DMXEditor.h"
+#include "DMXFixturePatchSharedData.h"
+#include "Editor.h"
+#include "Framework/Application/SlateApplication.h"
+#include "FixturePatchAutoAssignUtility.h"
+#include "Library/DMXEntityFixtureType.h"
+#include "Library/DMXLibrary.h"
+#include "Misc/ScopedSlowTask.h"
+#include "ScopedTransaction.h"
+#include "TimerManager.h"
+#include "UObject/Package.h"
+#include "Widgets/Input/SButton.h"
+#include "Widgets/Input/SCheckBox.h"
+#include "Widgets/Input/SEditableTextBox.h"
+#include "Widgets/Layout/SGridPanel.h"
+#include "Widgets/SBoxPanel.h"
+#include "Widgets/SDMXEntityDropdownMenu.h"
+
+
+#define LOCTEXT_NAMESPACE "SDMXAddFixturePatchMenu"
+
+namespace UE::DMXEditor::FixturePatchEditor
+{
+	SDMXAddFixturePatchMenu::~SDMXAddFixturePatchMenu()
+	{
+		UDMXAddFixturePatchMenuData* MenuData = GetMutableDefault<UDMXAddFixturePatchMenuData>();
+		MenuData->SoftFixtureType = WeakFixtureType.Get();
+	}
+
+	void SDMXAddFixturePatchMenu::Construct(const FArguments& InArgs, TWeakPtr<FDMXEditor> InWeakDMXEditor)
+	{
+		WeakDMXEditor = InWeakDMXEditor;
+
+		SharedData = WeakDMXEditor.IsValid() ? WeakDMXEditor.Pin()->GetFixturePatchSharedData() : nullptr;
+		if (SharedData.IsValid())
+		{
+			UDMXLibrary::GetOnEntitiesAdded().AddSP(this, &SDMXAddFixturePatchMenu::OnEntityAddedOrRemoved);
+			UDMXLibrary::GetOnEntitiesRemoved().AddSP(this, &SDMXAddFixturePatchMenu::OnEntityAddedOrRemoved);
+
+			Refresh();
+		}
+	}
+
+	void SDMXAddFixturePatchMenu::RequestRefresh()
+	{
+		RequestRefreshModeComboBoxTimerHandle.Invalidate();
+
+		if (!RequestRefreshModeComboBoxTimerHandle.IsValid())
+		{
+			RequestRefreshModeComboBoxTimerHandle = GEditor->GetTimerManager()->SetTimerForNextTick(FTimerDelegate::CreateSP(this, &SDMXAddFixturePatchMenu::Refresh));
+		}
+	}
+
+	void SDMXAddFixturePatchMenu::Refresh()
+	{
+		RequestRefreshModeComboBoxTimerHandle.Invalidate();
+
+		UDMXLibrary* DMXLibrary = WeakDMXEditor.IsValid() ? WeakDMXEditor.Pin()->GetDMXLibrary() : nullptr;
+		if (!DMXLibrary)
+		{
+			return;
+		}
+
+		// Mend the fixture type
+		if (!WeakFixtureType.IsValid())
+		{
+			const TArray<UDMXEntityFixtureType*> FixtureTypes = DMXLibrary->GetEntitiesTypeCast<UDMXEntityFixtureType>();
+			WeakFixtureType = FixtureTypes.IsEmpty() ? nullptr : FixtureTypes[0];
+		}
+
+
+		// Mend the active mode index
+		UDMXAddFixturePatchMenuData* MenuData = GetMutableDefault<UDMXAddFixturePatchMenuData>();
+		if (WeakFixtureType.IsValid() && !WeakFixtureType->Modes.IsValidIndex(MenuData->ActiveModeIndex))
+		{
+			MenuData->ActiveModeIndex = 0;
+			MenuData->SaveConfig();
+		}
+
+
+		// Create the combo box source and an intial selection
+		ModeSources.Reset();
+		if (UDMXEntityFixtureType* FixtureType = WeakFixtureType.Get())
+		{
+			for (int32 ModeIndex = 0; ModeIndex < FixtureType->Modes.Num(); ModeIndex++)
+			{
+				ModeSources.Add(MakeShared<uint32>(ModeIndex));
+			}
+		}
+
+		const TArray<UDMXEntityFixtureType*> FixtureTypes = DMXLibrary->GetEntitiesTypeCast<UDMXEntityFixtureType>();
+		UDMXEntityFixtureType* const* SelectedFixtureTypePtr = Algo::FindByPredicate(FixtureTypes, [MenuData](const UDMXEntityFixtureType* FixtureType)
+			{
+				return FixtureType == MenuData->SoftFixtureType;
+			});
+		if (!SelectedFixtureTypePtr)
+		{
+			SelectedFixtureTypePtr = FixtureTypes.IsEmpty() ? nullptr : &FixtureTypes[0];
+		}
+		WeakFixtureType = SelectedFixtureTypePtr ? *SelectedFixtureTypePtr : nullptr;
+
+
+		// Rebuild the widget
+		ChildSlot
+		[
+			SNew(SBorder)
+			.BorderImage(FAppStyle::GetBrush("ToolPanel.GroupBorder"))
+			[
+				SNew(SGridPanel)
+				.FillColumn(1, 1.f)
+
+				// Select Fixture Type
+				+ SGridPanel::Slot(0, 0)
+				.Padding(4.f)
+				.VAlign(VAlign_Center)
+				[
+					SNew(STextBlock)
+					.Text(LOCTEXT("SelectFixtureTypeLabel", "Fixture Type"))
+					.Font(FAppStyle::GetFontStyle(TEXT("PropertyWindow.NormalFont")))
+				]
+
+				+ SGridPanel::Slot(1, 0)
+				.Padding(4.f)
+				.VAlign(VAlign_Center)
+				[
+					MakeFixtureTypeSelectWidget()
+				]
+
+				// Select Mode
+				+ SGridPanel::Slot(0, 1)
+				.Padding(4.f)
+				.VAlign(VAlign_Center)
+				[
+					SNew(STextBlock)
+					.IsEnabled(this, &SDMXAddFixturePatchMenu::HasValidFixtureTypeAndMode)
+					.Text(LOCTEXT("SelectModeLabel", "Mode"))
+					.Font(FAppStyle::GetFontStyle(TEXT("PropertyWindow.NormalFont")))
+				]
+				+ SGridPanel::Slot(1, 1)
+				.Padding(4.f)
+				[
+					MakeModeSelectWidget()
+				]
+
+				// Universe label 
+				+ SGridPanel::Slot(0, 2)
+				.Padding(4.f)
+				.VAlign(VAlign_Center)
+				[
+					SNew(STextBlock)
+					.IsEnabled(this, &SDMXAddFixturePatchMenu::HasValidFixtureTypeAndMode)
+					.Text(LOCTEXT("UniverseDotChannelLabel", "Universe.Channel"))
+					.Font(FAppStyle::GetFontStyle(TEXT("PropertyWindow.NormalFont")))
+				]
+
+				// Select universe and channel 
+				+ SGridPanel::Slot(1, 2)
+				.Padding(4.f)
+				.VAlign(VAlign_Center)
+				[
+					SNew(SHorizontalBox)
+
+					+ SHorizontalBox::Slot()
+					.AutoWidth()
+					[
+						MakeUniverseChannelSelectWidget()
+					]
+
+					+ SHorizontalBox::Slot()
+					.AutoWidth()
+					.Padding(4.f, 0.f, 0.f, 0.f)
+					[
+						MakeAutoIncrementChannelCheckBox()
+					]
+				]
+
+				// Num Fixture Patches Label
+				+ SGridPanel::Slot(0, 3)
+				.Padding(4.f)
+				.VAlign(VAlign_Center)
+				[
+					SNew(STextBlock)
+					.IsEnabled(this, &SDMXAddFixturePatchMenu::HasValidFixtureTypeAndMode)
+					.Text(LOCTEXT("NumPatchesLabel", "Num Patches"))
+					.Font(FAppStyle::GetFontStyle(TEXT("PropertyWindow.NormalFont")))
+				]
+
+				// Num Fixture Patches editable Text Box
+				+ SGridPanel::Slot(1, 3)
+				.Padding(4.f)
+				[
+					MakeNumFixturePatchesEditableTextBox()
+				]
+
+				// Add fixture patches button
+				+ SGridPanel::Slot(1, 4)
+				.Padding(4.f)
+				[
+					MakeAddFixturePatchesButton()
+				]
+			]
+		];
+	}
+
+	TSharedRef<SWidget> SDMXAddFixturePatchMenu::MakeFixtureTypeSelectWidget()
+	{
+		return
+			SNew(SBox)
+			.MinDesiredWidth(100.f)
+			[
+				SAssignNew(FixtureTypeSelector, SDMXEntityPickerButton<UDMXEntityFixtureType>)
+				.DMXEditor(WeakDMXEditor)
+				.CurrentEntity_Lambda([this] { return WeakFixtureType.Get(); })
+				.OnEntitySelected(this, &SDMXAddFixturePatchMenu::OnFixtureTypeSelected)
+			];
+	}
+
+	TSharedRef<SWidget> SDMXAddFixturePatchMenu::MakeModeSelectWidget()
+	{
+		return
+			SAssignNew(ModeComboBox, SComboBox<TSharedPtr<uint32>>)
+			.IsEnabled(this, &SDMXAddFixturePatchMenu::HasValidFixtureTypeAndMode)
+			.OptionsSource(&ModeSources)
+			.OnGenerateWidget(this, &SDMXAddFixturePatchMenu::GenerateModeComboBoxEntry)
+			.OnSelectionChanged(this, &SDMXAddFixturePatchMenu::OnModeSelected)
+			.InitiallySelectedItem(0)
+			[
+				SNew(STextBlock)
+				.MinDesiredWidth(50.0f)
+				.Text(this, &SDMXAddFixturePatchMenu::GetActiveModeText)
+				.Font(FAppStyle::GetFontStyle(TEXT("PropertyWindow.NormalFont")))
+			];
+	}
+
+	TSharedRef<SWidget> SDMXAddFixturePatchMenu::MakeUniverseChannelSelectWidget()
+	{
+		return
+			SNew(SBox)
+			.HAlign(HAlign_Left)
+			[
+				SAssignNew(UniverseChannelEditableTextBox, SEditableTextBox)
+				.IsEnabled(this, &SDMXAddFixturePatchMenu::HasValidFixtureTypeAndMode)
+				.MinDesiredWidth(60.f)
+				.Font(FAppStyle::GetFontStyle(TEXT("PropertyWindow.NormalFont")))
+				.SelectAllTextWhenFocused(true)
+				.ClearKeyboardFocusOnCommit(true)
+				.RevertTextOnEscape(true)
+				.Text(this, &SDMXAddFixturePatchMenu::GetUniverseChannelText)
+				.OnTextChanged(this, &SDMXAddFixturePatchMenu::OnUniverseChannelTextChanged)
+				.OnTextCommitted(this, &SDMXAddFixturePatchMenu::OnUniverseChannelTextCommitted)
+			];
+	}
+
+	TSharedRef<SWidget> SDMXAddFixturePatchMenu::MakeAutoIncrementChannelCheckBox()
+	{
+		return 
+			SNew(SCheckBox)
+			.IsEnabled(this, &SDMXAddFixturePatchMenu::HasValidFixtureTypeAndMode)
+			.IsChecked_Lambda([]()
+				{
+					const UDMXAddFixturePatchMenuData* MenuData = GetDefault<UDMXAddFixturePatchMenuData>();
+					return MenuData->bIncrementChannelAfterPatching ? ECheckBoxState::Checked : ECheckBoxState::Unchecked;
+				})
+			.OnCheckStateChanged_Lambda([](ECheckBoxState NewState)
+				{
+					UDMXAddFixturePatchMenuData* MenuData = GetMutableDefault<UDMXAddFixturePatchMenuData>();
+					MenuData->bIncrementChannelAfterPatching = NewState == ECheckBoxState::Checked;
+					MenuData->SaveConfig();
+				})
+			[
+				SNew(STextBlock)
+				.Font(FAppStyle::GetFontStyle(TEXT("PropertyWindow.NormalFont")))
+				.Text(LOCTEXT("AutoIncrementChannelLabel", "Increment after patching"))
+				.ToolTipText(LOCTEXT("AutoIncrementChannelTooltip", "Automatically increments the universe/channel to the first subsequent channel after patching."))
+			];
+	}
+
+	TSharedRef<SWidget> SDMXAddFixturePatchMenu::MakeNumFixturePatchesEditableTextBox()
+	{
+		return
+			SNew(SBox)
+			.HAlign(HAlign_Left)
+			[
+				SNew(SEditableTextBox)
+				.IsEnabled(this, &SDMXAddFixturePatchMenu::HasValidFixtureTypeAndMode)
+				.MinDesiredWidth(60.f)
+				.Font(FAppStyle::GetFontStyle(TEXT("PropertyWindow.NormalFont")))
+				.SelectAllTextWhenFocused(true)
+				.ClearKeyboardFocusOnCommit(true)
+				.Text_Lambda([this]()
+					{
+						return FText::FromString(FString::FromInt(NumFixturePatchesToAdd));
+					})
+				.OnTextCommitted_Lambda([this](const FText& Text, ETextCommit::Type CommitType)
+					{
+						int32 Value;
+						if (LexTryParseString<int32>(Value, *Text.ToString()) &&
+							Value > 0)
+						{
+							constexpr int32 MaxNumFixturePatchesToAdd = 16384;
+							NumFixturePatchesToAdd = FMath::Clamp(Value, 1, MaxNumFixturePatchesToAdd);
+						}
+
+						// Add fixture patches if enter is pressed
+						if (CommitType == ETextCommit::OnEnter)
+						{
+							OnAddFixturePatchButtonClicked();
+						}
+					})
+			];
+	}
+
+	TSharedRef<SWidget> SDMXAddFixturePatchMenu::MakeAddFixturePatchesButton()
+	{
+		return
+			SNew(SBox)
+			.HAlign(HAlign_Right)
+			.MinDesiredWidth(120.f)
+			[
+				SNew(SButton)
+				.IsEnabled(this, &SDMXAddFixturePatchMenu::HasValidFixtureTypeAndMode)
+				.ContentPadding(FMargin(4.0f, 4.0f))
+				.ButtonStyle(FAppStyle::Get(), "FlatButton.Success")
+				.ForegroundColor(FLinearColor::White)
+				.OnClicked(this, &SDMXAddFixturePatchMenu::OnAddFixturePatchButtonClicked)
+				[
+					SNew(STextBlock)
+					.Text_Lambda([this]()
+						{
+							return FText::Format(LOCTEXT("AddFixturePatchButtonText", "Add Fixture {0}|plural(one=Patch, other=Patches)"), NumFixturePatchesToAdd);
+						})
+					.Font(FAppStyle::GetFontStyle(TEXT("PropertyWindow.BoldFont")))
+				]
+			];
+	}
+
+	TSharedRef<SWidget> SDMXAddFixturePatchMenu::GenerateModeComboBoxEntry(const TSharedPtr<uint32> InModeIndex) const
+	{
+		UDMXEntityFixtureType* FixtureType = Cast<UDMXEntityFixtureType>(WeakFixtureType);
+		if (!FixtureType)
+		{
+			return SNullWidget::NullWidget;
+		}
+
+		const TArray<FDMXFixtureMode>& Modes = FixtureType->Modes;
+
+		return
+			SNew(STextBlock)
+			.MinDesiredWidth(50.0f)
+			.Text_Lambda([&Modes, ModeIndex = *InModeIndex, this]()
+				{
+					const UDMXAddFixturePatchMenuData* MenuData = GetDefault<UDMXAddFixturePatchMenuData>();
+					if (Modes.IsValidIndex(ModeIndex))
+					{
+						return FText::FromString(Modes[ModeIndex].ModeName);
+					}
+					else
+					{
+						return LOCTEXT("NoModeAvailableText", "No Mode available");
+					}
+				})
+			.Font(FAppStyle::GetFontStyle(TEXT("PropertyWindow.NormalFont")));
+	}
+
+	void SDMXAddFixturePatchMenu::OnEntityAddedOrRemoved(UDMXLibrary* DMXLibrary, TArray<UDMXEntity*> Entities)
+	{
+		RequestRefresh();
+	}
+
+	void SDMXAddFixturePatchMenu::OnFixtureTypeSelected(UDMXEntity* InSelectedFixtureType)
+	{
+		UDMXEntityFixtureType* SelectedFixtureType = Cast<UDMXEntityFixtureType>(InSelectedFixtureType);
+		WeakFixtureType = SelectedFixtureType;
+
+		UDMXAddFixturePatchMenuData* MenuData = GetMutableDefault<UDMXAddFixturePatchMenuData>();
+		MenuData->SoftFixtureType = SelectedFixtureType;
+		MenuData->SaveConfig();
+
+		RequestRefresh();
+	}
+
+	void SDMXAddFixturePatchMenu::OnModeSelected(TSharedPtr<uint32> InSelectedMode, ESelectInfo::Type SelectInfo)
+	{
+		UDMXAddFixturePatchMenuData* MenuData = GetMutableDefault<UDMXAddFixturePatchMenuData>();
+		MenuData->ActiveModeIndex = InSelectedMode.IsValid() ? *InSelectedMode : 0;
+		MenuData->SaveConfig();
+	}
+
+	FText SDMXAddFixturePatchMenu::GetUniverseChannelText() const
+	{
+		if (!Universe.IsSet() && !Channel.IsSet())
+		{
+			const int32 SelectedUniverse = SharedData->GetSelectedUniverse();
+			return FText::FromString(FString::FromInt(SelectedUniverse) + TEXT(".") + FString::FromInt(1));
+		}
+		else
+		{
+			const FString UnivereString = Universe.IsSet() ? FString::FromInt(Universe.GetValue()) : TEXT("");
+			const FString ChannelString = Channel.IsSet() ? FString::FromInt(Channel.GetValue()) : TEXT("");
+			const FString Separator = Universe.IsSet() && Channel.IsSet() ? TEXT(".") : TEXT("");
+
+			return FText::FromString(UnivereString + Separator + ChannelString);
+		}
+	}
+
+	void SDMXAddFixturePatchMenu::OnUniverseChannelTextChanged(const FText& Text)
+	{		
+		if (!UniverseChannelEditableTextBox.IsValid())
+		{
+			return;
+		}
+
+		static const TCHAR* Delimiter[] = { TEXT("."), TEXT(","), TEXT(":"), TEXT(" ") };
+		TArray<FString> Substrings;
+		constexpr bool bCullEmpty = true;
+		Text.ToString().ParseIntoArray(Substrings, Delimiter, 4, bCullEmpty);
+
+		static const FText InvalidStringErrorMessage = LOCTEXT("InvalidUniverseString", "Must be in the form of 'Universe' or 'Universe.Channel'. E.g. '4', or '4.5'.");
+
+		FText ErrorMessage;
+		int32 ParsedUniverse;
+		if (Substrings.IsValidIndex(0) &&
+			LexTryParseString(ParsedUniverse, *Substrings[0]))
+		{
+			if (ParsedUniverse < 1 || ParsedUniverse > DMX_MAX_UNIVERSE)
+			{
+				ErrorMessage = FText::Format(LOCTEXT("InvalidUniverseValue", "Universe must be between 1 and {0}."), FText::FromString(FString::FromInt((int32)DMX_MAX_UNIVERSE)));
+			}
+		}
+		else if (Substrings.IsValidIndex(0))
+		{
+			ErrorMessage = InvalidStringErrorMessage;
+		}
+
+		// Channel is optional, only test if Substrings[1] was parsed
+		int32 ParsedChannel;
+		if (Substrings.IsValidIndex(1) &&
+			LexTryParseString(ParsedChannel, *Substrings[1]))
+		{
+			if (ParsedChannel < 1 || ParsedChannel > DMX_MAX_ADDRESS)
+			{
+				ErrorMessage = FText::Format(LOCTEXT("InvalidChannelValue", "Channel must be between 1 and {0}."), FText::FromString(FString::FromInt((int32)DMX_MAX_ADDRESS)));
+			}
+		}
+		else if (Substrings.IsValidIndex(1))
+		{
+			ErrorMessage = InvalidStringErrorMessage;
+		}
+
+		UniverseChannelEditableTextBox->SetError(ErrorMessage);
+	}
+
+	void SDMXAddFixturePatchMenu::OnUniverseChannelTextCommitted(const FText& Text, ETextCommit::Type CommitType)
+	{
+		if (!UniverseChannelEditableTextBox.IsValid())
+		{
+			return;
+		}
+
+		static const TCHAR* Delimiter[] = { TEXT("."), TEXT(","), TEXT(":"), TEXT(" ") };
+		TArray<FString> Substrings;
+		constexpr bool bCullEmpty = true;
+		Text.ToString().ParseIntoArray(Substrings, Delimiter, 4, bCullEmpty);
+
+		int32 NewUniverse;
+		if (Substrings.IsValidIndex(0) &&
+			LexTryParseString(NewUniverse, *Substrings[0]))
+		{
+			Universe = FMath::Clamp(NewUniverse, 1, DMX_MAX_UNIVERSE);
+			SharedData->SelectUniverse(NewUniverse);
+		}
+		else
+		{
+			Universe.Reset();
+		}
+
+		int32 NewChannel;
+		if (Universe.IsSet() && // Only parse the channel if parsing universe was successful
+			Substrings.IsValidIndex(1) &&
+			LexTryParseString(NewChannel, *Substrings[1]))
+		{
+			Channel = FMath::Clamp(NewChannel, 1, DMX_MAX_ADDRESS);
+		}
+		else
+		{
+			Channel.Reset();
+		}
+
+		if (UniverseChannelEditableTextBox.IsValid())
+		{
+			UniverseChannelEditableTextBox->SetError(FText::GetEmpty());
+		}
+
+		// Add fixture patches if enter is pressed
+		if (CommitType == ETextCommit::OnEnter)
+		{
+			OnAddFixturePatchButtonClicked();
+		}
+
+		// Never show an error, values are always mended.
+		UniverseChannelEditableTextBox->SetError(FText::GetEmpty());
+	}
+
+	FReply SDMXAddFixturePatchMenu::OnAddFixturePatchButtonClicked()
+	{
+		FSlateApplication::Get().DismissAllMenus();
+
+		UDMXEntityFixtureType* FixtureType = Cast<UDMXEntityFixtureType>(WeakFixtureType);
+		if (!FixtureType || FixtureType->Modes.IsEmpty())
+		{
+			return FReply::Handled();
+		}
+
+		const TSharedPtr<FDMXEditor> DMXEditor = WeakDMXEditor.Pin();
+		if (!DMXEditor.IsValid())
+		{
+			return FReply::Handled();
+		}
+
+		UDMXLibrary* DMXLibrary = DMXEditor->GetDMXLibrary();
+		if (!DMXLibrary)
+		{
+			return FReply::Handled();
+		}
+
+		// Ensure valid mode
+		const UDMXAddFixturePatchMenuData* MenuData = GetDefault<UDMXAddFixturePatchMenuData>();
+		if (!ensureMsgf(FixtureType->Modes.IsValidIndex(MenuData->ActiveModeIndex), TEXT("Cannot apply mode. Mode index is invalid.")))
+		{
+			return FReply::Handled();
+		}
+
+		// Create new fixture patches
+		const FScopedTransaction CreateFixturePatchTransaction(LOCTEXT("CreateFixturePatchTransaction", "Create Fixture Patch"));
+		DMXLibrary->PreEditChange(UDMXLibrary::StaticClass()->FindPropertyByName(UDMXLibrary::GetEntitiesPropertyName()));
+
+		const int32 PatchToUniverse = Universe.IsSet() ? Universe.GetValue() : SharedData->GetSelectedUniverse();
+
+		const float NumSteps = NumFixturePatchesToAdd;
+		FScopedSlowTask Task(NumSteps, LOCTEXT("AddFixturePatchesSlowTask", "Adding Fixture Patches..."));
+		Task.MakeDialogDelayed(.5f);
+
+		TArray<UDMXEntityFixturePatch*> NewFixturePatches;
+		NewFixturePatches.Reserve(NumFixturePatchesToAdd);
+		for (uint32 iNumFixturePatchesAdded = 0; iNumFixturePatchesAdded < NumFixturePatchesToAdd; iNumFixturePatchesAdded++)
+		{
+			Task.EnterProgressFrame();
+
+			FDMXEntityFixturePatchConstructionParams FixturePatchConstructionParams;
+			FixturePatchConstructionParams.FixtureTypeRef = FDMXEntityFixtureTypeRef(FixtureType);
+			FixturePatchConstructionParams.ActiveMode = MenuData->ActiveModeIndex;
+			FixturePatchConstructionParams.UniverseID = PatchToUniverse;
+			FixturePatchConstructionParams.StartingAddress = Channel.IsSet() ? Channel.GetValue() : 1;
+
+			constexpr bool bMarkLibraryDirty = false;
+			UDMXEntityFixturePatch* NewFixturePatch = UDMXEntityFixturePatch::CreateFixturePatchInLibrary(FixturePatchConstructionParams, FixtureType->Name, bMarkLibraryDirty);
+			
+			if (NewFixturePatch)
+			{
+				NewFixturePatches.Add(NewFixturePatch);
+			}
+		}
+
+		// Align
+		using namespace UE::DMXEditor::AutoAssign;
+		FAutoAssignUtility::Align(NewFixturePatches);
+
+		DMXLibrary->PostEditChange();
+
+		// Select universe and new fixture patches
+		SharedData->SelectUniverse(PatchToUniverse);
+
+		TArray<TWeakObjectPtr<UDMXEntityFixturePatch>> NewWeakFixturePatches;
+		NewWeakFixturePatches.Reserve(NumFixturePatchesToAdd);
+		Algo::Copy(NewFixturePatches, NewWeakFixturePatches);
+		SharedData->SelectFixturePatches(NewWeakFixturePatches);
+
+		// Increment universe/channel if desired
+		if (!NewFixturePatches.IsEmpty() && GetDefault<UDMXAddFixturePatchMenuData>()->bIncrementChannelAfterPatching)
+		{
+			Channel = NewFixturePatches.Last()->GetStartingChannel() + NewFixturePatches.Last()->GetChannelSpan();
+			Universe = NewFixturePatches.Last()->GetUniverseID();
+			if (Channel.GetValue() > DMX_UNIVERSE_SIZE)
+			{
+				Channel = 1;
+				Universe = Universe.GetValue() + 1;
+			}
+		}
+
+		return FReply::Handled();
+	}
+
+	FText SDMXAddFixturePatchMenu::GetActiveModeText() const
+	{
+		UDMXEntityFixtureType* FixtureType = Cast<UDMXEntityFixtureType>(WeakFixtureType);
+		if (!FixtureType)
+		{
+			return FText::GetEmpty();
+		}
+
+		const TArray<FDMXFixtureMode>& Modes = FixtureType->Modes;
+		const UDMXAddFixturePatchMenuData* MenuData = GetMutableDefault<UDMXAddFixturePatchMenuData>();
+		if (Modes.Num() > 0 && Modes.IsValidIndex(MenuData->ActiveModeIndex))
+		{
+			return FText::FromString(Modes[MenuData->ActiveModeIndex].ModeName);
+		}
+		else if (Modes.IsEmpty())
+		{
+			return LOCTEXT("NoModeAvailableComboButtonText", "No Modes in Fixture Type");
+		}
+		else
+		{
+			return LOCTEXT("NoFixtureTypeSelectedComboButtonText", "No Fixture Type selected");
+		}
+	}
+
+	bool SDMXAddFixturePatchMenu::HasValidFixtureTypeAndMode() const
+	{
+		return WeakFixtureType.IsValid() && !WeakFixtureType->Modes.IsEmpty();
+	}
+}
+
+#undef LOCTEXT_NAMESPACE
diff --git a/Engine/Plugins/VirtualProduction/DMX/DMXEngine/Source/DMXEditor/Private/Widgets/FixturePatch/SDMXAddFixturePatchMenu.h b/Engine/Plugins/VirtualProduction/DMX/DMXEngine/Source/DMXEditor/Private/Widgets/FixturePatch/SDMXAddFixturePatchMenu.h
new file mode 100644
--- /dev/null
+++ b/Engine/Plugins/VirtualProduction/DMX/DMXEngine/Source/DMXEditor/Private/Widgets/FixturePatch/SDMXAddFixturePatchMenu.h
@@ -0,0 +1,124 @@
+// Copyright Epic Games, Inc. All Rights Reserved.
+
+#pragma once
+
+#include "Engine/EngineTypes.h"
+#include "Widgets/SCompoundWidget.h"
+
+class FDMXEditor;
+class FDMXFixturePatchSharedData;
+class FReply;
+template <typename OptionType> class SComboBox;
+template <typename EntityType> class SDMXEntityPickerButton;
+class SEditableTextBox;
+class UDMXAddFixturePatchMenuData;
+class UDMXEntity;
+class UDMXEntityFixtureType;
+class UDMXLibrary;
+
+
+namespace UE::DMXEditor::FixturePatchEditor
+{
+	/** Editor for Fixture Patches */
+	class SDMXAddFixturePatchMenu final
+		: public SCompoundWidget
+	{
+	public:
+		SLATE_BEGIN_ARGS(SDMXAddFixturePatchMenu)
+		{}
+		SLATE_END_ARGS()
+
+		virtual ~SDMXAddFixturePatchMenu();
+
+		/** Constructs this widget */
+		void Construct(const FArguments& InArgs, TWeakPtr<FDMXEditor> InWeakDMXEditor);
+
+		/** Refreshes the widget on the next tick */
+		void RequestRefresh();
+
+	private:
+		/** Refreshes the widget */
+		void Refresh();
+
+		/** Creates a widget to select a fixture type */
+		TSharedRef<SWidget> MakeFixtureTypeSelectWidget();
+
+		/** Creates a widget to select a mode */
+		TSharedRef<SWidget> MakeModeSelectWidget();
+
+		/** Creates an widget to select the Universe and Channel */
+		TSharedRef<SWidget> MakeUniverseChannelSelectWidget();
+		
+		/** Creates a widget to toggle if the channels should be auto-incremented after patching */
+		TSharedRef<SWidget> MakeAutoIncrementChannelCheckBox();
+
+		/** Creates an editable text box to specify the number of patches */
+		TSharedRef<SWidget> MakeNumFixturePatchesEditableTextBox();
+
+		/** Creates a button that adds the fixture patches when clicked */
+		TSharedRef<SWidget> MakeAddFixturePatchesButton();
+
+		/** Generates an entry in the mode combo box */
+		TSharedRef<SWidget> GenerateModeComboBoxEntry(const TSharedPtr<uint32> InMode) const;
+
+		/** Called when an entity was added or removed from the DMX Library */
+		void OnEntityAddedOrRemoved(UDMXLibrary* DMXLibrary, TArray<UDMXEntity*> Entities);
+
+		/** Called when a fixture type was selected */
+		void OnFixtureTypeSelected(UDMXEntity* InSelectedFixtureType);
+
+		/** Called when a mode was selected */
+		void OnModeSelected(TSharedPtr<uint32> InSelectedMode, ESelectInfo::Type SelectInfo);
+
+		/** Returns the current universe channel text */
+		FText GetUniverseChannelText() const;
+
+		/** Called when the universe channel text changed */
+		void OnUniverseChannelTextChanged(const FText& InText);
+
+		/** Called when the universe channel text was committed */
+		void OnUniverseChannelTextCommitted(const FText& Text, ETextCommit::Type CommitType);
+
+		/** Called when the 'Add' button was clicked */
+		FReply OnAddFixturePatchButtonClicked();
+
+		/** Returns the text of the active mode */
+		FText GetActiveModeText() const;
+
+		/** Returns if a valid fixture type with a mode is selected */
+		bool HasValidFixtureTypeAndMode() const;
+
+		/** Current fixture type */
+		TWeakObjectPtr<UDMXEntityFixtureType> WeakFixtureType;
+
+		/** Universe where the patches will be added */
+		TOptional<int32> Universe;
+
+		/** Channel where the patches will be added */
+		TOptional<int32> Channel; 
+
+		/** The number of fixture patches to add */
+		uint32 NumFixturePatchesToAdd = 1;
+
+		/** Timer handle for the request refresh combo box timer */
+		FTimerHandle RequestRefreshModeComboBoxTimerHandle;
+
+		/** Widget to select the fixture type */
+		TSharedPtr<SDMXEntityPickerButton<UDMXEntityFixtureType>> FixtureTypeSelector;
+
+		/** Text boxt to set the channel where the patch will be spawned */
+		TSharedPtr<SEditableTextBox> UniverseChannelEditableTextBox;
+
+		/** Sources for the mode combo box */
+		TArray<TSharedPtr<uint32>> ModeSources;
+
+		/** Combo box to selected the mode */
+		TSharedPtr<SComboBox<TSharedPtr<uint32>>> ModeComboBox;
+
+		/** Shared data for fixture patches */
+		TSharedPtr<FDMXFixturePatchSharedData> SharedData;
+
+		/** The DMX editor this widget displays */
+		TWeakPtr<FDMXEditor> WeakDMXEditor;
+	};
+}
diff --git a/Engine/Plugins/VirtualProduction/DMX/DMXEngine/Source/DMXEditor/Private/Widgets/FixturePatch/SDMXFixturePatchList.cpp b/Engine/Plugins/VirtualProduction/DMX/DMXEngine/Source/DMXEditor/Private/Widgets/FixturePatch/SDMXFixturePatchList.cpp
--- a/Engine/Plugins/VirtualProduction/DMX/DMXEngine/Source/DMXEditor/Private/Widgets/FixturePatch/SDMXFixturePatchList.cpp
+++ b/Engine/Plugins/VirtualProduction/DMX/DMXEngine/Source/DMXEditor/Private/Widgets/FixturePatch/SDMXFixturePatchList.cpp
@@ -2,6 +2,7 @@
 
 #include "SDMXFixturePatchList.h"
 
+#include "Algo/AnyOf.h"
 #include "Algo/Copy.h"
 #include "Algo/MaxElement.h"
 #include "Algo/MinElement.h"
@@ -581,34 +582,6 @@ void SDMXFixturePatchList::Construct(const FArguments& InArgs, TWeakPtr<FDMXEdit
 	// Handle Shared Data selection changes
 	FixturePatchSharedData->OnFixturePatchSelectionChanged.AddSP(this, &SDMXFixturePatchList::OnFixturePatchSharedDataSelectedFixturePatches);
 
-	static const FTableViewStyle TableViewStyle = FAppStyle::Get().GetWidgetStyle<FTableViewStyle>("TreeView");
-
-	ChildSlot
-	[
-		SNew(SVerticalBox)
-
-		+ SVerticalBox::Slot()
-		.HAlign(HAlign_Fill)
-		.AutoHeight()
-		[
-			SAssignNew(Toolbar, SDMXFixturePatchListToolbar, WeakDMXEditor)
-			.OnSearchChanged(this, &SDMXFixturePatchList::OnSearchChanged)
-		]
-
-		+ SVerticalBox::Slot()
-		.HAlign(HAlign_Fill)
-		.FillHeight(1.f)
-		[
-			SAssignNew(ListView, SDMXFixturePatchListType)
-			.ListViewStyle(&TableViewStyle)
-			.HeaderRow(GenerateHeaderRow())
-			.ListItemsSource(&ListSource)
-			.OnGenerateRow(this, &SDMXFixturePatchList::OnGenerateRow)
-			.OnSelectionChanged(this, &SDMXFixturePatchList::OnSelectionChanged)
-			.OnContextMenuOpening(this, &SDMXFixturePatchList::OnContextMenuOpening)
-		]
-	];
-
 	RegisterCommands();
 	RefreshList();
 
@@ -623,12 +596,20 @@ void SDMXFixturePatchList::Construct(const FArguments& InArgs, TWeakPtr<FDMXEdit
 
 void SDMXFixturePatchList::RequestListRefresh()
 {
-	if (RequestListRefreshTimerHandle.IsValid())
-	{
-		return;
-	}
+	if (!RequestListRefreshTimerHandle.IsValid())
+	{	
+		// If a fixture patch item is changing a fixture patch, don't refresh. Instead let the rows update themselves.
+		const bool bAnyItemIsChangingFixturePatch = Algo::AnyOf(ListSource, [](const TSharedPtr<FDMXFixturePatchListItem>& Item)
+			{
+				return Item->IsChangingFixturePatch();
+			});
+		if (bAnyItemIsChangingFixturePatch)
+		{
+			return;
+		}
 
-	RequestListRefreshTimerHandle = GEditor->GetTimerManager()->SetTimerForNextTick(FTimerDelegate::CreateSP(this, &SDMXFixturePatchList::RefreshList));
+		RequestListRefreshTimerHandle = GEditor->GetTimerManager()->SetTimerForNextTick(FTimerDelegate::CreateSP(this, &SDMXFixturePatchList::RefreshList));
+	}
 }
 
 void SDMXFixturePatchList::EnterFixturePatchNameEditingMode()
@@ -671,23 +652,27 @@ void SDMXFixturePatchList::OnSearchChanged()
 }
 
 void SDMXFixturePatchList::RefreshList()
-{	
+{
 	RequestListRefreshTimerHandle.Invalidate();
 
-	// Clear cached data
-	Rows.Reset();
-	ListSource.Reset();
-
 	const TSharedPtr<FDMXEditor> DMXEditor = WeakDMXEditor.Pin();
 	UDMXLibrary* DMXLibrary = DMXEditor.IsValid() ? DMXEditor->GetDMXLibrary() : nullptr;
 	if (!DMXLibrary)
 	{
+		ChildSlot
+		[
+			SNullWidget::NullWidget
+		];
+
 		return;
 	}
 
-	// Remove all that are no longer a patch
-	const TArray<UDMXEntityFixturePatch*> FixturePatches = DMXLibrary->GetEntitiesTypeCast<UDMXEntityFixturePatch>();
+	// Clear cached data
+	Rows.Reset();
+	ListSource.Reset();
 
+	// Make a new list source
+	const TArray<UDMXEntityFixturePatch*> FixturePatches = DMXLibrary->GetEntitiesTypeCast<UDMXEntityFixturePatch>();
 	Algo::TransformIf(FixturePatches, ListSource,
 		[](const UDMXEntityFixturePatch* FixturePatch)
 		{
@@ -699,17 +684,40 @@ void SDMXFixturePatchList::RefreshList()
 		});
 	SortListSource(EColumnSortPriority::Max, SortedByColumnID, SortMode);
 
+	Toolbar =
+		SNew(SDMXFixturePatchListToolbar, WeakDMXEditor)
+		.OnSearchChanged(this, &SDMXFixturePatchList::OnSearchChanged);
+
+	// Apply search filters. Relies on up-to-date status to find conflicts.
+	ListSource = Toolbar->FilterItems(ListSource);
+
 	// Generate status texts
 	GenereateStatusText();
 
-	// Apply search filters. Relies on up-to-date status to find conflicts.
-	if (Toolbar.IsValid())
-	{
-		ListSource = Toolbar->FilterItems(ListSource);
-	}
-	
-	// Update and sort the list and its widgets
-	ListView->RebuildList();
+	ChildSlot
+	[
+		SNew(SVerticalBox)
+
+		+ SVerticalBox::Slot()
+		.HAlign(HAlign_Fill)
+		.AutoHeight()
+		[
+			Toolbar.ToSharedRef()
+		]
+
+		+ SVerticalBox::Slot()
+		.HAlign(HAlign_Fill)
+		.FillHeight(1.f)
+		[
+			SAssignNew(ListView, SDMXFixturePatchListType)
+			.ListViewStyle(&FAppStyle::Get().GetWidgetStyle<FTableViewStyle>("TreeView"))
+			.HeaderRow(GenerateHeaderRow())
+			.ListItemsSource(&ListSource)
+			.OnGenerateRow(this, &SDMXFixturePatchList::OnGenerateRow)
+			.OnSelectionChanged(this, &SDMXFixturePatchList::OnSelectionChanged)
+			.OnContextMenuOpening(this, &SDMXFixturePatchList::OnContextMenuOpening)
+		]
+	];
 
 	AdoptSelectionFromFixturePatchSharedData();
 }
@@ -781,21 +789,11 @@ void SDMXFixturePatchList::OnSelectionChanged(TSharedPtr<FDMXFixturePatchListIte
 
 void SDMXFixturePatchList::OnEntityAddedOrRemoved(UDMXLibrary* DMXLibrary, TArray<UDMXEntity*> Entities)
 {
-	if (bChangingDMXLibrary)
-	{
-		return;
-	}
-
 	RequestListRefresh();
 }
 
 void SDMXFixturePatchList::OnFixturePatchChanged(const UDMXEntityFixturePatch* FixturePatch)
 {
-	if (bChangingDMXLibrary)
-	{
-		return;
-	}
-
 	// Refresh only if the fixture patch is in the library this editor handles
 	const UDMXLibrary* DMXLibrary = WeakDMXEditor.IsValid() ? WeakDMXEditor.Pin()->GetDMXLibrary() : nullptr;
 	if (FixturePatch && FixturePatch->GetParentLibrary() == DMXLibrary)
@@ -806,11 +804,6 @@ void SDMXFixturePatchList::OnFixturePatchChanged(const UDMXEntityFixturePatch* F
 
 void SDMXFixturePatchList::OnFixtureTypeChanged(const UDMXEntityFixtureType* FixtureType)
 {
-	if (bChangingDMXLibrary)
-	{
-		return;
-	}
-
 	// Refresh only if the fixture type is in the library this editor handles
 	const UDMXLibrary* DMXLibrary = WeakDMXEditor.IsValid() ? WeakDMXEditor.Pin()->GetDMXLibrary() : nullptr;
 	if (FixtureType && FixtureType->GetParentLibrary() == DMXLibrary)
@@ -821,33 +814,30 @@ void SDMXFixturePatchList::OnFixtureTypeChanged(const UDMXEntityFixtureType* Fix
 
 void SDMXFixturePatchList::OnFixturePatchSharedDataSelectedFixturePatches()
 {
-	if (!bChangingDMXLibrary)
-	{
-		TArray<TWeakObjectPtr<UDMXEntityFixturePatch>> SelectedFixturePatches = FixturePatchSharedData->GetSelectedFixturePatches();
-		SelectedFixturePatches.RemoveAll([](TWeakObjectPtr<UDMXEntityFixturePatch> FixturePatch)
-			{
-				return !FixturePatch.IsValid();
-			});
-
-		TArray<TSharedPtr<FDMXFixturePatchListItem>> NewSelection;
-		for (const TSharedPtr<FDMXFixturePatchListItem>& Item : ListSource)
+	TArray<TWeakObjectPtr<UDMXEntityFixturePatch>> SelectedFixturePatches = FixturePatchSharedData->GetSelectedFixturePatches();
+	SelectedFixturePatches.RemoveAll([](TWeakObjectPtr<UDMXEntityFixturePatch> FixturePatch)
 		{
-			if (SelectedFixturePatches.Contains(Item->GetFixturePatch()))
-			{
-				NewSelection.Add(Item);
-			}
-		}
+			return !FixturePatch.IsValid();
+		});
 
-		if (NewSelection.Num() > 0)
-		{
-			ListView->ClearSelection();
-			ListView->SetItemSelection(NewSelection, true, ESelectInfo::OnMouseClick);
-		}
-		else
+	TArray<TSharedPtr<FDMXFixturePatchListItem>> NewSelection;
+	for (const TSharedPtr<FDMXFixturePatchListItem>& Item : ListSource)
+	{
+		if (SelectedFixturePatches.Contains(Item->GetFixturePatch()))
 		{
-			ListView->ClearSelection();
+			NewSelection.Add(Item);
 		}
 	}
+
+	if (NewSelection.Num() > 0)
+	{
+		ListView->ClearSelection();
+		ListView->SetItemSelection(NewSelection, true, ESelectInfo::OnMouseClick);
+	}
+	else
+	{
+		ListView->ClearSelection();
+	}
 }
 
 void SDMXFixturePatchList::AdoptSelectionFromFixturePatchSharedData()
@@ -1255,7 +1245,6 @@ void SDMXFixturePatchList::OnPasteItems()
 		return;
 	}
 
-	TGuardValue Guard(bChangingDMXLibrary, true);
 	const FText TransactionText = LOCTEXT("PasteFixturePatchesTransaction", "Paste Fixture Patches");
 	const FScopedTransaction PasteTransaction(TransactionText);
 
@@ -1277,8 +1266,6 @@ void SDMXFixturePatchList::OnPasteItems()
 
 		FixturePatchSharedData->SelectUniverse(AssignedToUniverse);
 		FixturePatchSharedData->SelectFixturePatches(WeakPastedFixturePatches);
-	
-		RequestListRefresh();
 	}
 }
 
@@ -1296,8 +1283,6 @@ void SDMXFixturePatchList::OnDuplicateItems()
 		return;
 	}
 
-	TGuardValue Guard(bChangingDMXLibrary, true);
-
 	const FText TransactionText = LOCTEXT("DuplicateFixturePatchesTransaction", "Duplicate Fixture Patches");
 	const FScopedTransaction PasteTransaction(TransactionText);
 	DMXLibrary->PreEditChange(nullptr);
@@ -1342,8 +1327,6 @@ void SDMXFixturePatchList::OnDuplicateItems()
 	Algo::Copy(NewFixturePatches, NewWeakFixturePatches);
 	FixturePatchSharedData->SelectFixturePatches(NewWeakFixturePatches);
 	FixturePatchSharedData->SelectUniverse(AssignedToUniverse);
-
-	RequestListRefresh();
 }
 
 bool SDMXFixturePatchList::CanDeleteItems() const
@@ -1353,8 +1336,6 @@ bool SDMXFixturePatchList::CanDeleteItems() const
 
 void SDMXFixturePatchList::OnDeleteItems()
 {
-	TGuardValue Guard(bChangingDMXLibrary, true);
-
 	const TArray<TSharedPtr<FDMXFixturePatchListItem>> SelectedItems = ListView->GetSelectedItems();
 
 	if (SelectedItems.Num() == 0)
@@ -1403,8 +1384,6 @@ void SDMXFixturePatchList::OnDeleteItems()
 	{
 		ListView->SetSelection(NewSelection, ESelectInfo::OnMouseClick);
 	}
-
-	RequestListRefresh();
 }
 
 #undef LOCTEXT_NAMESPACE
diff --git a/Engine/Plugins/VirtualProduction/DMX/DMXEngine/Source/DMXEditor/Private/Widgets/FixturePatch/SDMXFixturePatchListRow.cpp b/Engine/Plugins/VirtualProduction/DMX/DMXEngine/Source/DMXEditor/Private/Widgets/FixturePatch/SDMXFixturePatchListRow.cpp
--- a/Engine/Plugins/VirtualProduction/DMX/DMXEngine/Source/DMXEditor/Private/Widgets/FixturePatch/SDMXFixturePatchListRow.cpp
+++ b/Engine/Plugins/VirtualProduction/DMX/DMXEngine/Source/DMXEditor/Private/Widgets/FixturePatch/SDMXFixturePatchListRow.cpp
@@ -25,77 +25,58 @@
 class SDMXFixturePatchFixtureTypePicker
 	: public SCompoundWidget
 {
-	DECLARE_DELEGATE_OneParam(FDMXFixturePatchListRowOnFixtureTypeSelectedDelegate, UDMXEntityFixtureType* /** Selected Fixture Type */);
-
 public:
 	SLATE_BEGIN_ARGS(SDMXFixturePatchFixtureTypePicker)
 	{}
-		/** Called when the combo box selection changed */
-		SLATE_EVENT(FDMXFixturePatchListRowOnFixtureTypeSelectedDelegate, OnFixtureTypeSelected)
 
 	SLATE_END_ARGS()
 
 	/** Constructs this widget */
-	void Construct(const FArguments& InArgs, UDMXLibrary* InDMXLibrary)
+	void Construct(const FArguments& InArgs, const TSharedPtr<FDMXFixturePatchListItem>& InItem)
 	{
-		if (!InDMXLibrary)
+		if (!ensureMsgf(InItem.IsValid(), TEXT("Invalid Fixture Patch List Item, cannot draw Fixture Type Picker for patch")))
 		{
 			return;
 		}
+		Item = InItem;
 
-		WeakDMXLibrary = InDMXLibrary;
-		OnFixtureTypeSelectedDelegate = InArgs._OnFixtureTypeSelected;
-
+		UDMXLibrary* DMXLibrary = Item->GetDMXLibrary();
+		const TArray<UDMXEntityFixtureType*> FixtureTypes = DMXLibrary ? DMXLibrary->GetEntitiesTypeCast<UDMXEntityFixtureType>() : TArray<UDMXEntityFixtureType*>{};
 
-		InDMXLibrary->GetOnEntitiesAdded().AddSP(this, &SDMXFixturePatchFixtureTypePicker::OnEntitiesAddedOrRemoved);
-		InDMXLibrary->GetOnEntitiesRemoved().AddSP(this, &SDMXFixturePatchFixtureTypePicker::OnEntitiesAddedOrRemoved);
-		UDMXEntityFixtureType::GetOnFixtureTypeChanged().AddSP(this, &SDMXFixturePatchFixtureTypePicker::OnFixtureTypeChanged);
+		for (UDMXEntityFixtureType* FixtureType : FixtureTypes)
+		{
+			const TSharedPtr<FString> FixtureTypeName = MakeShared<FString>(FixtureType->Name);
 
-		ChildSlot
-			[
-				SAssignNew(ComboBox, SSearchableComboBox)
-				.OptionsSource(&FixtureTypeNames)
-				.OnGenerateWidget(this, &SDMXFixturePatchFixtureTypePicker::OnGenerateWidget)
-				.OnSelectionChanged(this, &SDMXFixturePatchFixtureTypePicker::OnSelectionChanged)
-				.Content()						
-				[
-					SNew(STextBlock)
-					.Text(this, &SDMXFixturePatchFixtureTypePicker::GetSelectedItemText)
-					.Font(FAppStyle::GetFontStyle(TEXT("PropertyWindow.NormalFont")))
-				]
-			];
-
-		RefreshInternal();
-	}
+			ComboBoxSource.Add(FixtureTypeName);
+		}
 
-	/** Sets the currently selected Fixture Type */
-	void SetSelection(UDMXEntityFixtureType* FixtureType)
-	{
-		if (FixtureType)
-		{
-			TSharedPtr<FString> const* FixtureTypeNameToSelectPtr = FixtureTypeNames.FindByPredicate([FixtureType](const TSharedPtr<FString>& FixtureTypeName)
+		// Find an initial selection
+		const TSharedPtr<FString>* InitialSelectionPtr = Algo::FindByPredicate(ComboBoxSource, 
+			[this](const TSharedPtr<FString>& FixtureTypeName)
+			{
+				if (const UDMXEntityFixtureType* FixtureType = Item->GetFixtureType())
 				{
 					return *FixtureTypeName == FixtureType->Name;
-				});
-			if (ensureAlwaysMsgf(FixtureTypeNameToSelectPtr, TEXT("Trying to select a fixture type but Fixture Type %s is not present."), *FixtureType->Name))
-			{
-				ComboBox->SetSelectedItem(*FixtureTypeNameToSelectPtr);
-			}
-		}
-		else
-		{
-			checkf(FixtureTypeNames.Num() > 0, TEXT("Expected at least a 'None' entry in FitureTypesName, but the array is empty."));
-			ComboBox->SetSelectedItem(FixtureTypeNames[0]);
-		}
-	}
+				}
 
-	/** Requests to refresh the widget on the next tick */
-	void RequestRefresh()
-	{
-		if (!RefreshTimerHandle.IsValid())
-		{
-			RefreshTimerHandle = GEditor->GetTimerManager()->SetTimerForNextTick(FTimerDelegate::CreateSP(this, &SDMXFixturePatchFixtureTypePicker::RefreshInternal));
-		}
+				return false;
+			});
+
+		const TSharedPtr<FString> InitialSelection = InitialSelectionPtr ? *InitialSelectionPtr : nullptr;
+
+		ChildSlot
+		[
+			SNew(SSearchableComboBox)
+			.OptionsSource(&ComboBoxSource)
+			.InitiallySelectedItem(InitialSelection)
+			.OnGenerateWidget(this, &SDMXFixturePatchFixtureTypePicker::OnGenerateWidget)
+			.OnSelectionChanged(this, &SDMXFixturePatchFixtureTypePicker::OnSelectionChanged)
+			[
+				SNew(STextBlock)
+				.Text(this, &SDMXFixturePatchFixtureTypePicker::GetSelectedItemText)
+				.Font(FAppStyle::GetFontStyle(TEXT("PropertyWindow.NormalFont")))
+			]
+		];
 	}
 
 private:
@@ -110,127 +91,31 @@ private:
 	/** Called when the combo box selection changed */
 	void OnSelectionChanged(TSharedPtr<FString> NewSelection, ESelectInfo::Type SelectInfo)
 	{
-		if (SelectInfo == ESelectInfo::Direct)
-		{
-			return;
-		}
-
-		if (UDMXEntityFixtureType* SelectedFixtureType = GetSelectedFixtureType())
-		{
-			OnFixtureTypeSelectedDelegate.ExecuteIfBound(SelectedFixtureType);
-		}
-	}
-
-	UDMXEntityFixtureType* GetSelectedFixtureType() const
-	{
-		UDMXEntityFixtureType* SelectedFixtureType = nullptr;
-		const TSharedPtr<FString> SelectedItem = ComboBox->GetSelectedItem();
-		if (TWeakObjectPtr<UDMXEntityFixtureType> const* FixtureTypePtr = FixtureTypeNameToFixtureTypeMap.Find(SelectedItem))
-		{
-			SelectedFixtureType = FixtureTypePtr->Get();
-		}
+		UDMXLibrary* DMXLibrary = Item->GetDMXLibrary();
+		const TArray<UDMXEntityFixtureType*> FixtureTypes = DMXLibrary ? DMXLibrary->GetEntitiesTypeCast<UDMXEntityFixtureType>() : TArray<UDMXEntityFixtureType*>{};
 
-		return SelectedFixtureType;
+		UDMXEntityFixtureType* const* FixtureTypePtr = Algo::FindBy(FixtureTypes, *NewSelection, &UDMXEntityFixtureType::Name);
+		UDMXEntityFixtureType* FixtureType = FixtureTypePtr ? *FixtureTypePtr : nullptr;
+		
+		Item->SetFixtureType(FixtureType);
 	}
 
 	/** Returns the text of the selected item */
 	FText GetSelectedItemText() const
 	{
-		if (ComboBox.IsValid())
+		if (UDMXEntityFixtureType* FixtureType = Item->GetFixtureType())
 		{
-			if (const TSharedPtr<FString> SelectedItem = ComboBox->GetSelectedItem())
-			{
-				return FText::FromString(*SelectedItem);
-			}
+			return FText::FromString(FixtureType->Name);
 		}
 
-		return FText::GetEmpty();
-	}
-
-	/** Refreshes the Combo Box */
-	void RefreshInternal()
-	{
-		RefreshTimerHandle.Invalidate();
-
-		FixtureTypeNames.Reset();
-		const UDMXLibrary* DMXLibrary = WeakDMXLibrary.Get();
-		if (!DMXLibrary)
-		{
-			return;
-		}
-
-		const TArray<UDMXEntityFixtureType*> FixtureTypes = DMXLibrary->GetEntitiesTypeCast<UDMXEntityFixtureType>();
-
-		if (FixtureTypes.IsEmpty())
-		{
-			const TSharedRef<FString> NoneItem = MakeShared<FString>(TEXT("None"));
-			FixtureTypeNames.Add(NoneItem);
-			ComboBox->SetSelectedItem(NoneItem);
-		}
-		else
-		{
-			for (UDMXEntityFixtureType* FixtureType : FixtureTypes)
-			{
-				const TSharedPtr<FString> FixtureTypeName = MakeShared<FString>(FixtureType->Name);
-
-				FixtureTypeNames.Add(FixtureTypeName);
-				FixtureTypeNameToFixtureTypeMap.Add(FixtureTypeName, FixtureType);
-			}
-		}
-
-		if (!ensureMsgf(!FixtureTypeNames.IsEmpty(), TEXT("No combo box option for Fixture Type.")))
-		{
-			return;
-		}
-
-		const FString SelectedFixtureTypeName = ComboBox->GetSelectedItem() ? *ComboBox->GetSelectedItem() : TEXT("");
-		const TSharedPtr<FString>* SelectionPtr = Algo::FindByPredicate(FixtureTypeNames, [SelectedFixtureTypeName](const TSharedPtr<FString>& FixtureTypeName)
-			{
-				return *FixtureTypeName == SelectedFixtureTypeName;
-			});
-		if (!SelectionPtr)
-		{
-			ComboBox->SetSelectedItem(FixtureTypeNames[0], ESelectInfo::Direct);
-		}
-
-		ComboBox->RefreshOptions();
-	}
-
-	/** Called when the Entities in the DMX Library were added or removed */
-	void OnEntitiesAddedOrRemoved(UDMXLibrary* DMXLibrary, TArray<UDMXEntity*> AddedOrRemovedEntities)
-	{
-		if (DMXLibrary == WeakDMXLibrary)
-		{
-			RequestRefresh();
-		}
-	}
-
-	/** Called when a Fixture Type changed */
-	void OnFixtureTypeChanged(const UDMXEntityFixtureType* ChangedFixtureType)
-	{
-		if (ChangedFixtureType && ChangedFixtureType->GetParentLibrary() == WeakDMXLibrary)
-		{
-			RequestRefresh();
-		}
+		return LOCTEXT("NoFixtureType", "None");
 	}
 
 	/** Names of Fixture Types in the Combo Box */
-	TArray<TSharedPtr<FString>> FixtureTypeNames;
-
-	/** Map from the Fixture Type Name to the actual Fixture Type object */
-	TMap<TSharedPtr<FString>, TWeakObjectPtr<UDMXEntityFixtureType>> FixtureTypeNameToFixtureTypeMap;
-
-	/** The combo box widget in use */
-	TSharedPtr<SSearchableComboBox> ComboBox;
-
-	/** The DMX Library from which to select fixture types */
-	TWeakObjectPtr<UDMXLibrary> WeakDMXLibrary;
-
-	/** Timer handle used for the RequestRefresh method */
-	FTimerHandle RefreshTimerHandle;
+	TArray<TSharedPtr<FString>> ComboBoxSource;
 
-	// Slate args
-	FDMXFixturePatchListRowOnFixtureTypeSelectedDelegate OnFixtureTypeSelectedDelegate;
+	/** The Fixture Patch Item for which this Fixture Type Picker is displayed */
+	TSharedPtr<FDMXFixturePatchListItem> Item;
 };
 
 
@@ -241,166 +126,137 @@ private:
 class SDMXFixturePatchModePicker
 	: public SCompoundWidget
 {
-	DECLARE_DELEGATE_OneParam(FDMXFixturePatchListRowOnModeSelectedDelegate, int32 /** Selected Mode Index */);
-
 public:
 	SLATE_BEGIN_ARGS(SDMXFixturePatchModePicker)
 	{}
-		/** Called when the combo box selection changed */
-		SLATE_EVENT(FDMXFixturePatchListRowOnModeSelectedDelegate, OnModeSelected)
 
 	SLATE_END_ARGS()
 
 	/** Constructs this widget */
-	void Construct(const FArguments& InArgs)
+	void Construct(const FArguments& InArgs, const TSharedPtr<FDMXFixturePatchListItem>& InItem)
 	{
-		OnModeSelectedDelegate = InArgs._OnModeSelected;
+		if (!ensureMsgf(InItem.IsValid(), TEXT("Invalid Fixture Patch List Item, cannot draw Mode Picker for patch")))
+		{
+			return;
+		}
+		Item = InItem;
 
-		UDMXEntityFixtureType::GetOnFixtureTypeChanged().AddSP(this, &SDMXFixturePatchModePicker::OnFixtureTypeChanged);
+		UpdateComboBoxSource();
 
 		ChildSlot
+		[
+			SNew(SVerticalBox)
+			.IsEnabled_Lambda([this]
+				{
+					return !ComboBoxSource.IsEmpty();
+				})
+				
+			+ SVerticalBox::Slot()
 			[
-				SNew(SVerticalBox)
-				.IsEnabled_Lambda([this]
+				SAssignNew(ComboBox, STextComboBox)
+				.Visibility_Lambda([this]()
 					{
-						return bEnabled;
+						return Item->GetFixtureType() ? EVisibility::Visible : EVisibility::Collapsed;
 					})
-				
-				+ SVerticalBox::Slot()
-				[
-					SAssignNew(ComboBox, STextComboBox)
-					.Visibility_Lambda([this]()
-						{
-							return WeakFixtureType.IsValid() ? EVisibility::Visible : EVisibility::Hidden;
-						})
-					.OptionsSource(&ModeNames)
-					.OnSelectionChanged(this, &SDMXFixturePatchModePicker::OnSelectionChanged)
-					.Font(FAppStyle::GetFontStyle(TEXT("PropertyWindow.NormalFont")))
-				]
-
-				+ SVerticalBox::Slot()
-				[
-					SNew(STextBlock)
-					.Visibility_Lambda([this]()
-						{
-							return WeakFixtureType.IsValid() ? EVisibility::Collapsed : EVisibility::Visible;
-						})
-					.Text(LOCTEXT("NoModeBecauseNoFixtureTypeSelectedInfo", "No Fixture Type Selected"))
-					.Font(FAppStyle::GetFontStyle(TEXT("PropertyWindow.NormalFont")))
-				]
-			];
-
-		RefreshOptions();
-	}
+				.OptionsSource(&ComboBoxSource)
+				.OnSelectionChanged(this, &SDMXFixturePatchModePicker::OnSelectionChanged)
+				.Font(FAppStyle::GetFontStyle(TEXT("PropertyWindow.NormalFont")))
+			]
 
-	/** Sets the selection */
-	void SetSelection(UDMXEntityFixtureType* FixtureType, int32 ModeIndex)
-	{
-		WeakFixtureType = FixtureType;		
-		RefreshOptions();
+			+ SVerticalBox::Slot()
+			[
+				SNew(STextBlock)
+				.Visibility_Lambda([this]()
+					{
+						return Item->GetFixtureType() ? EVisibility::Collapsed : EVisibility::Visible;
+					})
+				.Text(LOCTEXT("NoModeBecauseNoFixtureTypeSelectedInfo", "No Fixture Type Selected"))
+				.Font(FAppStyle::GetFontStyle(TEXT("PropertyWindow.NormalFont")))
+			]
+		];
 
-		if (FixtureType && FixtureType->Modes.IsValidIndex(ModeIndex))
-		{
-			if (ensureAlwaysMsgf(ModeNames.IsValidIndex(ModeIndex), TEXT("Trying to select a Mode but its index is not valid.")))
-			{
-				ComboBox->SetSelectedItem(ModeNames[ModeIndex]);
-			}
-		}
+		AdoptSelectionFromFixturePatch();
+
+		UDMXEntityFixturePatch::GetOnFixturePatchChanged().AddSP(this, &SDMXFixturePatchModePicker::OnFixturePatchChanged);
 	}
 
 private:
 	/** Called when the combo box selection changed */
 	void OnSelectionChanged(TSharedPtr<FString> NewSelection, ESelectInfo::Type SelectInfo)
 	{
-		UDMXEntityFixtureType* FixtureType = WeakFixtureType.Get();
-		if (ensureAlwaysMsgf(FixtureType, TEXT("Tried to select Mode, but Fixture Type is no longer valid.")))
+		if (const UDMXEntityFixtureType* FixtureType = Item->GetFixtureType())
 		{
-			const TSharedPtr<FString> SelectedItem = ComboBox->GetSelectedItem();
-			int32 const* ModeIndexPtr = ModeNameToModeIndexMap.Find(SelectedItem);
-			if (ensureAlwaysMsgf(ModeIndexPtr, TEXT("Tried to select Mode, but the Mode Index is no longer valid.")))
-			{
-				OnModeSelectedDelegate.ExecuteIfBound(*ModeIndexPtr);
-			}
-		}
-	}
+			const int32 ModeIndex = FixtureType->Modes.IndexOfByPredicate([&NewSelection](const FDMXFixtureMode& Mode)
+				{
+					return Mode.ModeName == *NewSelection;
+				});
 
-	/** Returns the text of the selected item */
-	FText GetSelectedItemText() const
-	{
-		const TSharedPtr<FString> SelectedItem = ComboBox->GetSelectedItem();
-		return FText::FromString(*SelectedItem);
+			Item->SetModeIndex(ModeIndex);
+		}
 	}
-
-	/** Refreshes the Combo Box */
-	void RefreshOptions()
+	
+	/** Updates the combo box source */
+	void UpdateComboBoxSource()
 	{
-		// Remember selection
-		const FString PreviousSelectedModeName = ComboBox->GetSelectedItem().IsValid() ? *ComboBox->GetSelectedItem() : FString();
-
-		ModeNames.Reset();
-
-		if (UDMXEntityFixtureType* FixtureType = WeakFixtureType.Get())
+		if (UDMXEntityFixtureType* FixtureType = Item->GetFixtureType())
 		{
 			for (int32 ModeIndex = 0; ModeIndex < FixtureType->Modes.Num(); ModeIndex++)
 			{
 				const TSharedPtr<FString> ModeName = MakeShared<FString>(FixtureType->Modes[ModeIndex].ModeName);
 
-				ModeNames.Add(ModeName);
-				ModeNameToModeIndexMap.Add(ModeName, ModeIndex);
+				ComboBoxSource.Add(ModeName);
 			}
 		}
+	}
 
-		if (ModeNames.Num() == 0)
-		{
-			const FText NoModeAvailableText = LOCTEXT("NoModeAvailableInModeComboBoxText", "No Mode available");
-			ModeNames.Add(MakeShared<FString>(NoModeAvailableText.ToString()));
-
-			bEnabled = false;
-		}
-		else
+	/** Adopts the currently selected Item from the Mode of the Fixture Patch */
+	void AdoptSelectionFromFixturePatch()
+	{
+		FString ActiveModeName;
+		if (Item->GetActiveModeName(ActiveModeName))
 		{
-			bEnabled = true;
-		}
+			const TSharedPtr<FString>* SelectionPtr = Algo::FindByPredicate(ComboBoxSource,
+				[this](const TSharedPtr<FString>& ModeName)
+				{
+					FString ActiveModeName;
+					if (Item->GetActiveModeName(ActiveModeName))
+					{
+						return *ModeName == ActiveModeName;
+					}
 
-		ComboBox->RefreshOptions();
+					return false;
+				});
 
-		// Restore selection
-		TSharedPtr<FString>* PreviousSelectionPtr = Algo::FindByPredicate(ModeNames, [PreviousSelectedModeName](const TSharedPtr<FString>& ModeName)
+			if (SelectionPtr)
 			{
-				return ModeName.IsValid() && *ModeName == PreviousSelectedModeName;
-			});
-		if (PreviousSelectionPtr)
-		{
-			ComboBox->SetSelectedItem(*PreviousSelectionPtr);
+				ComboBox->SetSelectedItem(*SelectionPtr);
+				return;
+			}
 		}
+
+		ComboBox->ClearSelection();
 	}
 
-	/** Called when the Entities in the DMX Library were added or removed */
-	void OnFixtureTypeChanged(const UDMXEntityFixtureType* ChangedFixtureType)
+	/** Called when the fixture patch for which modes are displayed changed */
+	void OnFixturePatchChanged(const UDMXEntityFixturePatch* FixturePatch)
 	{
-		if (ChangedFixtureType == WeakFixtureType)
+		if (Item->GetFixturePatch() == FixturePatch)
 		{
-			RefreshOptions();
+			UpdateComboBoxSource();
+			ComboBox->RefreshOptions();
+
+			AdoptSelectionFromFixturePatch();
 		}
 	}
 
-	/** True if the Combo Box should be enabled */
-	bool bEnabled = true;
-
-	/** Names of Modes in the Combo Box */
-	TArray<TSharedPtr<FString>> ModeNames;
-
-	/** Map from the Fixture Type Name to the actual Mode Index */
-	TMap<TSharedPtr<FString>, int32> ModeNameToModeIndexMap;
-
-	/** The combo box widget in use */
+	/** The combo box to select a mode */
 	TSharedPtr<STextComboBox> ComboBox;
 
-	/** The Fixture Type for which the Modes are currently displayed */
-	TWeakObjectPtr<UDMXEntityFixtureType> WeakFixtureType;
+	/** Names of Fixture Types in the Combo Box */
+	TArray<TSharedPtr<FString>> ComboBoxSource;
 
-	// Slate args
-	FDMXFixturePatchListRowOnModeSelectedDelegate OnModeSelectedDelegate;
+	/** The Fixture Patch Item for which this Fixture Type Picker is displayed */
+	TSharedPtr<FDMXFixturePatchListItem> Item;
 };
 
 
@@ -623,19 +479,6 @@ void SDMXFixturePatchListRow::OnFixtureIDCommitted(const FText& InNewText, EText
 
 TSharedRef<SWidget> SDMXFixturePatchListRow::GenerateFixtureTypeWidget()
 {
-	UDMXLibrary* DMXLibrary = Item->GetDMXLibrary();
-	if (!ensureAlwaysMsgf(DMXLibrary, TEXT("Tried to set fixture type for MVR Fixture, but fixture type is invalid.")))
-	{
-		return SNullWidget::NullWidget;
-	}
-
-	const TSharedRef<SDMXFixturePatchFixtureTypePicker> FixtureTypePicker =
-		SNew(SDMXFixturePatchFixtureTypePicker, DMXLibrary)
-		.OnFixtureTypeSelected(this, &SDMXFixturePatchListRow::OnFixtureTypeSelected);
-
-	UDMXEntityFixtureType* SelectedFixtureType = Item->GetFixtureType();
-	FixtureTypePicker->SetSelection(SelectedFixtureType);
-		
 	return 
 		SNew(SBorder)
 		.HAlign(HAlign_Fill)
@@ -643,26 +486,12 @@ TSharedRef<SWidget> SDMXFixturePatchListRow::GenerateFixtureTypeWidget()
 		.Padding(4.f)
 		.BorderImage(FAppStyle::GetBrush("NoBorder"))
 		[
-			FixtureTypePicker
+			SNew(SDMXFixturePatchFixtureTypePicker, Item)
 		];
 }
 
-void SDMXFixturePatchListRow::OnFixtureTypeSelected(UDMXEntityFixtureType* SelectedFixtureType)
-{
-	Item->SetFixtureType(SelectedFixtureType);
-}
-
 TSharedRef<SWidget> SDMXFixturePatchListRow::GenerateModeWidget()
 {
-	const TSharedRef<SDMXFixturePatchModePicker> ModePicker =
-		SNew(SDMXFixturePatchModePicker)
-		.OnModeSelected(this, &SDMXFixturePatchListRow::OnModeSelected);
-
-	UDMXEntityFixtureType* SelectedFixtureType = Item->GetFixtureType();
-	const int32 SelectedModeIndex = Item->GetModeIndex();
-
-	ModePicker->SetSelection(SelectedFixtureType, SelectedModeIndex);
-
 	return
 		SNew(SBorder)
 		.HAlign(HAlign_Fill)
@@ -670,15 +499,10 @@ TSharedRef<SWidget> SDMXFixturePatchListRow::GenerateModeWidget()
 		.Padding(4.f)
 		.BorderImage(FAppStyle::GetBrush("NoBorder"))
 		[
-			ModePicker
+			SNew(SDMXFixturePatchModePicker, Item)
 		];
 }
 
-void SDMXFixturePatchListRow::OnModeSelected(int32 SelectedModeIndex)
-{
-	Item->SetModeIndex(SelectedModeIndex);
-}
-
 TSharedRef<SWidget> SDMXFixturePatchListRow::GeneratePatchWidget()
 {
 	return
@@ -709,7 +533,7 @@ TSharedRef<SWidget> SDMXFixturePatchListRow::GeneratePatchWidget()
 						}
 					})
 				.Font(FAppStyle::GetFontStyle(TEXT("PropertyWindow.NormalFont")))
-				.OnTextCommitted(this, &SDMXFixturePatchListRow::OnPatchCommitted)
+				.OnTextCommitted(this, &SDMXFixturePatchListRow::OnPatchNameCommitted)
 				.IsSelected(IsSelected)
 			]
 		];
@@ -728,7 +552,7 @@ FReply SDMXFixturePatchListRow::OnPatchBorderDoubleClicked(const FGeometry& InMy
 	return FReply::Handled();
 }
 
-void SDMXFixturePatchListRow::OnPatchCommitted(const FText& InNewText, ETextCommit::Type InTextCommit)
+void SDMXFixturePatchListRow::OnPatchNameCommitted(const FText& InNewText, ETextCommit::Type InTextCommit)
 {
 	const FString PatchString = InNewText.ToString();
 	static const TCHAR* ParamDelimiters[] =
diff --git a/Engine/Plugins/VirtualProduction/DMX/DMXEngine/Source/DMXEditor/Private/Widgets/FixturePatch/SDMXFixturePatchListRow.h b/Engine/Plugins/VirtualProduction/DMX/DMXEngine/Source/DMXEditor/Private/Widgets/FixturePatch/SDMXFixturePatchListRow.h
--- a/Engine/Plugins/VirtualProduction/DMX/DMXEngine/Source/DMXEditor/Private/Widgets/FixturePatch/SDMXFixturePatchListRow.h
+++ b/Engine/Plugins/VirtualProduction/DMX/DMXEngine/Source/DMXEditor/Private/Widgets/FixturePatch/SDMXFixturePatchListRow.h
@@ -71,15 +71,9 @@ private:
 	/** Generates the Widget that displays the Fixture Type */
 	TSharedRef<SWidget> GenerateFixtureTypeWidget();
 
-	/** Called when a Fixture Type was selected */
-	void OnFixtureTypeSelected(UDMXEntityFixtureType* SelectedFixtureType);
-
 	/** Generates the Widget that displays the Mode */
 	TSharedRef<SWidget> GenerateModeWidget();
 
-	/** Called when a Mode was selected */
-	void OnModeSelected(int32 SelectedModeIndex);
-
 	/** Generates the Widget that displays the Patch */
 	TSharedRef<SWidget> GeneratePatchWidget();
 
@@ -87,7 +81,7 @@ private:
 	FReply OnPatchBorderDoubleClicked(const FGeometry& InMyGeometry, const FPointerEvent& InMouseEvent);
 
 	/** Called when a Patch was committed */
-	void OnPatchCommitted(const FText& InNewText, ETextCommit::Type InTextCommit);
+	void OnPatchNameCommitted(const FText& InNewText, ETextCommit::Type InTextCommit);
 
 	/** The Edit Mode the Widget should present */
 	EDMXFixturePatchListEditMode EditMode;
diff --git a/Engine/Plugins/VirtualProduction/DMX/DMXEngine/Source/DMXEditor/Private/Widgets/FixturePatch/SDMXFixturePatchListToolbar.cpp b/Engine/Plugins/VirtualProduction/DMX/DMXEngine/Source/DMXEditor/Private/Widgets/FixturePatch/SDMXFixturePatchListToolbar.cpp
--- a/Engine/Plugins/VirtualProduction/DMX/DMXEngine/Source/DMXEditor/Private/Widgets/FixturePatch/SDMXFixturePatchListToolbar.cpp
+++ b/Engine/Plugins/VirtualProduction/DMX/DMXEngine/Source/DMXEditor/Private/Widgets/FixturePatch/SDMXFixturePatchListToolbar.cpp
@@ -12,8 +12,8 @@
 #include "FixturePatchAutoAssignUtility.h"
 #include "Library/DMXEntityFixturePatch.h"
 #include "Library/DMXEntityFixtureType.h"
-#include "SAddFixturePatchMenu.h"
 #include "ScopedTransaction.h"
+#include "SDMXAddFixturePatchMenu.h"
 #include "Widgets/Input/SCheckBox.h"
 #include "Widgets/Input/SEditableTextBox.h"
 #include "Widgets/Input/SSearchBox.h"
@@ -201,7 +201,7 @@ TSharedRef<SWidget> SDMXFixturePatchListToolbar::GenerateFixtureTypeDropdownMenu
 				.Padding(4.f)
 				.MinDesiredWidth(460.f)
 				[
-					SAssignNew(AddFixturePatchMenu, UE::DMXEditor::FixturePatchEditor::SAddFixturePatchMenu, WeakDMXEditor)
+					SAssignNew(AddFixturePatchMenu, UE::DMXEditor::FixturePatchEditor::SDMXAddFixturePatchMenu, WeakDMXEditor)
 				]
 			]
 			.IsFocusable(true)
diff --git a/Engine/Plugins/VirtualProduction/DMX/DMXEngine/Source/DMXEditor/Private/Widgets/FixturePatch/SDMXFixturePatchListToolbar.h b/Engine/Plugins/VirtualProduction/DMX/DMXEngine/Source/DMXEditor/Private/Widgets/FixturePatch/SDMXFixturePatchListToolbar.h
--- a/Engine/Plugins/VirtualProduction/DMX/DMXEngine/Source/DMXEditor/Private/Widgets/FixturePatch/SDMXFixturePatchListToolbar.h
+++ b/Engine/Plugins/VirtualProduction/DMX/DMXEngine/Source/DMXEditor/Private/Widgets/FixturePatch/SDMXFixturePatchListToolbar.h
@@ -11,7 +11,7 @@ class FDMXFixturePatchListItem;
 template<typename TEntityType> class SDMXEntityDropdownMenu;
 class UDMXEntity;
 class UDMXEntityFixtureType;
-namespace UE::DMXEditor::FixturePatchEditor { class SAddFixturePatchMenu; }
+namespace UE::DMXEditor::FixturePatchEditor { class SDMXAddFixturePatchMenu; }
 
 
 /** Search bar for the MVR Fixture List */
@@ -50,7 +50,7 @@ private:
 	bool bShowConfictsOnly = false;
 
 	/** Menu displayed to add fixture patches */
-	TSharedPtr<UE::DMXEditor::FixturePatchEditor::SAddFixturePatchMenu> AddFixturePatchMenu;
+	TSharedPtr<UE::DMXEditor::FixturePatchEditor::SDMXAddFixturePatchMenu> AddFixturePatchMenu;
 
 	/** The DMX Editor that owns this widget */
 	TWeakPtr<FDMXEditor> WeakDMXEditor;
