DMX: Fix DMX Library editor related issues

* Fix regression crash when adding new Fixture Type
* Undo deleting Fixture Types now restores the Fixture Types of patches
* Updating the Fixture Patch list view retains size of the toolbar

#rnx

diff --git a/Engine/Plugins/VirtualProduction/DMX/DMXEngine/Source/DMXEditor/Private/Customizations/DMXEntityFixturePatchDetails.cpp b/Engine/Plugins/VirtualProduction/DMX/DMXEngine/Source/DMXEditor/Private/Customizations/DMXEntityFixturePatchDetails.cpp
--- a/Engine/Plugins/VirtualProduction/DMX/DMXEngine/Source/DMXEditor/Private/Customizations/DMXEntityFixturePatchDetails.cpp
+++ b/Engine/Plugins/VirtualProduction/DMX/DMXEngine/Source/DMXEditor/Private/Customizations/DMXEntityFixturePatchDetails.cpp
@@ -59,20 +59,6 @@ void FDMXEntityFixturePatchDetails::CustomizeDetails(IDetailLayoutBuilder& Detai
 	TSharedPtr<uint32> DefaultSelectedActiveMode = nullptr;
 	GenerateActiveModesSource();
 
-	int32 ActiveMode;
-	if (ActiveModeHandle->GetValue(ActiveMode) == FPropertyAccess::Success)
-	{
-		const bool bActiveModeExists = ActiveModesSource.ContainsByPredicate([ActiveMode](TSharedPtr<uint32> Option) 
-			{
-				return Option.IsValid() && *Option == ActiveMode;
-			});
-
-		if (!bActiveModeExists)
-		{
-			SetActiveMode(0);
-		}
-	}
-
 	DetailBuilder.EditDefaultProperty(ActiveModeHandle)->CustomWidget(false)
 		.NameContent()
 		[
@@ -125,21 +111,7 @@ void FDMXEntityFixturePatchDetails::OnParentFixtureTypeChanged(UDMXEntity* NewTe
 
 void FDMXEntityFixturePatchDetails::OnFixtureTypeChanged(const UDMXEntityFixtureType* FixtureType)
 {
-	if (IsValid(FixtureType))
-	{
-		// Keep the active mode valid
-		int32 ActiveMode;
-		if (ActiveModeHandle->GetValue(ActiveMode) == FPropertyAccess::Success)
-		{
-			if (!FixtureType->Modes.IsValidIndex(ActiveMode))
-			{
-				const int32 NewActiveMode = FixtureType->Modes.Num() > 0 ? 0 : INDEX_NONE;
-				ActiveModeHandle->SetValue(NewActiveMode);
-			}
-		}
-	
-		PropertyUtilities->RequestRefresh();
-	}
+	PropertyUtilities->RequestRefresh();
 }
 
 void FDMXEntityFixturePatchDetails::OnUniverseIDChanged()
diff --git a/Engine/Plugins/VirtualProduction/DMX/DMXEngine/Source/DMXEditor/Private/Widgets/FixturePatch/SDMXFixturePatchList.cpp b/Engine/Plugins/VirtualProduction/DMX/DMXEngine/Source/DMXEditor/Private/Widgets/FixturePatch/SDMXFixturePatchList.cpp
--- a/Engine/Plugins/VirtualProduction/DMX/DMXEngine/Source/DMXEditor/Private/Widgets/FixturePatch/SDMXFixturePatchList.cpp
+++ b/Engine/Plugins/VirtualProduction/DMX/DMXEngine/Source/DMXEditor/Private/Widgets/FixturePatch/SDMXFixturePatchList.cpp
@@ -582,16 +582,30 @@ void SDMXFixturePatchList::Construct(const FArguments& InArgs, TWeakPtr<FDMXEdit
 	// Handle Shared Data selection changes
 	FixturePatchSharedData->OnFixturePatchSelectionChanged.AddSP(this, &SDMXFixturePatchList::OnFixturePatchSharedDataSelectedFixturePatches);
 
-	RegisterCommands();
-	RefreshList();
+	ChildSlot
+	[
+		SNew(SVerticalBox)
 
-	AdoptSelectionFromFixturePatchSharedData();
+		+ SVerticalBox::Slot()
+		.HAlign(HAlign_Fill)
+		.AutoHeight()
+		[
+			SAssignNew(Toolbar, SDMXFixturePatchListToolbar, WeakDMXEditor)
+			.OnSearchChanged(this, &SDMXFixturePatchList::OnSearchChanged)
+		]
 
-	// Make an initial selection if nothing was selected from Fixture Patch Shared Data, as if the user clicked it
-	if (ListView->GetSelectedItems().IsEmpty() && !ListSource.IsEmpty())
-	{
-		ListView->SetSelection(ListSource[0], ESelectInfo::OnMouseClick);
-	}
+		+ SVerticalBox::Slot()
+		.HAlign(HAlign_Fill)
+		.FillHeight(1.f)
+		[
+			SAssignNew(ListContentBorder, SBorder)
+			.BorderImage(FAppStyle::GetBrush("NoBorder"))
+		]
+	];
+
+	RefreshList();
+
+	RegisterCommands();
 }
 
 void SDMXFixturePatchList::RequestListRefresh()
@@ -614,19 +628,16 @@ void SDMXFixturePatchList::RequestListRefresh()
 
 void SDMXFixturePatchList::EnterFixturePatchNameEditingMode()
 {
-	if (ListView.IsValid())
+	const TArray<TSharedPtr<FDMXFixturePatchListItem>> SelectedItems = ListView->GetSelectedItems();
+	if (SelectedItems.Num() == 0)
 	{
-		const TArray<TSharedPtr<FDMXFixturePatchListItem>> SelectedItems = ListView->GetSelectedItems();
-		if (SelectedItems.Num() == 0)
-		{
-			const TSharedPtr<SDMXFixturePatchListRow>* SelectedRowPtr = Rows.FindByPredicate([&SelectedItems](const TSharedPtr<SDMXFixturePatchListRow>& Row)
-				{
-					return Row->GetItem() == SelectedItems[0];
-				});
-			if (SelectedRowPtr)
+		const TSharedPtr<SDMXFixturePatchListRow>* SelectedRowPtr = Rows.FindByPredicate([&SelectedItems](const TSharedPtr<SDMXFixturePatchListRow>& Row)
 			{
-				(*SelectedRowPtr)->EnterFixturePatchNameEditingMode();
-			}
+				return Row->GetItem() == SelectedItems[0];
+			});
+		if (SelectedRowPtr)
+		{
+			(*SelectedRowPtr)->EnterFixturePatchNameEditingMode();
 		}
 	}
 }	
@@ -684,40 +695,20 @@ void SDMXFixturePatchList::RefreshList()
 		});
 	SortListSource(EColumnSortPriority::Max, SortedByColumnID, SortMode);
 
-	Toolbar =
-		SNew(SDMXFixturePatchListToolbar, WeakDMXEditor)
-		.OnSearchChanged(this, &SDMXFixturePatchList::OnSearchChanged);
-
 	// Apply search filters. Relies on up-to-date status to find conflicts.
 	ListSource = Toolbar->FilterItems(ListSource);
 
 	// Generate status texts
 	GenereateStatusText();
 
-	ChildSlot
-	[
-		SNew(SVerticalBox)
-
-		+ SVerticalBox::Slot()
-		.HAlign(HAlign_Fill)
-		.AutoHeight()
-		[
-			Toolbar.ToSharedRef()
-		]
-
-		+ SVerticalBox::Slot()
-		.HAlign(HAlign_Fill)
-		.FillHeight(1.f)
-		[
-			SAssignNew(ListView, SDMXFixturePatchListType)
-			.ListViewStyle(&FAppStyle::Get().GetWidgetStyle<FTableViewStyle>("TreeView"))
-			.HeaderRow(GenerateHeaderRow())
-			.ListItemsSource(&ListSource)
-			.OnGenerateRow(this, &SDMXFixturePatchList::OnGenerateRow)
-			.OnSelectionChanged(this, &SDMXFixturePatchList::OnSelectionChanged)
-			.OnContextMenuOpening(this, &SDMXFixturePatchList::OnContextMenuOpening)
-		]
-	];
+	ListContentBorder->SetContent(
+		SAssignNew(ListView, SDMXFixturePatchListType)
+		.ListViewStyle(&FAppStyle::Get().GetWidgetStyle<FTableViewStyle>("TreeView"))
+		.HeaderRow(GenerateHeaderRow())
+		.ListItemsSource(&ListSource)
+		.OnGenerateRow(this, &SDMXFixturePatchList::OnGenerateRow)
+		.OnSelectionChanged(this, &SDMXFixturePatchList::OnSelectionChanged)
+		.OnContextMenuOpening(this, &SDMXFixturePatchList::OnContextMenuOpening));
 
 	AdoptSelectionFromFixturePatchSharedData();
 }
@@ -842,11 +833,6 @@ void SDMXFixturePatchList::OnFixturePatchSharedDataSelectedFixturePatches()
 
 void SDMXFixturePatchList::AdoptSelectionFromFixturePatchSharedData()
 {
-	if (!ListView.IsValid())
-	{
-		return;
-	}
-
 	const TArray<TWeakObjectPtr<UDMXEntityFixturePatch>> SelectedFixturePatches = FixturePatchSharedData->GetSelectedFixturePatches();
 
 	TArray<TSharedPtr<FDMXFixturePatchListItem>> NewSelection;
@@ -871,6 +857,11 @@ void SDMXFixturePatchList::AdoptSelectionFromFixturePatchSharedData()
 		ListView->SetItemSelection(NewSelection, bSelected, ESelectInfo::OnMouseClick);
 		ListView->RequestScrollIntoView(NewSelection[0]);
 	}
+	else if (ListView->GetSelectedItems().IsEmpty() && !ListSource.IsEmpty())
+	{	
+		// Make an initial selection if nothing was selected from Fixture Patch Shared Data, as if the user clicked it
+		ListView->SetSelection(ListSource[0], ESelectInfo::OnMouseClick);
+	}
 }
 
 void SDMXFixturePatchList::AutoAssignFixturePatches(UE::DMXEditor::AutoAssign::EAutoAssignMode Mode)
@@ -1095,10 +1086,7 @@ void SDMXFixturePatchList::SortList(const EColumnSortPriority::Type SortPriority
 {
 	SortListSource(SortPriority, ColumnId, InSortMode);
 
-	if (ListView.IsValid())
-	{
-		ListView->RequestListRefresh();
-	}
+	ListView->RequestListRefresh();
 
 	UDMXEditorSettings* EditorSettings = GetMutableDefault<UDMXEditorSettings>();
 	EditorSettings->MVRFixtureListSettings.SortByCollumnID = SortedByColumnID;
@@ -1109,11 +1097,6 @@ void SDMXFixturePatchList::SortList(const EColumnSortPriority::Type SortPriority
 
 TSharedPtr<SWidget> SDMXFixturePatchList::OnContextMenuOpening()
 {
-	if (!ListView.IsValid())
-	{
-		return SNullWidget::NullWidget;
-	}
-
 	const bool bCloseWindowAfterMenuSelection = true;
 	FMenuBuilder MenuBuilder(bCloseWindowAfterMenuSelection, CommandList);
 	MenuBuilder.BeginSection("BasicOperationsSection", LOCTEXT("BasicOperationsSection", "Basic Operations"));
diff --git a/Engine/Plugins/VirtualProduction/DMX/DMXEngine/Source/DMXEditor/Private/Widgets/FixturePatch/SDMXFixturePatchList.h b/Engine/Plugins/VirtualProduction/DMX/DMXEngine/Source/DMXEditor/Private/Widgets/FixturePatch/SDMXFixturePatchList.h
--- a/Engine/Plugins/VirtualProduction/DMX/DMXEngine/Source/DMXEditor/Private/Widgets/FixturePatch/SDMXFixturePatchList.h
+++ b/Engine/Plugins/VirtualProduction/DMX/DMXEngine/Source/DMXEditor/Private/Widgets/FixturePatch/SDMXFixturePatchList.h
@@ -158,6 +158,9 @@ private:
 	/** Rows of Mode widgets in the List */
 	TArray<TSharedPtr<SDMXFixturePatchListRow>> Rows;
 
+	/** The border that holds the list */
+	TSharedPtr<SBorder> ListContentBorder;
+
 	/** Timer handle for the Request List Refresh method */
 	FTimerHandle RequestListRefreshTimerHandle;
 
diff --git a/Engine/Plugins/VirtualProduction/DMX/DMXEngine/Source/DMXEditor/Private/Widgets/FixturePatch/SDMXFixturePatchListRow.cpp b/Engine/Plugins/VirtualProduction/DMX/DMXEngine/Source/DMXEditor/Private/Widgets/FixturePatch/SDMXFixturePatchListRow.cpp
--- a/Engine/Plugins/VirtualProduction/DMX/DMXEngine/Source/DMXEditor/Private/Widgets/FixturePatch/SDMXFixturePatchListRow.cpp
+++ b/Engine/Plugins/VirtualProduction/DMX/DMXEngine/Source/DMXEditor/Private/Widgets/FixturePatch/SDMXFixturePatchListRow.cpp
@@ -184,7 +184,11 @@ private:
 	/** Called when the combo box selection changed */
 	void OnSelectionChanged(TSharedPtr<FString> NewSelection, ESelectInfo::Type SelectInfo)
 	{
-		if (const UDMXEntityFixtureType* FixtureType = Item->GetFixtureType())
+		if (!NewSelection.IsValid())
+		{
+			Item->SetModeIndex(INDEX_NONE);
+		}
+		else if (const UDMXEntityFixtureType* FixtureType = Item->GetFixtureType())
 		{
 			const int32 ModeIndex = FixtureType->Modes.IndexOfByPredicate([&NewSelection](const FDMXFixtureMode& Mode)
 				{
diff --git a/Engine/Plugins/VirtualProduction/DMX/DMXEngine/Source/DMXEditor/Private/Widgets/FixtureType/SDMXFixtureTypeTree.cpp b/Engine/Plugins/VirtualProduction/DMX/DMXEngine/Source/DMXEditor/Private/Widgets/FixtureType/SDMXFixtureTypeTree.cpp
--- a/Engine/Plugins/VirtualProduction/DMX/DMXEngine/Source/DMXEditor/Private/Widgets/FixtureType/SDMXFixtureTypeTree.cpp
+++ b/Engine/Plugins/VirtualProduction/DMX/DMXEngine/Source/DMXEditor/Private/Widgets/FixtureType/SDMXFixtureTypeTree.cpp
@@ -2,31 +2,28 @@
 
 #include "SDMXFixtureTypeTree.h"
 
+#include "Algo/Transform.h"
+#include "Commands/DMXEditorCommands.h"
+#include "Dialogs/Dialogs.h"
 #include "DMXEditor.h"
 #include "DMXEditorUtils.h"
 #include "DMXFixturePatchSharedData.h"
 #include "DMXFixtureTypeSharedData.h"
 #include "DMXRuntimeUtils.h"
-#include "SDMXFixtureTypeTreeCategoryRow.h"
-#include "SDMXFixtureTypeTreeFixtureTypeRow.h"
-#include "Commands/DMXEditorCommands.h"
+#include "Framework/Commands/GenericCommands.h"
+#include "Framework/MultiBox/MultiBoxBuilder.h"
 #include "Library/DMXEntityFixturePatch.h"
 #include "Library/DMXEntityFixtureType.h"
 #include "Library/DMXLibrary.h"
-#include "Widgets/DMXEntityTreeNode.h"
-
-#include "ScopedTransaction.h"
-#include "Dialogs/Dialogs.h"
-#include "Framework/Application/SlateApplication.h"
-#include "Framework/Commands/GenericCommands.h"
-#include "Framework/MultiBox/MultiBoxBuilder.h"
 #include "Misc/MessageDialog.h"
-#include "Widgets/SBoxPanel.h"
-#include "Widgets/Input/SButton.h"
+#include "ScopedTransaction.h"
+#include "SDMXFixtureTypeTreeCategoryRow.h"
+#include "SDMXFixtureTypeTreeFixtureTypeRow.h"
+#include "Widgets/DMXEntityTreeNode.h"
 #include "Widgets/Images/SImage.h"
+#include "Widgets/Input/SButton.h"
+#include "Widgets/SBoxPanel.h"
 #include "Widgets/Text/STextBlock.h"
-#include "Widgets/Views/STreeView.h"
-
 
 #define LOCTEXT_NAMESPACE "SDMXFixtureTypeTree"
 
@@ -468,9 +465,26 @@ void SDMXFixtureTypeTree::OnDeleteNodes()
 			FixtureTypeSharedData->SelectFixtureTypes(TArray<TWeakObjectPtr<UDMXEntityFixtureType>>());
 		}
 
-		// Remove the Fixture Types from the DMX Library
 		const FScopedTransaction Transaction(EntitiesToDelete.Num() > 1 ? LOCTEXT("RemoveEntities", "Remove Entities") : LOCTEXT("RemoveEntity", "Remove Entity"));
 
+		// Modify affected fixture patches
+		const TArray<UDMXEntityFixturePatch*> FixturePatches = DMXLibrary->GetEntitiesTypeCast<UDMXEntityFixturePatch>();
+		TArray<UDMXEntityFixturePatch*> AffectedFixturePatches;
+		Algo::TransformIf(FixturePatches, AffectedFixturePatches,
+			[&FixtureTypesToDelete](const UDMXEntityFixturePatch* FixturePatch)
+			{
+				return FixturePatch && FixtureTypesToDelete.Contains(FixturePatch->GetFixtureType());
+			},
+			[&FixtureTypesToDelete](UDMXEntityFixturePatch* FixturePatch)
+			{
+				return FixturePatch;
+			});
+		for (UDMXEntityFixturePatch* FixturePatch : AffectedFixturePatches)
+		{
+			FixturePatch->Modify();
+		}
+
+		// Remove the Fixture Types from the DMX 
 		DMXLibrary->PreEditChange(UDMXLibrary::StaticClass()->FindPropertyByName(UDMXLibrary::GetEntitiesPropertyName()));
 		for (UDMXEntityFixtureType* FixtureType : FixtureTypesToDelete)
 		{
@@ -480,7 +494,6 @@ void SDMXFixtureTypeTree::OnDeleteNodes()
 		}
 		DMXLibrary->PostEditChange();
 
-
 		UpdateTree();
 	}
 }
diff --git a/Engine/Plugins/VirtualProduction/DMX/DMXEngine/Source/DMXRuntime/Private/Library/DMXEntityFixturePatch.cpp b/Engine/Plugins/VirtualProduction/DMX/DMXEngine/Source/DMXRuntime/Private/Library/DMXEntityFixturePatch.cpp
--- a/Engine/Plugins/VirtualProduction/DMX/DMXEngine/Source/DMXRuntime/Private/Library/DMXEntityFixturePatch.cpp
+++ b/Engine/Plugins/VirtualProduction/DMX/DMXEngine/Source/DMXRuntime/Private/Library/DMXEntityFixturePatch.cpp
@@ -517,10 +517,14 @@ bool UDMXEntityFixturePatch::IsValidEntity(FText& OutReason) const
 
 void UDMXEntityFixturePatch::ValidateActiveMode()
 {
-	if (ParentFixtureTypeTemplate != nullptr)
+	if (ParentFixtureTypeTemplate && !ParentFixtureTypeTemplate->Modes.IsEmpty())
 	{
 		ActiveMode = FMath::Clamp(ActiveMode, 0, ParentFixtureTypeTemplate->Modes.Num() - 1);
 	}
+	else
+	{
+		ActiveMode = INDEX_NONE;
+	}
 }
 
 const FDMXFixtureMode* UDMXEntityFixturePatch::GetActiveMode() const
diff --git a/Engine/Plugins/VirtualProduction/DMX/DMXEngine/Source/DMXRuntime/Public/Library/DMXEntityFixturePatch.h b/Engine/Plugins/VirtualProduction/DMX/DMXEngine/Source/DMXRuntime/Public/Library/DMXEntityFixturePatch.h
--- a/Engine/Plugins/VirtualProduction/DMX/DMXEngine/Source/DMXRuntime/Public/Library/DMXEntityFixturePatch.h
+++ b/Engine/Plugins/VirtualProduction/DMX/DMXEngine/Source/DMXRuntime/Public/Library/DMXEntityFixturePatch.h
@@ -280,7 +280,7 @@ protected:
 
 	/** The Index of the Mode in the Fixture Type the Patch uses */
 	UPROPERTY(EditAnywhere, BlueprintReadOnly, meta = (DisplayName = "Active Mode Index"), Category = "Fixture Patch")
-	int32 ActiveMode;
+	int32 ActiveMode = INDEX_NONE;
 
 #if WITH_EDITORONLY_DATA
 	/** 
