DMX: Fix modes cannot be multi-edited in Fixture Patch details

diff --git a/Engine/Plugins/VirtualProduction/DMX/DMXEngine/Source/DMXEditor/Private/Customizations/DMXEntityFixturePatchDetails.cpp b/Engine/Plugins/VirtualProduction/DMX/DMXEngine/Source/DMXEditor/Private/Customizations/DMXEntityFixturePatchDetails.cpp
--- a/Engine/Plugins/VirtualProduction/DMX/DMXEngine/Source/DMXEditor/Private/Customizations/DMXEntityFixturePatchDetails.cpp
+++ b/Engine/Plugins/VirtualProduction/DMX/DMXEngine/Source/DMXEditor/Private/Customizations/DMXEntityFixturePatchDetails.cpp
@@ -2,16 +2,15 @@
 
 #include "Customizations/DMXEntityFixturePatchDetails.h"
 
-#include "DMXEditorUtils.h"
-#include "DMXFixturePatchSharedData.h"
-#include "Library/DMXEntityFixturePatch.h"
-#include "Widgets/SDMXEntityDropdownMenu.h"
-
 #include "DetailLayoutBuilder.h"
 #include "DetailWidgetRow.h"
+#include "DMXEditorUtils.h"
+#include "DMXFixturePatchSharedData.h"
 #include "IPropertyUtilities.h"
+#include "Library/DMXEntityFixturePatch.h"
 #include "ScopedTransaction.h"
-
+#include "Widgets/Layout/SWidgetSwitcher.h"
+#include "Widgets/SDMXEntityDropdownMenu.h"
 
 #define LOCTEXT_NAMESPACE "DMXEntityFixturePatchFixtureSettingsDetails"
 
@@ -65,7 +64,6 @@ void FDMXEntityFixturePatchDetails::CustomizeDetails(IDetailLayoutBuilder& Detai
 			ActiveModeHandle->CreatePropertyNameWidget()
 		]
 		.ValueContent()
-		.MaxDesiredWidth(160.0f)
 		[
 			SAssignNew(ActiveModeComboBox, SComboBox<TSharedPtr<uint32>>)
 			.IsEnabled(this, &FDMXEntityFixturePatchDetails::IsActiveModeEditable)
@@ -111,6 +109,9 @@ void FDMXEntityFixturePatchDetails::OnParentFixtureTypeChanged(UDMXEntity* NewTe
 
 void FDMXEntityFixturePatchDetails::OnFixtureTypeChanged(const UDMXEntityFixtureType* FixtureType)
 {
+	GenerateActiveModesSource();
+	ActiveModeComboBox->RefreshOptions();
+
 	PropertyUtilities->RequestRefresh();
 }
 
@@ -191,8 +192,8 @@ FText FDMXEntityFixturePatchDetails::GetCurrentActiveModeLabel() const
 {
 	static const FText MultipleValuesLabel = LOCTEXT("MultipleValuesLabel", "Multiple Values");
 	static const FText NullTypeLabel = LOCTEXT("NullFixtureTypeLabel", "No Fixture Type selected");
-	static const FText MultipleTypesLabel = LOCTEXT("MultipleFixtureTypesLabel", "Multiple Types Selected");
-	static const FText NoModesLabel = LOCTEXT("NoModesLabel", "No modes in Fixture Type");
+	static const FText MultipleTypesLabel = LOCTEXT("MultipleFixtureTypesLabel", "Multiple Fixture Types Selected");
+	static const FText NoModesLabel = LOCTEXT("NoModesLabel", "No Modes in Fixture Type");
 
 	UObject* Object = nullptr;
 	const FPropertyAccess::Result FixtureTemplateAccessResult = ParentFixtureTypeHandle->GetValue(Object);
