DMX Control Console - Fix various issues related to DMX Editor settings using a custom config instead of EditorPerProjectUserSettings

diff --git a/Engine/Plugins/VirtualProduction/DMX/DMXEngine/Source/DMXEditor/Public/DMXEditorSettings.h b/Engine/Plugins/VirtualProduction/DMX/DMXEngine/Source/DMXEditor/Public/DMXEditorSettings.h
--- a/Engine/Plugins/VirtualProduction/DMX/DMXEngine/Source/DMXEditor/Public/DMXEditorSettings.h
+++ b/Engine/Plugins/VirtualProduction/DMX/DMXEngine/Source/DMXEditor/Public/DMXEditorSettings.h
@@ -171,7 +171,7 @@ struct FDMXConflictMonitorSettings
 };
 
 /** Settings that holds editor configurations. Not accessible in Project Settings. TODO: Idealy rename to UDMXEditorConfiguration */
-UCLASS(Config = DMXEditor)
+UCLASS(Config = EditorPerProjectUserSettings)
 class DMXEDITOR_API UDMXEditorSettings : public UObject
 {
 	GENERATED_BODY()
