DMX: Fix auto-assign aligns valid stacked patches one after another instead of preserving the stack

diff --git a/Engine/Plugins/VirtualProduction/DMX/DMXEngine/Source/DMXEditor/Private/FixturePatchAutoAssignUtility.cpp b/Engine/Plugins/VirtualProduction/DMX/DMXEngine/Source/DMXEditor/Private/FixturePatchAutoAssignUtility.cpp
--- a/Engine/Plugins/VirtualProduction/DMX/DMXEngine/Source/DMXEditor/Private/FixturePatchAutoAssignUtility.cpp
+++ b/Engine/Plugins/VirtualProduction/DMX/DMXEngine/Source/DMXEditor/Private/FixturePatchAutoAssignUtility.cpp
@@ -2,19 +2,14 @@
 
 #include "FixturePatchAutoAssignUtility.h"
 
-#include "Algo/Accumulate.h"
 #include "Algo/MaxElement.h"
-#include "Commands/DMXEditorCommands.h"
-#include "DMXProtocolConstants.h"
 #include "DMXEditor.h"
 #include "DMXFixturePatchSharedData.h"
-#include "Framework/MultiBox/MultiBoxBuilder.h"
+#include "DMXProtocolConstants.h"
 #include "IO/DMXInputPort.h"
 #include "IO/DMXOutputPort.h"
 #include "Library/DMXEntityFixturePatch.h"
 #include "Library/DMXLibrary.h"
-#include "ScopedTransaction.h"
-
 
 #define LOCTEXT_NAMESPACE "FixturePatchAutoAssignUtility"
 
@@ -69,6 +64,21 @@ namespace UE::DMXEditor::AutoAssign::Private
 		AbsoluteRange = TRange<int64>(NewStartingChannel, NewEndingAddress);
 	}
 
+	void FAutoAssignElement::AlignAfter(const TSharedRef<FAutoAssignElement>& Other)
+	{
+		const bool bFitsCurrentUniverse = (Other->GetUpperBoundValue() + GetSize() - 1) / DMX_UNIVERSE_SIZE == Other->GetUpperBoundValue() / DMX_UNIVERSE_SIZE;
+		if (bFitsCurrentUniverse)
+		{
+			SetAbsoluteStartingChannel(Other->GetUpperBoundValue());
+		}
+		else
+		{
+			// Assign to next universe
+			const int64 NewStartingChannel = (Other->GetUpperBoundValue() + DMX_UNIVERSE_SIZE) / DMX_UNIVERSE_SIZE * DMX_UNIVERSE_SIZE;
+			SetAbsoluteStartingChannel(NewStartingChannel);
+		}
+	}
+
 	void FAutoAssignElement::ApplyToPatch()
 	{
 		if (UDMXEntityFixturePatch* Patch = WeakFixturePatch.Get())
@@ -141,7 +151,7 @@ namespace UE::DMXEditor::AutoAssign::Private
 		return FirstPatchedUniverse;
 	}
 
-	void FAutoAssignElementsUtility::Align(TArray<TSharedRef<FAutoAssignElement>> ElementsToAlign)
+	void FAutoAssignElementsUtility::Align(TArray<TSharedRef<FAutoAssignElement>> ElementsToAlign, bool bRetainStacks)
 	{
 		if (ElementsToAlign.IsEmpty())
 		{
@@ -155,30 +165,74 @@ namespace UE::DMXEditor::AutoAssign::Private
 
 		Algo::StableSortBy(ElementsToAlign, &FAutoAssignElement::GetLowerBoundValue);
 
+		// Move all patches into valid universe range
 		if (ElementsToAlign[0]->GetLowerBoundValue() < DMX_UNIVERSE_SIZE)
 		{
-			ElementsToAlign[0]->SetAbsoluteStartingChannel(DMX_UNIVERSE_SIZE);
-			ElementsToAlign[0]->ApplyToPatch();
+			const int64 Offset = DMX_UNIVERSE_SIZE - ElementsToAlign[0]->GetLowerBoundValue();
+			for (const TSharedRef<FAutoAssignElement>& ElementToAlign : ElementsToAlign)
+			{
+				ElementToAlign->SetAbsoluteStartingChannel(ElementToAlign->GetLowerBoundValue() + Offset);
+			}
 		}
 
-		for (int32 ElementIndex = 1; ElementIndex < ElementsToAlign.Num(); ElementIndex++)
+		if (bRetainStacks)
 		{
-			const TSharedRef<FAutoAssignElement>& PreviousElement = ElementsToAlign[ElementIndex - 1];
-			const TSharedRef<FAutoAssignElement>& Element = ElementsToAlign[ElementIndex];
-
-			const bool bFitsCurrentUniverse = (PreviousElement->GetUpperBoundValue() + Element->GetSize() - 1) / DMX_UNIVERSE_SIZE == PreviousElement->GetUpperBoundValue() / DMX_UNIVERSE_SIZE;
-			if (bFitsCurrentUniverse)
+			// Create groups of stacked elements
+			TArray<TArray<TSharedRef<FAutoAssignElement>>> StackedElements;
+			for (int32 ElementIndex = 0; ElementIndex < ElementsToAlign.Num(); ElementIndex++)
 			{
-				Element->SetAbsoluteStartingChannel(PreviousElement->GetUpperBoundValue());
+				const bool bIsStack = [&ElementsToAlign, ElementIndex]()
+					{
+						if (ElementsToAlign.IsValidIndex(ElementIndex - 1))
+						{
+							return AreElementsStacked(ElementsToAlign[ElementIndex - 1], ElementsToAlign[ElementIndex]);
+						}
+						else
+						{
+							return false;
+						}
+					}();
+
+				if (bIsStack &&
+					ensureMsgf(!StackedElements.IsEmpty(), TEXT("Trying to group auto assign elements into a stack, but cannot find the stack.")))
+				{
+					StackedElements.Last().Add(ElementsToAlign[ElementIndex]);
+				}
+				else
+				{
+					StackedElements.Add(TArray<TSharedRef<FAutoAssignElement>>({ ElementsToAlign[ElementIndex] }));
+				}
 			}
-			else
+
+			// Align stacks of elements
+			for (int32 StackIndex = 1; StackIndex < StackedElements.Num(); StackIndex++)
 			{
-				// Assign to next universe
-				const int64 NewStartingChannel = (PreviousElement->GetUpperBoundValue() + DMX_UNIVERSE_SIZE) / DMX_UNIVERSE_SIZE * DMX_UNIVERSE_SIZE;
-				Element->SetAbsoluteStartingChannel(NewStartingChannel);
+				const TArray<TSharedRef<FAutoAssignElement>>& PreviousStack = StackedElements[StackIndex - 1];
+				if (!ensureMsgf(!PreviousStack.IsEmpty(), TEXT("Cannot find elements in previous stack, cannot align elements retaining the stack.")))
+				{
+					continue;
+				}
+				const TSharedRef<FAutoAssignElement>& PreviousElement = PreviousStack[0];
+				
+				const TArray<TSharedRef<FAutoAssignElement>>& Stack = StackedElements[StackIndex];
+				for (const TSharedRef<FAutoAssignElement>& Element : Stack)
+				{
+					Element->AlignAfter(PreviousElement);
+					Element->ApplyToPatch();
+				}
 			}
+		}
+		else
+		{
+			// Align all elements
+			for (int32 ElementIndex = 1; ElementIndex < ElementsToAlign.Num(); ElementIndex++)
+			{
+				const TSharedRef<FAutoAssignElement>& PreviousElement = ElementsToAlign[ElementIndex - 1];
+				const TSharedRef<FAutoAssignElement>& Element = ElementsToAlign[ElementIndex];
 
-			Element->ApplyToPatch();
+				Element->AlignAfter(PreviousElement);
+				Element->ApplyToPatch();
+			}
 		}
 	}
 
@@ -275,7 +329,9 @@ namespace UE::DMXEditor::AutoAssign::Private
 			return;
 		}
 		Algo::StableSortBy(AutoAssignElements, &FAutoAssignElement::GetLowerBoundValue);
-		Align(AutoAssignElements);
+
+		constexpr bool bRetainStacks = true;
+		Align(AutoAssignElements, bRetainStacks);
 
 		TArray<int64> StartingChannels;
 		for (const TSharedRef<FAutoAssignElement>& FreeElement : FreeElements)
@@ -318,6 +374,29 @@ namespace UE::DMXEditor::AutoAssign::Private
 		}
 	}
 
+	bool FAutoAssignElementsUtility::AreElementsStacked(const TSharedRef<FAutoAssignElement>& FirstElement, const TSharedRef<FAutoAssignElement>& SecondElement)
+	{
+		if (FirstElement->GetLowerBoundValue() != SecondElement->GetLowerBoundValue() ||
+			FirstElement->GetUpperBoundValue() != SecondElement->GetUpperBoundValue())
+		{
+			return false;
+		}
+
+		const UDMXEntityFixturePatch* FirstFixturePatch = FirstElement->GetFixturePatch();
+		const UDMXEntityFixturePatch* SecondFixturePatch = SecondElement->GetFixturePatch();
+
+		if (FirstFixturePatch && SecondFixturePatch)
+		{
+			return 
+				FirstFixturePatch->GetFixtureType() &&
+				FirstFixturePatch->GetFixtureType() == SecondFixturePatch->GetFixtureType();
+		}
+		else
+		{
+			return FirstElement->GetSize() == SecondElement->GetSize();
+		}
+	}
+
 	bool FAutoAssignElementsUtility::EnsureValidPatchElements(const TArray<TSharedRef<FAutoAssignElement>>& Elements)
 	{
 		const bool bValidElements = Algo::FindByPredicate(Elements, [](const TSharedRef<FAutoAssignElement>& Element)
@@ -354,7 +433,9 @@ namespace UE::DMXEditor::AutoAssign
 
 		using namespace Private;
 		const TArray<TSharedRef<FAutoAssignElement>> PatchElements = Instance.CreatePatchElements(FixturePatches);
-		FAutoAssignElementsUtility::Align(PatchElements);
+
+		constexpr bool bRetainStacks = false;
+		FAutoAssignElementsUtility::Align(PatchElements, bRetainStacks);
 	}
 
 	void FAutoAssignUtility::Stack(TArray<UDMXEntityFixturePatch*> FixturePatches)
diff --git a/Engine/Plugins/VirtualProduction/DMX/DMXEngine/Source/DMXEditor/Private/FixturePatchAutoAssignUtility.h b/Engine/Plugins/VirtualProduction/DMX/DMXEngine/Source/DMXEditor/Private/FixturePatchAutoAssignUtility.h
--- a/Engine/Plugins/VirtualProduction/DMX/DMXEngine/Source/DMXEditor/Private/FixturePatchAutoAssignUtility.h
+++ b/Engine/Plugins/VirtualProduction/DMX/DMXEngine/Source/DMXEditor/Private/FixturePatchAutoAssignUtility.h
@@ -25,15 +25,28 @@ namespace UE::DMXEditor::AutoAssign::Private
 		FAutoAssignElement(const TRange<int64>& InAbsoluteRange);
 		FAutoAssignElement(UDMXEntityFixturePatch& FixturePatch);
 
+		/** Returns the lower bound value, inclusive */
 		int64 GetLowerBoundValue() const;
+
+		/** Returns the upper bound value, exclusive */
 		int64 GetUpperBoundValue() const;
+
+		/** Returns the size */
 		int64 GetSize() const;
+
+		/** Returns the element as range */
 		const TRange<int64>& GetRange() const { return AbsoluteRange; }
 
+		/** Sets the absolute starting channel */
 		void SetAbsoluteStartingChannel(int64 NewStartingAddress);
 
+		/** Aligns this element after the other */
+		void AlignAfter(const TSharedRef<FAutoAssignElement>& Other);
+
+		/** If the element has a patch, writes the current range to the patch */
 		void ApplyToPatch();
 
+		/** Returns the fixture patch, or nullptr if the element has no patch */
 		UDMXEntityFixturePatch* GetFixturePatch() const;
 
 	private:
@@ -59,8 +72,8 @@ namespace UE::DMXEditor::AutoAssign::Private
 		/** Assigns the elements. Returns the universe of the first assigned patch */
 		static int32 Assign(TArray<TSharedRef<FAutoAssignElement>> ElementsToAssign, int32 AssignToUniverse, int32 AssignToChannel);
 
-		/** Aligns the elements */
-		static void Align(TArray<TSharedRef<FAutoAssignElement>> ElementsToAlign);
+		/** Aligns the elements. If bRetainStacks is true, retains stacks instead of aligning them too. */
+		static void Align(TArray<TSharedRef<FAutoAssignElement>> ElementsToAlign, bool bRetainStacks);
 
 		/** Stacks the elements */
 		static void Stack(TArray<TSharedRef<FAutoAssignElement>> ElementsToStack);
@@ -72,6 +85,9 @@ namespace UE::DMXEditor::AutoAssign::Private
 		static void AutoAssign(const TArray<TSharedRef<FAutoAssignElement>>& FreeElements, TArray<TSharedRef<FAutoAssignElement>> AutoAssignElements);
 
 	private:
+		/** Returns true if the elements are stacked. */
+		static bool AreElementsStacked(const TSharedRef<FAutoAssignElement>& FirstElement, const TSharedRef<FAutoAssignElement>& SecondElement);
+
 		/** Returns true if elments are in valid DMX range. Hits an ensure condition if elements are not valid. */
 		static bool EnsureValidPatchElements(const TArray<TSharedRef<FAutoAssignElement>>& Elements);
 	};
