DMX: Fix cases where a UDMXLibrary::GeneralSceneDescription or UDMXImportGDTF::GDTFAssetImportData was not always flagged RF_Public

#rnx

diff --git a/Engine/Plugins/VirtualProduction/DMX/DMXEngine/Source/DMXRuntime/Private/DMXRuntimeMainStreamObjectVersion.h b/Engine/Plugins/VirtualProduction/DMX/DMXEngine/Source/DMXRuntime/Private/DMXRuntimeMainStreamObjectVersion.h
--- a/Engine/Plugins/VirtualProduction/DMX/DMXEngine/Source/DMXRuntime/Private/DMXRuntimeMainStreamObjectVersion.h
+++ b/Engine/Plugins/VirtualProduction/DMX/DMXEngine/Source/DMXRuntime/Private/DMXRuntimeMainStreamObjectVersion.h
@@ -48,6 +48,12 @@ struct FDMXRuntimeMainStreamObjectVersion
 		// 5.5 Upgrade MVR Scene Actor to spawn Actors per Fixture Type of per GDTF
 		DMXMVRSceneActorSpawnsActorsPerFixtureType,
 
+		// 5.5 Fix DMX Library assets with General Scene Descriptions that are not flagged RF_Public
+		FixDMXLibrariesWithGeneralSceneDescriptionsNotFlaggedRFPublic,
+
+		// 5.5 Fix DMX Import GDTF assets with GDTF Asset Import Data that are not flagged RF_Public
+		FixImportGDTFWithGDTFAssetImportDataNotFlaggedRFPublic,
+
 		VersionPlusOne,
 		LatestVersion = VersionPlusOne - 1
 	};
diff --git a/Engine/Plugins/VirtualProduction/DMX/DMXEngine/Source/DMXRuntime/Private/Library/DMXImportGDTF.cpp b/Engine/Plugins/VirtualProduction/DMX/DMXEngine/Source/DMXRuntime/Private/Library/DMXImportGDTF.cpp
--- a/Engine/Plugins/VirtualProduction/DMX/DMXEngine/Source/DMXRuntime/Private/Library/DMXImportGDTF.cpp
+++ b/Engine/Plugins/VirtualProduction/DMX/DMXEngine/Source/DMXRuntime/Private/Library/DMXImportGDTF.cpp
@@ -3,6 +3,7 @@
 #include "Library/DMXImportGDTF.h"
 
 #include "DMXGDTF.h"
+#include "DMXRuntimeMainStreamObjectVersion.h"
 #include "GDTF/DMXGDTFDescription.h"
 #include "Library/DMXGDTFAssetImportData.h"
 
@@ -15,6 +16,24 @@ UDMXImportGDTF::UDMXImportGDTF()
 	GDTFAssetImportData = NewObject<UDMXGDTFAssetImportData>(this, TEXT("GDTFAssetImportData"), RF_Public);
 }
 
+void UDMXImportGDTF::Serialize(FArchive& Ar)
+{
+	Super::Serialize(Ar);
+
+	Ar.UsingCustomVersion(FDMXRuntimeMainStreamObjectVersion::GUID);
+
+#if WITH_EDITOR
+	// Fix an issue where GDTFAssetImportData was not always flagged RF_Public
+	if (Ar.CustomVer(FDMXRuntimeMainStreamObjectVersion::GUID) < FDMXRuntimeMainStreamObjectVersion::FixImportGDTFWithGDTFAssetImportDataNotFlaggedRFPublic)
+	{
+		if (GDTFAssetImportData && !GDTFAssetImportData->HasAllFlags(RF_Public))
+		{
+			GDTFAssetImportData->SetFlags(GDTFAssetImportData->GetFlags() | RF_Public);
+		}
+	}
+#endif
+}
+
 void UDMXImportGDTF::PostLoad()
 {
 	Super::PostLoad();
@@ -22,7 +41,7 @@ void UDMXImportGDTF::PostLoad()
 	// Upgrade so this object always holds asset data.
 	if (!GDTFAssetImportData)
 	{
-		GDTFAssetImportData = NewObject<UDMXGDTFAssetImportData>(this, TEXT("GDTFAssetImportData"));
+		GDTFAssetImportData = NewObject<UDMXGDTFAssetImportData>(this, TEXT("GDTFAssetImportData"), RF_Public);
 	}
 
 #if WITH_EDITORONLY_DATA
diff --git a/Engine/Plugins/VirtualProduction/DMX/DMXEngine/Source/DMXRuntime/Private/Library/DMXLibrary.cpp b/Engine/Plugins/VirtualProduction/DMX/DMXEngine/Source/DMXRuntime/Private/Library/DMXLibrary.cpp
--- a/Engine/Plugins/VirtualProduction/DMX/DMXEngine/Source/DMXRuntime/Private/Library/DMXLibrary.cpp
+++ b/Engine/Plugins/VirtualProduction/DMX/DMXEngine/Source/DMXRuntime/Private/Library/DMXLibrary.cpp
@@ -46,6 +46,8 @@ void UDMXLibrary::Serialize(FArchive& Ar)
 {
 	Super::Serialize(Ar);
 
+	Ar.UsingCustomVersion(FDMXRuntimeMainStreamObjectVersion::GUID);
+
 #if WITH_EDITOR
 	if (Ar.IsLoading())
 	{
@@ -57,6 +59,15 @@ void UDMXLibrary::Serialize(FArchive& Ar)
 				FixturePatch->GenerateFixtureID();
 			}
 		}
+
+		// Fix an issue where the GeneralSceneDescription was not always flagged RF_Public
+		if (Ar.CustomVer(FDMXRuntimeMainStreamObjectVersion::GUID) < FDMXRuntimeMainStreamObjectVersion::FixDMXLibrariesWithGeneralSceneDescriptionsNotFlaggedRFPublic)
+		{
+			if (GeneralSceneDescription && !GeneralSceneDescription->HasAllFlags(RF_Public))
+			{
+				GeneralSceneDescription->SetFlags(GeneralSceneDescription->GetFlags() | RF_Public);
+			}
+		}
 	}
 #endif
 }
@@ -570,7 +581,12 @@ void UDMXLibrary::UpdatePorts()
 
 void UDMXLibrary::SetMVRGeneralSceneDescription(UDMXMVRGeneralSceneDescription* NewGeneralSceneDescription)
 {
-	GeneralSceneDescription = NewGeneralSceneDescription;
+	if (ensureMsgf(NewGeneralSceneDescription, TEXT("Invalid General Scene Description provided when setting it in the DMX Library.")) &&
+		ensureMsgf(NewGeneralSceneDescription->GetOuter() == this, TEXT("Cannot set a General Scene Description if the outer is not the DMX Library")))
+	{
+		GeneralSceneDescription = NewGeneralSceneDescription;
+		GeneralSceneDescription->SetFlags(GeneralSceneDescription->GetFlags() | RF_Public);
+	}
 }
 
 #if WITH_EDITOR
diff --git a/Engine/Plugins/VirtualProduction/DMX/DMXEngine/Source/DMXRuntime/Public/Library/DMXImportGDTF.h b/Engine/Plugins/VirtualProduction/DMX/DMXEngine/Source/DMXRuntime/Public/Library/DMXImportGDTF.h
--- a/Engine/Plugins/VirtualProduction/DMX/DMXEngine/Source/DMXRuntime/Public/Library/DMXImportGDTF.h
+++ b/Engine/Plugins/VirtualProduction/DMX/DMXEngine/Source/DMXRuntime/Public/Library/DMXImportGDTF.h
@@ -28,6 +28,7 @@ public:
 	UDMXImportGDTF();
 
 	//~ Begin UObject interface
+	virtual void Serialize(FArchive& Ar) override;
 	virtual void PostLoad() override;
 	//~ End UObject interface
 
