DMX: Fix GetAttributeValues and GetNormalizedAttributeValues do not redirect their AttributeValues pin correctly, fix warning with property redirect for LoadDMXLibrariesSynchronous

#rnx

diff --git a/Engine/Plugins/VirtualProduction/DMX/DMXEngine/Config/BaseDMXEngine.ini b/Engine/Plugins/VirtualProduction/DMX/DMXEngine/Config/BaseDMXEngine.ini
--- a/Engine/Plugins/VirtualProduction/DMX/DMXEngine/Config/BaseDMXEngine.ini
+++ b/Engine/Plugins/VirtualProduction/DMX/DMXEngine/Config/BaseDMXEngine.ini
@@ -3,7 +3,9 @@
 +EnumRedirects=(OldName="EDMXCommunicationTypes", NewName="EDMXCommunicationType") ; Only used for DMX Libraries that used deprecated UDMXEntitiyController, to upgrade them to 4.27
 
 ; UE5.4 to UE5.5
-+FunctionRedirects=(OldName="GetAllDMXLibraries",NewName="/Script/DMXEngine.DMXSubsystem.LoadDMXLibrariesSynchronous")
++FunctionRedirects=(OldName="DMXSubsystem.GetAllDMXLibraries",NewName="DMXSubsystem.LoadDMXLibrariesSynchronous")
 +StructRedirects=(OldName="DMXColorCIE1931xyY", NewName="/Script/DMXGDTF.DMXGDTFColorCIE1931xyY")
 +FunctionRedirects=(OldName="DMXEntityFixturePatch.GetAttributesValues",NewName="DMXEntityFixturePatch.GetAttributeValues")
++PropertyRedirects=(OldName="DMXEntityFixturePatch.GetAttributeValues.AttributesValues", NewName="GetAttributeValues.AttributeValues")
 +FunctionRedirects=(OldName="DMXEntityFixturePatch.GetNormalizedAttributesValues",NewName="DMXEntityFixturePatch.GetNormalizedAttributeValues")
++PropertyRedirects=(OldName="DMXEntityFixturePatch.GetNormalizedAttributeValues.AttributesValues", NewName="GetNormalizedAttributeValues.AttributeValues")
