DMX: Fix possible crash when spawning a DMX Library in a Level

#rnx

diff --git a/Engine/Plugins/VirtualProduction/DMX/DMXEngine/Source/DMXRuntime/Private/MVR/DMXMVRFixtureActorLibrary.cpp b/Engine/Plugins/VirtualProduction/DMX/DMXEngine/Source/DMXRuntime/Private/MVR/DMXMVRFixtureActorLibrary.cpp
--- a/Engine/Plugins/VirtualProduction/DMX/DMXEngine/Source/DMXRuntime/Private/MVR/DMXMVRFixtureActorLibrary.cpp
+++ b/Engine/Plugins/VirtualProduction/DMX/DMXEngine/Source/DMXRuntime/Private/MVR/DMXMVRFixtureActorLibrary.cpp
@@ -303,10 +303,12 @@ FDMXMVRFixtureActorLibrary::FDMXMVRFixtureActorLibrary()
 	for (const FTopLevelAssetPath& ClassPath : ClassPaths)
 	{
 		UClass* Class = LoadClass<UObject>(nullptr, *ClassPath.ToString());
-		AActor* Actor = PreviewScene.GetWorld()->SpawnActor<AActor>(Class);
-		check(Actor);
-
-		MVRActors.Add(Actor);
+		AActor* Actor = Class ? PreviewScene.GetWorld()->SpawnActor<AActor>(Class) : nullptr;
+		
+		if (Actor)
+		{
+			MVRActors.Add(Actor);
+		}
 	}
 }
 
