DMX: Fixture Patches now generate unique names in all cases, fix getting stuck in an ensureAlways loop when patches share the same name

diff --git a/Engine/Plugins/VirtualProduction/DMX/DMXEngine/Source/DMXEditor/Private/Widgets/FixturePatch/DMXFixturePatchListItem.cpp b/Engine/Plugins/VirtualProduction/DMX/DMXEngine/Source/DMXEditor/Private/Widgets/FixturePatch/DMXFixturePatchListItem.cpp
--- a/Engine/Plugins/VirtualProduction/DMX/DMXEngine/Source/DMXEditor/Private/Widgets/FixturePatch/DMXFixturePatchListItem.cpp
+++ b/Engine/Plugins/VirtualProduction/DMX/DMXEngine/Source/DMXEditor/Private/Widgets/FixturePatch/DMXFixturePatchListItem.cpp
@@ -5,6 +5,7 @@
 #include "Algo/MinElement.h"
 #include "DMXEditor.h"
 #include "DMXFixturePatchSharedData.h"
+#include "DMXRuntimeUtils.h"
 #include "Library/DMXEntityFixturePatch.h"
 #include "Library/DMXEntityReference.h"
 #include "Library/DMXLibrary.h"
@@ -62,8 +63,11 @@ void FDMXFixturePatchListItem::SetFixturePatchName(const FString& InDesiredName,
 		const FScopedTransaction SetFixturePatchNameTransaction(LOCTEXT("SetFixturePatchNameTransaction", "Set Fixture Patch Name"));
 		FixturePatch->PreEditChange(UDMXEntityFixturePatch::StaticClass()->FindPropertyByName(GET_MEMBER_NAME_CHECKED(UDMXEntity, Name)));
 
-		FixturePatch->SetName(InDesiredName);
-		OutNewName = FixturePatch->Name;
+		// Fixture Patches require to have a unique name to work with sequencer.
+		UDMXLibrary* DMXLibrary = FixturePatch->GetParentLibrary();
+		OutNewName = DMXLibrary ? FDMXRuntimeUtils::FindUniqueEntityName(DMXLibrary, UDMXEntityFixturePatch::StaticClass(), InDesiredName) : TEXT("Invalid Fixture Patch");
+
+		FixturePatch->SetName(OutNewName);
 
 		FixturePatch->PostEditChange();
 	}
diff --git a/Engine/Plugins/VirtualProduction/DMX/DMXEngine/Source/DMXRuntime/Private/DMXRuntimeUtils.cpp b/Engine/Plugins/VirtualProduction/DMX/DMXEngine/Source/DMXRuntime/Private/DMXRuntimeUtils.cpp
--- a/Engine/Plugins/VirtualProduction/DMX/DMXEngine/Source/DMXRuntime/Private/DMXRuntimeUtils.cpp
+++ b/Engine/Plugins/VirtualProduction/DMX/DMXEngine/Source/DMXRuntime/Private/DMXRuntimeUtils.cpp
@@ -4,7 +4,7 @@
 
 #include "Library/DMXLibrary.h"
 #include "Library/DMXEntityFixturePatch.h"
-
+#include "Algo/Find.h"
 #include "Algo/MaxElement.h"
 #include "Internationalization/Regex.h"
 
@@ -156,99 +156,86 @@ FString FDMXRuntimeUtils::GenerateUniqueNameFromExisting(const TSet<FString>& In
 		InDesiredName = TEXT("New");
 	}
 
-	auto GetBaseNameLambda([](const FString InName) -> FString
+	// If the desired name is unique, just use that.
+	if (Algo::Find(InExistingNames, InDesiredName) == nullptr)
+	{
+		return InDesiredName;
+	}
+
+	// Create a unique name
+	FName UniqueName = [&InDesiredName]()
 		{
-			const FRegexPattern BaseNamePattern(TEXT("(\\w*)_\\d*"));
-			FRegexMatcher Regex(BaseNamePattern, *InName);
-			if (Regex.FindNext())
+			int32 Len = InDesiredName.Len();
+			while (Len > 0 && FChar::IsDigit(InDesiredName[Len - 1]))
 			{
-				return Regex.GetCaptureGroup(1);
+				--Len;
 			}
-			return InName;
-		});
 
-	auto GetIndexStringLambda([](const FString& InName, FString& OutIndexString) -> bool
-	{
-		const FRegexPattern IndexPattern(TEXT("\\w*_(\\d*)"));
-		FRegexMatcher Regex(IndexPattern, *InName);
-		if (Regex.FindNext())
-		{
-			OutIndexString = Regex.GetCaptureGroup(1);
-			return true;
-		}
-		return false;
-	});
+			int32 Number;
+			if (!LexTryParseString(Number, *InDesiredName.RightChop(Len)))
+			{
+				return FName(*InDesiredName.Left(Len));
+			}
 
-	FString DesiredName = InDesiredName;
-	DesiredName.RemoveSpacesInline();
-	const FString DesiredBaseName = GetBaseNameLambda(DesiredName);
+			if (Len > 0 && InDesiredName[Len - 1] == TCHAR('_'))
+			{
+				--Len;
+			}
 
-	TSet<FString> ConflictingNames;
-	ConflictingNames.Reserve(InExistingNames.Num());
-	for (const FString& OtherName : InExistingNames)
-	{
-		const FString OtherBaseName = GetBaseNameLambda(OtherName);
-		if (OtherBaseName.Equals(DesiredBaseName))
-		{
-			ConflictingNames.Add(OtherName);
-		}
-	}
+			return FName(*InDesiredName.Left(Len), Number++);
+		}();
 
-	TSet<FString> ExistingIndexStrings;
-	ExistingIndexStrings.Reserve(ConflictingNames.Num());
-	for (const FString& ExistingName : ConflictingNames)
-	{
-		FString IndexString;
-		if (GetIndexStringLambda(ExistingName, IndexString))
+	const auto IsUniqueLambda = [MinNumDigitsInPostfix](TArray<FString> Others, const FName& InName)
 		{
-			ExistingIndexStrings.Add(IndexString);
-		}
-	}
+			const FString NameString = InName.ToString();
+
+			const FString& PlainNameString = InName.GetPlainNameString();
+			const int32 Number = InName.GetNumber();
+			const FString PossibleNameString = FString::Printf(TEXT("%s_%0*d"), *PlainNameString, MinNumDigitsInPostfix, Number);
 
-	int32 NewIndex = 0;
-	const FString* const MaxIndexStringPtr = Algo::MaxElementBy(ExistingIndexStrings, [](const FString& IndexString)
+			for (const FString& Other : Others)
+			{
+				if (NameString == Other ||
+					PossibleNameString == Other)
+				{
+					return false;
+				}
+			}
+
+			return true;
+		};
+
+	// Only test names that start with the plain name string to optimize for performance
+	TArray<FString> ExistingNames;
+	Algo::TransformIf(InExistingNames, ExistingNames,
+		[UniqueNameString = UniqueName.GetPlainNameString()](const FString& ExistingName)
+		{
+			return ExistingName.StartsWith(UniqueNameString);
+		},
+		[](const FString& ExistingName)
 		{
-			int32 Index = INDEX_NONE;
-			LexTryParseString(Index, *IndexString);
-			return Index;
+			return ExistingName;
 		});
 
-	if (MaxIndexStringPtr &&
-		LexTryParseString(NewIndex, **MaxIndexStringPtr))
-	{
-		// Increment the trailing index string
-		NewIndex++;
-		FString NewIndexString = FString::FromInt(NewIndex);
-		while(NewIndexString.Len() < MinNumDigitsInPostfix)
+	// Sort in descending order to optimize for performance
+	Algo::Sort(ExistingNames, [](const FString& NameA, const FString& NameB)
 		{
-			NewIndexString = FString(TEXT("0")) + NewIndexString;
-		}
+			return NameA > NameB;
+		});
 
-		const FRegexPattern BaseNamePattern(TEXT("(\\w*)_\\d*"));
-		FRegexMatcher Regex(BaseNamePattern, *DesiredName);
+	while (!IsUniqueLambda(ExistingNames, UniqueName))
+	{
+		UniqueName.SetNumber(UniqueName.GetNumber() + 1);
+	}
 
-		const FString BaseName = Regex.FindNext() ? Regex.GetCaptureGroup(1) : DesiredName;
-		return BaseName + TEXT("_") + NewIndexString;
+	const int32 UniqueNumber = UniqueName.GetNumber();
+	if (UniqueNumber == 0)
+	{
+		return UniqueName.ToString();
 	}
 	else
 	{
-		// Use the alreay appended index
-		FString DesiredIndexString;
-		if (GetIndexStringLambda(DesiredName, DesiredIndexString) &&
-			LexTryParseString(NewIndex, *DesiredIndexString))
-		{
-			return DesiredName;
-		}
-		else
-		{
-			FString NewPostfix;
-			while (NewPostfix.Len() < MinNumDigitsInPostfix)
-			{
-				NewPostfix = NewPostfix + FString(TEXT("0"));
-			}
-
-			return DesiredName.EndsWith(TEXT("_")) ? DesiredName + NewPostfix : DesiredName + TEXT("_") + NewPostfix;
-		}
+		return FString::Printf(TEXT("%s_%0*d"), *UniqueName.GetPlainNameString(), MinNumDigitsInPostfix, UniqueName.GetNumber());
 	}
 }
 
diff --git a/Engine/Plugins/VirtualProduction/RemoteControlProtocolDMX/Source/RemoteControlProtocolDMXEditor/Private/Library/RemoteControlDMXPatchBuilder.cpp b/Engine/Plugins/VirtualProduction/RemoteControlProtocolDMX/Source/RemoteControlProtocolDMXEditor/Private/Library/RemoteControlDMXPatchBuilder.cpp
--- a/Engine/Plugins/VirtualProduction/RemoteControlProtocolDMX/Source/RemoteControlProtocolDMXEditor/Private/Library/RemoteControlDMXPatchBuilder.cpp
+++ b/Engine/Plugins/VirtualProduction/RemoteControlProtocolDMX/Source/RemoteControlProtocolDMXEditor/Private/Library/RemoteControlDMXPatchBuilder.cpp
@@ -5,6 +5,7 @@
 #include "Algo/AnyOf.h"
 #include "Algo/Find.h"
 #include "Algo/MaxElement.h"
+#include "DMXRuntimeUtils.h"
 #include "GameFramework/Actor.h"
 #include "Library/DMXEntityFixturePatch.h"
 #include "Library/DMXLibrary.h"
@@ -432,7 +433,7 @@ namespace UE::RemoteControl::DMX
 				FixturePatch = UDMXEntityFixturePatch::CreateFixturePatchInLibrary(FixturePatchConstructionParams);
 				check(FixturePatch);
 
-				FixturePatch->Name = GetDesiredName();
+				FixturePatch->Name = FDMXRuntimeUtils::FindUniqueEntityName(&DMXLibrary, UDMXEntityFixturePatch::StaticClass(), GetDesiredName());
 			}
 
 			if (FixturePatch && FixtureType)
