DMX: Fix Load DMX Libraries Synchronous node does not return all DMX Libraries

diff --git a/Engine/Plugins/VirtualProduction/DMX/DMXEngine/Source/DMXRuntime/Private/DMXSubsystem.cpp b/Engine/Plugins/VirtualProduction/DMX/DMXEngine/Source/DMXRuntime/Private/DMXSubsystem.cpp
--- a/Engine/Plugins/VirtualProduction/DMX/DMXEngine/Source/DMXRuntime/Private/DMXSubsystem.cpp
+++ b/Engine/Plugins/VirtualProduction/DMX/DMXEngine/Source/DMXRuntime/Private/DMXSubsystem.cpp
@@ -575,7 +575,7 @@ TArray<UDMXLibrary*> UDMXSubsystem::LoadDMXLibrariesSynchronous() const
 	Algo::TransformIf(SoftDMXLibraries, DMXLibraries,
 		[](const TSoftObjectPtr<UDMXLibrary>& DMXLibrary)
 		{
-			return DMXLibrary.IsValid();
+			return !DMXLibrary.IsNull();
 		},
 		[](const TSoftObjectPtr<UDMXLibrary>& DMXLibrary)
 		{
