DMX: Fix projects using DMX components cannot be packaged

diff --git a/Engine/Plugins/VirtualProduction/DMX/DMXEngine/Source/DMXEditor/Private/Customizations/DMXEntityReferenceCustomization.cpp b/Engine/Plugins/VirtualProduction/DMX/DMXEngine/Source/DMXEditor/Private/Customizations/DMXEntityReferenceCustomization.cpp
--- a/Engine/Plugins/VirtualProduction/DMX/DMXEngine/Source/DMXEditor/Private/Customizations/DMXEntityReferenceCustomization.cpp
+++ b/Engine/Plugins/VirtualProduction/DMX/DMXEngine/Source/DMXEditor/Private/Customizations/DMXEntityReferenceCustomization.cpp
@@ -226,7 +226,7 @@ TWeakObjectPtr<UDMXLibrary> FDMXEntityReferenceCustomization::GetDMXLibrary() co
 {
 	TSharedPtr<IPropertyHandle> LibraryHandle = StructHandle->GetChildHandle(GET_MEMBER_NAME_CHECKED(FDMXEntityReference, DMXLibrary));
 	UObject* Object = nullptr;
-	if (LibraryHandle->GetValue(Object) == FPropertyAccess::Success)
+	if (LibraryHandle.IsValid() && LibraryHandle->GetValue(Object) == FPropertyAccess::Success)
 	{
 		return Cast<UDMXLibrary>(Object);
 	}
diff --git a/Engine/Plugins/VirtualProduction/DMX/DMXEngine/Source/DMXRuntime/Private/Game/DMXComponent.cpp b/Engine/Plugins/VirtualProduction/DMX/DMXEngine/Source/DMXRuntime/Private/Game/DMXComponent.cpp
--- a/Engine/Plugins/VirtualProduction/DMX/DMXEngine/Source/DMXRuntime/Private/Game/DMXComponent.cpp
+++ b/Engine/Plugins/VirtualProduction/DMX/DMXEngine/Source/DMXRuntime/Private/Game/DMXComponent.cpp
@@ -40,7 +40,7 @@ void UDMXComponent::SetFixturePatch(UDMXEntityFixturePatch* InFixturePatch)
 		// Remove the old receive binding
 		if (IsValid(PreviousFixturePatch))
 		{
-			PreviousFixturePatch->OnFixturePatchReceivedDMX.RemoveAll(this);
+			PreviousFixturePatch->OnFixturePatchReceivedDMXNative.RemoveAll(this);
 		}
 
 		// Set the new patch and setup the new binding
@@ -63,7 +63,7 @@ void UDMXComponent::OnFixturePatchReceivedDMX(UDMXEntityFixturePatch* FixturePat
 	// Hotfix 5.1 - Old bindings are not properly removed when duplicating Actors that use this component @todo.
 	if (FixturePatch != GetFixturePatch())
 	{
-		FixturePatch->OnFixturePatchReceivedDMX.RemoveAll(this);
+		FixturePatch->OnFixturePatchReceivedDMXNative.RemoveAll(this);
 		return;
 	}
 
@@ -94,10 +94,10 @@ void UDMXComponent::SetupReceiveDMXBinding()
 			return bReceiveDMXFromPatch;
 		}();
 		
-		if (bReceiveDMX && !FixturePatch->OnFixturePatchReceivedDMX.Contains(this, GET_FUNCTION_NAME_CHECKED(UDMXComponent, OnFixturePatchReceivedDMX)))
+		if (bReceiveDMX && !FixturePatch->OnFixturePatchReceivedDMXNative.IsBoundToObject(this))
 		{
 			// Enable receive DMX
-			FixturePatch->OnFixturePatchReceivedDMX.AddDynamic(this, &UDMXComponent::OnFixturePatchReceivedDMX);
+			FixturePatch->OnFixturePatchReceivedDMXNative.AddUObject(this, &UDMXComponent::OnFixturePatchReceivedDMX);
 
 			FDMXNormalizedAttributeValueMap NormalizeAttributeValues;
 			FixturePatch->GetNormalizedAttributeValues(NormalizeAttributeValues);
@@ -110,7 +110,7 @@ void UDMXComponent::SetupReceiveDMXBinding()
 		else if (!bReceiveDMX)
 		{			
 			// Disable receive DMX
-			FixturePatch->OnFixturePatchReceivedDMX.RemoveAll(this);
+			FixturePatch->OnFixturePatchReceivedDMXNative.RemoveAll(this);
 		}
 	}
 }
@@ -196,7 +196,7 @@ void UDMXComponent::DestroyComponent(bool bPromoteChildren)
 	UDMXEntityFixturePatch* FixturePatch = FixturePatchRef.GetFixturePatch();
 	if (IsValid(FixturePatch))
 	{
-		FixturePatch->OnFixturePatchReceivedDMX.RemoveAll(this);
+		FixturePatch->OnFixturePatchReceivedDMXNative.RemoveAll(this);
 	}
 }
 
diff --git a/Engine/Plugins/VirtualProduction/DMX/DMXEngine/Source/DMXRuntime/Private/Library/DMXEntityFixturePatch.cpp b/Engine/Plugins/VirtualProduction/DMX/DMXEngine/Source/DMXRuntime/Private/Library/DMXEntityFixturePatch.cpp
--- a/Engine/Plugins/VirtualProduction/DMX/DMXEngine/Source/DMXRuntime/Private/Library/DMXEntityFixturePatch.cpp
+++ b/Engine/Plugins/VirtualProduction/DMX/DMXEngine/Source/DMXRuntime/Private/Library/DMXEntityFixturePatch.cpp
@@ -278,6 +278,7 @@ void UDMXEntityFixturePatch::Tick(float DeltaTime)
 	{
 		if (const FDMXNormalizedAttributeValueMap* NormalizedAttributeValuesPtr = Cache.GetAllNormalizedAttributeValues())
 		{
+			OnFixturePatchReceivedDMXNative.Broadcast(this, *NormalizedAttributeValuesPtr);
 			OnFixturePatchReceivedDMX.Broadcast(this, *NormalizedAttributeValuesPtr);
 		}
 	}
@@ -285,12 +286,12 @@ void UDMXEntityFixturePatch::Tick(float DeltaTime)
 
 bool UDMXEntityFixturePatch::IsTickable() const
 {
-	return OnFixturePatchReceivedDMX.IsBound();
+	return OnFixturePatchReceivedDMXNative.IsBound() || OnFixturePatchReceivedDMX.IsBound();
 }
 
 bool UDMXEntityFixturePatch::IsTickableInEditor() const
 {
-	const bool bHasListener = OnFixturePatchReceivedDMX.IsBound();
+	const bool bHasListener = OnFixturePatchReceivedDMXNative.IsBound() || OnFixturePatchReceivedDMX.IsBound();
 	if (!bHasListener)
 	{
 		return false;
diff --git a/Engine/Plugins/VirtualProduction/DMX/DMXEngine/Source/DMXRuntime/Private/Library/DMXEntityReference.cpp b/Engine/Plugins/VirtualProduction/DMX/DMXEngine/Source/DMXRuntime/Private/Library/DMXEntityReference.cpp
--- a/Engine/Plugins/VirtualProduction/DMX/DMXEngine/Source/DMXRuntime/Private/Library/DMXEntityReference.cpp
+++ b/Engine/Plugins/VirtualProduction/DMX/DMXEngine/Source/DMXRuntime/Private/Library/DMXEntityReference.cpp
@@ -48,6 +48,7 @@ UDMXEntity* FDMXEntityReference::GetEntity() const
 		UDMXEntity* Entity = DMXLibrary->FindEntity(EntityId);
 		if (IsValid(Entity) && Entity->GetClass()->IsChildOf(GetEntityType()))
 		{
+			CachedEntity = Entity;
 			return Entity;
 		}
 	}
diff --git a/Engine/Plugins/VirtualProduction/DMX/DMXEngine/Source/DMXRuntime/Public/Library/DMXEntityFixturePatch.h b/Engine/Plugins/VirtualProduction/DMX/DMXEngine/Source/DMXRuntime/Public/Library/DMXEntityFixturePatch.h
--- a/Engine/Plugins/VirtualProduction/DMX/DMXEngine/Source/DMXRuntime/Public/Library/DMXEntityFixturePatch.h
+++ b/Engine/Plugins/VirtualProduction/DMX/DMXEngine/Source/DMXRuntime/Public/Library/DMXEntityFixturePatch.h
@@ -78,6 +78,8 @@ class DMXRUNTIME_API UDMXEntityFixturePatch
 
 	DECLARE_DYNAMIC_MULTICAST_DELEGATE_TwoParams(FDMXOnFixturePatchReceivedDMXDelegate, UDMXEntityFixturePatch*, FixturePatch, const FDMXNormalizedAttributeValueMap&, ValuePerAttribute);
 
+	DECLARE_MULTICAST_DELEGATE_TwoParams(FDMXOnFixturePatchReceivedDMXDelegateNative, UDMXEntityFixturePatch* /** FixturePatch */, const FDMXNormalizedAttributeValueMap& /** ValuePerAttribute */)
+
 public:
 	UDMXEntityFixturePatch();
 
@@ -143,7 +145,20 @@ public:
 	/** Returns the last received DMX signal. */
 	const FDMXSignalSharedPtr& GetLastReceivedDMXSignal() const { return LastDMXSignal; }
 
-	/** Broadcasts when the patch received dmx */
+	/** 
+	 * Delegate broadcast when the fixture patch received DMX. 
+	 * This event only fires when any attribute value changed. 
+	 * Use GetAttributeValues or GetNormalizedAttributeValues to get unchanged values.
+	 * 
+	 * Native version, should be preferred when binding UObjects in C++ to avoid serializing the reference. 
+	 */
+	FDMXOnFixturePatchReceivedDMXDelegateNative OnFixturePatchReceivedDMXNative;
+
+	/**
+	 * Returns a delegate broadcast when the fixture patch received DMX.
+	 * This event only fires when any attribute value changed.
+	 * Use GetAttributeValues or GetNormalizedAttributeValues to get unchanged values.
+	 */
 	UPROPERTY(BlueprintAssignable, Category = "DMX");
 	FDMXOnFixturePatchReceivedDMXDelegate OnFixturePatchReceivedDMX;
 
@@ -314,7 +329,7 @@ public:
 	FLinearColor EditorColor;
 
 	/** 
-	 * If true, the patch receives dmx and raises the OnFixturePatchReceivedDMX event in editor. 
+	 * If true, the patch receives dmx and raises the OnFixturePatchReceivedDMX events in editor. 
 	 * NOTE: If 'All Fixture Patches receive DMX in editor' is set to true in Project Settings -> Plugins -> DMX, this setting here is ignored.
 	 */
 	UPROPERTY(EditAnywhere, BlueprintReadOnly, Category = "Fixture Patch")
