Add #include for FStringOutputDevice now that it is separated into its own header.

diff --git a/Engine/Plugins/VirtualProduction/DMX/DMXEngine/Source/DMXEditor/Private/DMXEditorUtils.cpp b/Engine/Plugins/VirtualProduction/DMX/DMXEngine/Source/DMXEditor/Private/DMXEditorUtils.cpp
--- a/Engine/Plugins/VirtualProduction/DMX/DMXEngine/Source/DMXEditor/Private/DMXEditorUtils.cpp
+++ b/Engine/Plugins/VirtualProduction/DMX/DMXEngine/Source/DMXEditor/Private/DMXEditorUtils.cpp
@@ -15,6 +15,7 @@
 #include "IO/DMXInputPort.h"
 #include "IO/DMXOutputPort.h"
 #include "IO/DMXPortManager.h"
+#include "Misc/StringOutputDevice.h"
 #include "Library/DMXEntityFixturePatch.h"
 #include "Library/DMXEntityFixtureType.h"
 #include "Library/DMXLibrary.h"
diff --git a/Engine/Plugins/VirtualProduction/DMX/DMXEngine/Source/DMXEditor/Private/Widgets/FixturePatch/SDMXPatchedUniverse.cpp b/Engine/Plugins/VirtualProduction/DMX/DMXEngine/Source/DMXEditor/Private/Widgets/FixturePatch/SDMXPatchedUniverse.cpp
--- a/Engine/Plugins/VirtualProduction/DMX/DMXEngine/Source/DMXEditor/Private/Widgets/FixturePatch/SDMXPatchedUniverse.cpp
+++ b/Engine/Plugins/VirtualProduction/DMX/DMXEngine/Source/DMXEditor/Private/Widgets/FixturePatch/SDMXPatchedUniverse.cpp
@@ -20,7 +20,6 @@
 #include "IO/DMXPortManager.h"
 #include "Library/DMXEntityFixturePatch.h"
 #include "Library/DMXLibrary.h"
-#include "Misc/StringOutputDevice.h"
 #include "ScopedTransaction.h"
 #include "SDMXChannelConnector.h"
 #include "SDMXFixturePatchFragment.h"
