Add #include for FStringOutputDevice now that it is separated into its own header.

diff --git a/Engine/Plugins/VirtualProduction/DMX/DMXEngine/Source/DMXEditor/Private/Widgets/FixturePatch/SDMXFixturePatchList.cpp b/Engine/Plugins/VirtualProduction/DMX/DMXEngine/Source/DMXEditor/Private/Widgets/FixturePatch/SDMXFixturePatchList.cpp
--- a/Engine/Plugins/VirtualProduction/DMX/DMXEngine/Source/DMXEditor/Private/Widgets/FixturePatch/SDMXFixturePatchList.cpp
+++ b/Engine/Plugins/VirtualProduction/DMX/DMXEngine/Source/DMXEditor/Private/Widgets/FixturePatch/SDMXFixturePatchList.cpp
@@ -25,6 +25,7 @@
 #include "Library/DMXEntityFixturePatch.h"
 #include "Library/DMXEntityFixtureType.h"
 #include "Library/DMXLibrary.h"
+#include "Misc/StringOutputDevice.h"
 #include "ScopedTransaction.h"
 #include "Styling/AppStyle.h"
 #include "TimerManager.h"
