DMX - Fix a crash when adding a DMX Fixture Patch Channel to a DMX Library Track that has more than one DMX Library Section added to it

diff --git a/Engine/Plugins/VirtualProduction/DMX/DMXEngine/Source/DMXEditor/Private/Sequencer/DMXLibraryTrackEditor.cpp b/Engine/Plugins/VirtualProduction/DMX/DMXEngine/Source/DMXEditor/Private/Sequencer/DMXLibraryTrackEditor.cpp
--- a/Engine/Plugins/VirtualProduction/DMX/DMXEngine/Source/DMXEditor/Private/Sequencer/DMXLibraryTrackEditor.cpp
+++ b/Engine/Plugins/VirtualProduction/DMX/DMXEngine/Source/DMXEditor/Private/Sequencer/DMXLibraryTrackEditor.cpp
@@ -82,7 +82,14 @@ void FDMXLibraryTrackEditor::BuildTrackContextMenu(FMenuBuilder& MenuBuilder, UM
 			{
 				FScopedTransaction Transaction(LOCTEXT("SetDMXAssetTransaction", "Assign DMX Library to track"));
 				DMXTrack->Modify();
-				DMXTrack->GetAllSections()[0]->Modify();
+				for (UMovieSceneSection* Section : DMXTrack->GetAllSections())
+				{
+					if (UMovieSceneDMXLibrarySection* DMXLibrarySection = Cast<UMovieSceneDMXLibrarySection>(Section))
+					{
+						DMXLibrarySection->Modify();
+					}
+				}
+
 				DMXTrack->SetDMXLibrary(DMXLibrary);
 			}
 		}
@@ -272,8 +279,16 @@ void FDMXLibraryTrackEditor::HandlePatchSelectedFromAddMenu(UMovieSceneDMXLibrar
 	
 	const FName PatchFName(*Patch->GetDisplayName());
 
-	check(Track->GetAllSections().Num() == 1);
-	UMovieSceneDMXLibrarySection* DMXSection = CastChecked<UMovieSceneDMXLibrarySection>(Track->GetAllSections()[0]);
+	UMovieSceneSection* const* DMXLibrarySectionPtr = Track->GetAllSections().FindByPredicate([](UMovieSceneSection* Section)
+		{
+			return Section && Section->IsA(UMovieSceneDMXLibrarySection::StaticClass());
+		});
+	if (!DMXLibrarySectionPtr)
+	{
+		return;
+	}
+
+	UMovieSceneDMXLibrarySection* DMXSection = CastChecked<UMovieSceneDMXLibrarySection>(*DMXLibrarySectionPtr);
 
 	// Don't add repeated Fixture Patches
 	if (DMXSection->ContainsFixturePatch(Patch))
@@ -317,8 +332,16 @@ void FDMXLibraryTrackEditor::HandleAddAllPatchesClicked(UMovieSceneDMXLibraryTra
 		return;
 	}
 
-	check(Track->GetAllSections().Num() == 1);
-	UMovieSceneDMXLibrarySection* DMXSection = CastChecked<UMovieSceneDMXLibrarySection>(Track->GetAllSections()[0]);
+	UMovieSceneSection* const* DMXLibrarySectionPtr = Track->GetAllSections().FindByPredicate([](UMovieSceneSection* Section)
+		{
+			return Section && Section->IsA(UMovieSceneDMXLibrarySection::StaticClass());
+		});
+	if (!DMXLibrarySectionPtr)
+	{
+		return;
+	}
+
+	UMovieSceneDMXLibrarySection* DMXSection = CastChecked<UMovieSceneDMXLibrarySection>(*DMXLibrarySectionPtr);
 
 	const TArray<UDMXEntityFixturePatch*>&& LibraryPatches = Track->GetDMXLibrary()->GetEntitiesTypeCast<UDMXEntityFixturePatch>();
 	// Cache existing Patches in the section to not add repeated ones
