DMX: The Fixture Type Editor tab  in the DMX Library now has a button to reset a Fixture Type to how it was imported from GDTF

diff --git a/Engine/Plugins/VirtualProduction/DMX/DMXEngine/Source/DMXEditor/Private/Customizations/DMXEntityFixtureTypeDetails.cpp b/Engine/Plugins/VirtualProduction/DMX/DMXEngine/Source/DMXEditor/Private/Customizations/DMXEntityFixtureTypeDetails.cpp
--- a/Engine/Plugins/VirtualProduction/DMX/DMXEngine/Source/DMXEditor/Private/Customizations/DMXEntityFixtureTypeDetails.cpp
+++ b/Engine/Plugins/VirtualProduction/DMX/DMXEngine/Source/DMXEditor/Private/Customizations/DMXEntityFixtureTypeDetails.cpp
@@ -4,11 +4,14 @@
 
 #include "DetailCategoryBuilder.h"
 #include "DetailLayoutBuilder.h"
+#include "DetailWidgetRow.h"
 #include "Factories/DMXGDTFToFixtureTypeConverter.h"
 #include "IPropertyUtilities.h"
 #include "Library/DMXEntityFixtureType.h"
 #include "Library/DMXImportGDTF.h"
 #include "PropertyHandle.h"
+#include "Widgets/Input/SButton.h"
+#include "Widgets/SBoxPanel.h"
 
 #define LOCTEXT_NAMESPACE "DMXEntityFixtureTypeDetails"
 
@@ -21,17 +24,60 @@ void FDMXEntityFixtureTypeDetails::CustomizeDetails(IDetailLayoutBuilder& Detail
 {
 	PropertyUtilities = DetailBuilder.GetPropertyUtilities();
 
-	DetailBuilder.HideProperty(GET_MEMBER_NAME_CHECKED(UDMXEntityFixtureType, Modes));
+	// Hide the Modes property, it has its own view in the editor
+	{
+		DetailBuilder.HideProperty(GET_MEMBER_NAME_CHECKED(UDMXEntityFixtureType, Modes));
+	}
 
-	GDTFSourceHandle = DetailBuilder.GetProperty(GET_MEMBER_NAME_CHECKED(UDMXEntityFixtureType, GDTFSource));
-	GDTFSourceHandle->SetOnPropertyValueChanged(FSimpleDelegate::CreateSP(this, &FDMXEntityFixtureTypeDetails::OnGDTFSourceChanged));
+	// Handle property changes for the GDTFSource property
+	{
+		GDTFSourceHandle = DetailBuilder.GetProperty(GET_MEMBER_NAME_CHECKED(UDMXEntityFixtureType, GDTFSource));
+		GDTFSourceHandle->SetOnPropertyValueChanged(FSimpleDelegate::CreateSP(this, &FDMXEntityFixtureTypeDetails::OnGDTFSourceChanged));
+	}
 
 	// Customize the bExportGeneratedGDTF property
-	const TSharedRef<IPropertyHandle> ExportGeneratedGDTFHandle = DetailBuilder.GetProperty(GET_MEMBER_NAME_CHECKED(UDMXEntityFixtureType, bExportGeneratedGDTF));
-	const FName GDTFCategory = ExportGeneratedGDTFHandle->GetDefaultCategoryName();
-	DetailBuilder.EditCategory(GDTFCategory)
-		.AddProperty(ExportGeneratedGDTFHandle, EPropertyLocation::Advanced)
-		.Visibility(TAttribute<EVisibility>::CreateSP(this, &FDMXEntityFixtureTypeDetails::GetExportGeneratedGDTFPropertyVisibility));
+	{
+		const TSharedRef<IPropertyHandle> ExportGeneratedGDTFHandle = DetailBuilder.GetProperty(GET_MEMBER_NAME_CHECKED(UDMXEntityFixtureType, bExportGeneratedGDTF));
+		const FName GDTFCategory = ExportGeneratedGDTFHandle->GetDefaultCategoryName();
+		DetailBuilder.EditCategory(GDTFCategory)
+			.AddProperty(ExportGeneratedGDTFHandle, EPropertyLocation::Advanced)
+			.Visibility(TAttribute<EVisibility>::CreateSP(this, &FDMXEntityFixtureTypeDetails::GetGDTFOptionsVisibility));
+	}
+
+	// Add the "Reset to GDTF" button
+	{
+		const FName GDTFCategoryName = GDTFSourceHandle->GetDefaultCategoryName();
+		const FText EmptyFilterText = FText::GetEmpty();
+
+		IDetailCategoryBuilder& CategoryBuilder = DetailBuilder.EditCategory(GDTFCategoryName);
+
+		TArray<TSharedRef<IPropertyHandle>> DefaultProperties;
+		CategoryBuilder.GetDefaultProperties(DefaultProperties);
+
+		for (const TSharedRef<IPropertyHandle>& DefaultProperty : DefaultProperties)
+		{
+			CategoryBuilder.AddProperty(DefaultProperty);
+		}
+
+		CategoryBuilder
+			.AddCustomRow(EmptyFilterText)
+			.WholeRowContent()
+			[
+				SNew(SHorizontalBox)
+
+				+ SHorizontalBox::Slot()
+				.AutoWidth()
+				[
+					SNew(SButton)
+					.HAlign(HAlign_Left)
+					.VAlign(VAlign_Center)
+					.Text(LOCTEXT("ResetToGDTFLabel", "Reset To GDTF"))
+					.ToolTipText(LOCTEXT("ResetToGDTFTooltip", "Resets the Fixture Type to the state defined in the GDTF"))
+					.OnClicked(this, &FDMXEntityFixtureTypeDetails::OnResetToGDTFClicked)
+				]
+			]
+		.Visibility(TAttribute<EVisibility>::CreateSP(this, &FDMXEntityFixtureTypeDetails::GetGDTFOptionsVisibility));
+	}
 }
 
 void FDMXEntityFixtureTypeDetails::OnGDTFSourceChanged()
@@ -90,7 +136,15 @@ void FDMXEntityFixtureTypeDetails::OnGDTFSourceChanged()
 	}
 }
 
-EVisibility FDMXEntityFixtureTypeDetails::GetExportGeneratedGDTFPropertyVisibility() const
+FReply FDMXEntityFixtureTypeDetails::OnResetToGDTFClicked()
+{
+	// Run as if the source changed
+	OnGDTFSourceChanged();
+
+	return FReply::Handled();
+}
+
+EVisibility FDMXEntityFixtureTypeDetails::GetGDTFOptionsVisibility() const
 {
 	TArray<const void*> RawDataArray;
 	GDTFSourceHandle->AccessRawData(RawDataArray);
diff --git a/Engine/Plugins/VirtualProduction/DMX/DMXEngine/Source/DMXEditor/Private/Customizations/DMXEntityFixtureTypeDetails.h b/Engine/Plugins/VirtualProduction/DMX/DMXEngine/Source/DMXEditor/Private/Customizations/DMXEntityFixtureTypeDetails.h
--- a/Engine/Plugins/VirtualProduction/DMX/DMXEngine/Source/DMXEditor/Private/Customizations/DMXEntityFixtureTypeDetails.h
+++ b/Engine/Plugins/VirtualProduction/DMX/DMXEngine/Source/DMXEditor/Private/Customizations/DMXEntityFixtureTypeDetails.h
@@ -3,6 +3,7 @@
 #pragma once
 
 #include "IDetailCustomization.h"
+#include "Input/Reply.h"
 #include "Layout/Visibility.h"
 
 class IPropertyHandle;
@@ -26,8 +27,11 @@ private:
 	/** Called when the GDTFSource property changed */
 	void OnGDTFSourceChanged();
 
-	/** Returns the visibility of the bExportGeneratedGDTF property */
-	EVisibility GetExportGeneratedGDTFPropertyVisibility() const;
+	/** Called when the Reset To GDTF button was clicked */
+	FReply OnResetToGDTFClicked();
+
+	/** Returns the visibility of the GDTF options, they're only visible if there is a GDTF */
+	EVisibility GetGDTFOptionsVisibility() const;
 
 	/** Handle to the GDTFSource property */
 	TSharedPtr<IPropertyHandle> GDTFSourceHandle;
