DMX: Add missing transaction for the new Fixture Type Reset To GDTF button

#rnx

diff --git a/Engine/Plugins/VirtualProduction/DMX/DMXEngine/Source/DMXEditor/Private/Customizations/DMXEntityFixtureTypeDetails.cpp b/Engine/Plugins/VirtualProduction/DMX/DMXEngine/Source/DMXEditor/Private/Customizations/DMXEntityFixtureTypeDetails.cpp
--- a/Engine/Plugins/VirtualProduction/DMX/DMXEngine/Source/DMXEditor/Private/Customizations/DMXEntityFixtureTypeDetails.cpp
+++ b/Engine/Plugins/VirtualProduction/DMX/DMXEngine/Source/DMXEditor/Private/Customizations/DMXEntityFixtureTypeDetails.cpp
@@ -10,6 +10,7 @@
 #include "Library/DMXEntityFixtureType.h"
 #include "Library/DMXImportGDTF.h"
 #include "PropertyHandle.h"
+#include "ScopedTransaction.h"
 #include "Widgets/Input/SButton.h"
 #include "Widgets/SBoxPanel.h"
 
@@ -138,6 +139,8 @@ void FDMXEntityFixtureTypeDetails::OnGDTFSourceChanged()
 
 FReply FDMXEntityFixtureTypeDetails::OnResetToGDTFClicked()
 {
+	const FScopedTransaction ResetToGDTFTransaction(LOCTEXT("ResetToGDTFTransaction", "Reset To GDTF"));
+
 	// Run as if the source changed
 	OnGDTFSourceChanged();
 
