DMX: Fix DMX components cannot be patched

diff --git a/Engine/Plugins/VirtualProduction/DMX/DMXEngine/Source/DMXRuntime/Private/Library/DMXEntityReference.cpp b/Engine/Plugins/VirtualProduction/DMX/DMXEngine/Source/DMXRuntime/Private/Library/DMXEntityReference.cpp
--- a/Engine/Plugins/VirtualProduction/DMX/DMXEngine/Source/DMXRuntime/Private/Library/DMXEntityReference.cpp
+++ b/Engine/Plugins/VirtualProduction/DMX/DMXEngine/Source/DMXRuntime/Private/Library/DMXEntityReference.cpp
@@ -39,7 +39,9 @@ void FDMXEntityReference::SetEntity(UDMXEntity* NewEntity)
 
 UDMXEntity* FDMXEntityReference::GetEntity() const
 {
-	if (CachedEntity.IsValid())
+	if (CachedEntity.IsValid() &&
+		CachedEntity->GetID() == EntityId &&
+		CachedEntity->GetParentLibrary() == DMXLibrary)
 	{
 		return CachedEntity.Get();
 	}
@@ -51,6 +53,10 @@ UDMXEntity* FDMXEntityReference::GetEntity() const
 			CachedEntity = Entity;
 			return Entity;
 		}
+		else
+		{
+			CachedEntity.Reset();
+		}
 	}
 	return nullptr;
 }
